#ttp 331362 - CRASH: EDITOR: Occurs when attempting to open a material layer within the parented material of a material instance constant.

#branch UE4

#proj Engine

#summary Disabled overhead stats for expression preview material.

diff --git a/Engine/Source/Editor/MaterialEditor/Private/MaterialEditor.cpp b/Engine/Source/Editor/MaterialEditor/Private/MaterialEditor.cpp
--- a/Engine/Source/Editor/MaterialEditor/Private/MaterialEditor.cpp
+++ b/Engine/Source/Editor/MaterialEditor/Private/MaterialEditor.cpp
@@ -1117,6 +1117,9 @@ void FMaterialEditor::UpdatePreviewMaterial( )
 		// If we are previewing an expression, update the expression preview material
 		ExpressionPreviewMaterial->PreEditChange( NULL );
 		ExpressionPreviewMaterial->PostEditChange();
+
+		//Don't give overhead stats for preview expression material.
+		MaterialDevelopmentOverheadStats.Update(NULL);
 	}
 	else 
 	{
@@ -1126,9 +1129,10 @@ void FMaterialEditor::UpdatePreviewMaterial( )
 
 		// Null out the expression preview material so they can be GC'ed
 		ExpressionPreviewMaterial = NULL;
+
+		MaterialDevelopmentOverheadStats.Update(Material);
 	}
 
-	MaterialDevelopmentOverheadStats.Update(Material);
 
 	// Reregister all components that use the preview material, since UMaterial::PEC does not reregister components using a bIsPreviewMaterial=true material
 	RefreshPreviewViewport();
@@ -1483,7 +1487,7 @@ void FMaterialEditor::UpdateMaterialInfoList(bool bForceDisplay)
 				MaterialResource->GetRepresentativeInstructionCounts(Descriptions, InstructionCounts);
 
 				TArray<int32> OverheadCounts;
-				bool bDoOverheadCount = false;//MaterialDevelopmentOverheadStats.GetOverheadCounts(OverheadCounts, FeatureLevel);
+				bool bDoOverheadCount = MaterialDevelopmentOverheadStats.GetOverheadCounts(OverheadCounts, FeatureLevel);
 				bDoOverheadCount = bDoOverheadCount && OverheadCounts.Num() > 0 && InstructionCounts.Num() > 0;
 				check(!bDoOverheadCount || OverheadCounts.Num() == InstructionCounts.Num());
 
@@ -3856,11 +3860,21 @@ void FMaterialEditor::FlipExpressionPositions(const TArray<UMaterialExpression*>
 
 void FMaterialDevelopmentOverheadStats::Init(UMaterial* InMaterial)
 {
+	bDisabled = false;
 	Update(InMaterial);
 }
 
 void FMaterialDevelopmentOverheadStats::Update(UMaterial* InMaterial)
 {
+	if (NULL == InMaterial)
+	{
+		//Don't bother clearing the materials as we'll either be reusing them again shortly or exiting the editor.
+		bDisabled = true;
+		return;
+	}
+
+	bDisabled = false;
+
 	//Duplicate the material in place (same name).
 	EmptyMaterial = (UMaterial*)StaticDuplicateObject(InMaterial, GetTransientPackage(), TEXT("EmptyMaterial"), ~RF_Standalone, UPreviewMaterial::StaticClass()); 		
 	EmptyMaterial->Expressions.Empty();		
@@ -3896,6 +3910,11 @@ void FMaterialDevelopmentOverheadStats::AddReferencedObjects(FReferenceCollector
 
 bool FMaterialDevelopmentOverheadStats::GetOverheadCounts(TArray<int32>& OverheadCounts, ERHIFeatureLevel::Type FeatureLevel)
 {		
+	if (bDisabled)
+	{
+		return false;
+	}
+
 	const FMaterialResource* EmptyMaterialResource = EmptyMaterial ? EmptyMaterial->GetMaterialResource(FeatureLevel) : NULL;
 	const FMaterialResource* EmptyMaterialResourceWithOverhead = EmptyMaterialWithOverhead ? EmptyMaterialWithOverhead->GetMaterialResource(FeatureLevel) : NULL;
 	
diff --git a/Engine/Source/Editor/MaterialEditor/Private/MaterialEditor.h b/Engine/Source/Editor/MaterialEditor/Private/MaterialEditor.h
--- a/Engine/Source/Editor/MaterialEditor/Private/MaterialEditor.h
+++ b/Engine/Source/Editor/MaterialEditor/Private/MaterialEditor.h
@@ -185,6 +185,7 @@ public:
 	FMaterialDevelopmentOverheadStats()
 		: EmptyMaterialWithOverhead(NULL)
 		, EmptyMaterial(NULL)
+		, bDisabled(false)
 	{
 
 	}
@@ -196,6 +197,8 @@ public:
 private:
 	UMaterial* EmptyMaterialWithOverhead;
 	UMaterial* EmptyMaterial;
+
+	bool bDisabled;
 };
 
 /**
