Changed a check to an ensure as it seems to be firing occasionally.

diff --git a/Engine/Source/Editor/MaterialEditor/Private/MaterialEditor.cpp b/Engine/Source/Editor/MaterialEditor/Private/MaterialEditor.cpp
--- a/Engine/Source/Editor/MaterialEditor/Private/MaterialEditor.cpp
+++ b/Engine/Source/Editor/MaterialEditor/Private/MaterialEditor.cpp
@@ -1489,12 +1489,11 @@ void FMaterialEditor::UpdateMaterialInfoList(bool bForceDisplay)
 				TArray<int32> OverheadCounts;
 				bool bDoOverheadCount = MaterialDevelopmentOverheadStats.GetOverheadCounts(OverheadCounts, FeatureLevel);
 				bDoOverheadCount = bDoOverheadCount && OverheadCounts.Num() > 0 && InstructionCounts.Num() > 0;
-				check(!bDoOverheadCount || OverheadCounts.Num() == InstructionCounts.Num());
 
 				for (int32 InstructionIndex = 0; InstructionIndex < Descriptions.Num(); InstructionIndex++)
 				{
 					FString InstructionCountString = FString::Printf(TEXT("%s: %u instructions"),*Descriptions[InstructionIndex],InstructionCounts[InstructionIndex]);
-					if(bDoOverheadCount)
+					if (bDoOverheadCount && ensure(OverheadCounts.Num() == InstructionCounts.Num()))
 					{
 						InstructionCountString += FString::Printf( TEXT(" - Dev Overhead: %u"), OverheadCounts[InstructionIndex] );
 					}
