#ttp 315757 - EDITOR: ME2: Add Find results window similar to the Blueprint Editor

#proj UE4.Engine

#summary Added Find Results tab to Material Editor

#add Added SFindInMaterial class to search and display results from materials, similar to the Blueprint version but didn't seem like there was an obvious method of splitting it up into a base class that could be used by both. Made sure to use the search functions from Material Expressions as well to ensure that the results will be the same or better than they were previously.

#change Added new function to UMaterialExpression to get a one line description of it. Most of the time this is just all of the expressions captions combined into one line using spaces, though there were exceptions for constants, textures and function calls so that better information is shown in the results window. Also added a JumpToNode function to FMaterialEditor so that the find results can show the nodes when clicked on.

#extra This is only for searching the current material at present, adding support for searching all assets would be a lot more work as we would need to store meta data allowing materials to be searched before they are loaded.

diff --git a/Engine/Source/Editor/EditorStyle/Private/SlateEditorStyle.cpp b/Engine/Source/Editor/EditorStyle/Private/SlateEditorStyle.cpp
--- a/Engine/Source/Editor/EditorStyle/Private/SlateEditorStyle.cpp
+++ b/Engine/Source/Editor/EditorStyle/Private/SlateEditorStyle.cpp
@@ -4528,6 +4528,8 @@ void FSlateEditorStyle::FStyle::SetupToolkitStyles()
 		Set( "MaterialEditor.ShowHideConnectors.Small", new IMAGE_BRUSH( "Icons/icon_MatEd_Connectors_40x", Icon20x20 ) );
 		Set( "MaterialEditor.CameraHome", new IMAGE_BRUSH( "Icons/icon_MatEd_Home_40x", Icon40x40 ) );
 		Set( "MaterialEditor.CameraHome.Small", new IMAGE_BRUSH( "Icons/icon_MatEd_Home_40x", Icon20x20 ) );
+		Set( "MaterialEditor.FindInMaterial", new IMAGE_BRUSH("Icons/icon_Blueprint_Find_40px", Icon40x40));
+		Set( "MaterialEditor.FindInMaterial.Small", new IMAGE_BRUSH("Icons/icon_Blueprint_Find_40px", Icon20x20));
 	}
 
 	// Material Instance Editor
diff --git a/Engine/Source/Editor/MaterialEditor/Private/FindInMaterial.cpp b/Engine/Source/Editor/MaterialEditor/Private/FindInMaterial.cpp
new file mode 100644
--- /dev/null
+++ b/Engine/Source/Editor/MaterialEditor/Private/FindInMaterial.cpp
@@ -0,0 +1,431 @@
+// Copyright 1998-2014 Epic Games, Inc. All Rights Reserved.
+#include "MaterialEditorModule.h"
+#include "FindInMaterial.h"
+#include "MaterialEditor.h"
+
+#define LOCTEXT_NAMESPACE "FindInMaterial"
+
+//////////////////////////////////////////////////////////////////////////
+// FFindInMaterialResult
+
+FFindInMaterialResult::FFindInMaterialResult(const FString& InValue)
+:Value(InValue), DuplicationIndex(0), Class(NULL), Pin(NULL), GraphNode(NULL)
+{
+}
+
+FFindInMaterialResult::FFindInMaterialResult(const FString& InValue, TSharedPtr<FFindInMaterialResult>& InParent, UClass* InClass, int InDuplicationIndex)
+: Value(InValue), DuplicationIndex(InDuplicationIndex), Parent(InParent), Class(InClass), Pin(NULL), GraphNode(NULL)
+{
+}
+
+FFindInMaterialResult::FFindInMaterialResult(const FString& InValue, TSharedPtr<FFindInMaterialResult>& InParent, UEdGraphPin* InPin)
+: Value(InValue), DuplicationIndex(0), Parent(InParent), Class(InPin->GetClass()), Pin(InPin), GraphNode(NULL)
+{
+}
+
+FFindInMaterialResult::FFindInMaterialResult(const FString& InValue, TSharedPtr<FFindInMaterialResult>& InParent, UEdGraphNode* InNode)
+: Value(InValue), DuplicationIndex(0), Parent(InParent), Class(InNode->GetClass()), Pin(NULL), GraphNode(InNode)
+{
+	if (GraphNode.IsValid())
+	{
+		CommentText = GraphNode->NodeComment;
+	}
+}
+
+FReply FFindInMaterialResult::OnClick(TWeakPtr<class FMaterialEditor> MaterialEditor)
+{
+	if (GraphNode.IsValid())
+	{
+		MaterialEditor.Pin()->JumpToNode(GraphNode.Get());
+	}
+	else if (Pin.IsValid())
+	{
+		MaterialEditor.Pin()->JumpToNode(Pin.Get()->GetOwningNode());
+	}
+	return FReply::Handled();
+}
+
+FString FFindInMaterialResult::GetCategory() const
+{
+	if (Class == UEdGraphPin::StaticClass())
+	{
+		return LOCTEXT("PinCategory", "Pin").ToString();
+	}
+	else
+	{
+		return LOCTEXT("NodeCategory", "Node").ToString();
+	}
+	return FString();
+}
+
+TSharedRef<SWidget> FFindInMaterialResult::CreateIcon() const
+{
+	FSlateColor IconColor = FSlateColor::UseForeground();
+	const FSlateBrush* Brush = NULL;
+	if (Pin.IsValid())
+	{
+		if (Pin->PinType.bIsArray)
+		{
+			Brush = FEditorStyle::GetBrush(TEXT("GraphEditor.ArrayPinIcon"));
+		}
+		else if (Pin->PinType.bIsReference)
+		{
+			Brush = FEditorStyle::GetBrush(TEXT("GraphEditor.RefPinIcon"));
+		}
+		else
+		{
+			Brush = FEditorStyle::GetBrush(TEXT("GraphEditor.PinIcon"));
+		}
+		const UEdGraphSchema* Schema = Pin->GetSchema();
+		IconColor = Schema->GetPinTypeColor(Pin->PinType);
+	}
+	else if (GraphNode.IsValid())
+	{
+		Brush = FEditorStyle::GetBrush(TEXT("GraphEditor.NodeGlyph"));
+	}
+
+	return 	SNew(SImage)
+		.Image(Brush)
+		.ColorAndOpacity(IconColor)
+		.ToolTipText(GetCategory());
+}
+
+FString FFindInMaterialResult::GetCommentText() const
+{
+	return CommentText;
+}
+
+FString FFindInMaterialResult::GetValueText() const
+{
+	FString DisplayVal;
+
+	if (Pin.IsValid()
+		&& (!Pin->DefaultValue.IsEmpty() || !Pin->AutogeneratedDefaultValue.IsEmpty() || Pin->DefaultObject || !Pin->DefaultTextValue.IsEmpty())
+		&& (Pin->Direction == EGPD_Input)
+		&& (Pin->LinkedTo.Num() == 0))
+	{
+		if (Pin->DefaultObject)
+		{
+			const AActor* DefaultActor = Cast<AActor>(Pin->DefaultObject);
+			DisplayVal = FString::Printf(TEXT("(%s)"), DefaultActor ? *DefaultActor->GetActorLabel() : *Pin->DefaultObject->GetName());
+		}
+		else if (!Pin->DefaultValue.IsEmpty())
+		{
+			DisplayVal = FString::Printf(TEXT("(%s)"), *Pin->DefaultValue);
+		}
+		else if (!Pin->AutogeneratedDefaultValue.IsEmpty())
+		{
+			DisplayVal = FString::Printf(TEXT("(%s)"), *Pin->AutogeneratedDefaultValue);
+		}
+		else if (!Pin->DefaultTextValue.IsEmpty())
+		{
+			DisplayVal = FString::Printf(TEXT("(%s)"), *Pin->DefaultTextValue.ToString());
+		}
+	}
+
+	return DisplayVal;
+}
+
+//////////////////////////////////////////////////////////////////////////
+// SFindInMaterial
+
+void SFindInMaterial::Construct(const FArguments& InArgs, TSharedPtr<FMaterialEditor> InMaterialEditor)
+{
+	MaterialEditorPtr = InMaterialEditor;
+
+	this->ChildSlot
+	[
+		SNew(SVerticalBox)
+		+ SVerticalBox::Slot()
+		.AutoHeight()
+		[
+			SNew(SHorizontalBox)
+			+ SHorizontalBox::Slot()
+			.FillWidth(1)
+			[
+				SAssignNew(SearchTextField, SSearchBox)
+				.HintText(LOCTEXT("GraphSearchHint", "Search"))
+				.OnTextChanged(this, &SFindInMaterial::OnSearchTextChanged)
+				.OnTextCommitted(this, &SFindInMaterial::OnSearchTextCommitted)
+			]
+		]
+		+ SVerticalBox::Slot()
+		.FillHeight(1.0f)
+		.Padding(0.f, 4.f, 0.f, 0.f)
+		[
+			SNew(SBorder)
+			.BorderImage(FEditorStyle::GetBrush("Menu.Background"))
+			[
+				SAssignNew(TreeView, STreeViewType)
+				.ItemHeight(24)
+				.TreeItemsSource(&ItemsFound)
+				.OnGenerateRow(this, &SFindInMaterial::OnGenerateRow)
+				.OnGetChildren(this, &SFindInMaterial::OnGetChildren)
+				.OnSelectionChanged(this, &SFindInMaterial::OnTreeSelectionChanged)
+				.SelectionMode(ESelectionMode::Multi)
+			]
+		]
+	];
+}
+
+void SFindInMaterial::FocusForUse()
+{
+	// NOTE: Careful, GeneratePathToWidget can be reentrant in that it can call visibility delegates and such
+	FWidgetPath FilterTextBoxWidgetPath;
+	FSlateApplication::Get().GeneratePathToWidgetUnchecked(SearchTextField.ToSharedRef(), FilterTextBoxWidgetPath);
+
+	// Set keyboard focus directly
+	FSlateApplication::Get().SetKeyboardFocus(FilterTextBoxWidgetPath, EKeyboardFocusCause::SetDirectly);
+}
+
+void SFindInMaterial::OnSearchTextChanged(const FText& Text)
+{
+	SearchValue = Text.ToString();
+}
+
+void SFindInMaterial::OnSearchTextCommitted(const FText& Text, ETextCommit::Type CommitType)
+{
+	if (CommitType == ETextCommit::OnEnter)
+	{
+		InitiateSearch();
+	}
+}
+
+void SFindInMaterial::InitiateSearch()
+{
+	TArray<FString> Tokens;
+	if (SearchValue.Contains("\"") && SearchValue.ParseIntoArray(&Tokens, TEXT("\""), true) > 0)
+	{
+		for (auto &TokenIt : Tokens)
+		{
+			// we have the token, we don't need the quotes anymore, they'll just confused the comparison later on
+			TokenIt = TokenIt.TrimQuotes();
+			// We remove the spaces as all later comparison strings will also be de-spaced
+			TokenIt = TokenIt.Replace(TEXT(" "), TEXT(""));
+		}
+
+		// due to being able to handle multiple quoted blocks like ("Make Epic" "Game Now") we can end up with
+		// and empty string between (" ") blocks so this simply removes them
+		struct FRemoveMatchingStrings
+		{
+			bool operator()(const FString& RemovalCandidate) const
+			{
+				return RemovalCandidate.IsEmpty();
+			}
+		};
+		Tokens.RemoveAll(FRemoveMatchingStrings());
+	}
+	else
+	{
+		// unquoted search equivalent to a match-any-of search
+		SearchValue.ParseIntoArray(&Tokens, TEXT(" "), true);
+	}
+
+	for (auto It(ItemsFound.CreateIterator()); It; ++It)
+	{
+		TreeView->SetItemExpansion(*It, false);
+	}
+	ItemsFound.Empty();
+	if (Tokens.Num() > 0)
+	{
+		HighlightText = FText::FromString(SearchValue);
+		MatchTokens(Tokens);
+	}
+
+	// Insert a fake result to inform user if none found
+	if (ItemsFound.Num() == 0)
+	{
+		ItemsFound.Add(FSearchResult(new FFindInMaterialResult(LOCTEXT("MaterialSearchNoResults", "No Results found").ToString())));
+	}
+
+	TreeView->RequestTreeRefresh();
+
+	for (auto It(ItemsFound.CreateIterator()); It; ++It)
+	{
+		TreeView->SetItemExpansion(*It, true);
+	}
+}
+
+void SFindInMaterial::MatchTokens(const TArray<FString> &Tokens)
+{
+	RootSearchResult.Reset();
+
+	UEdGraph* Graph = MaterialEditorPtr.Pin()->Material->MaterialGraph;
+
+	if (Graph == NULL)
+	{
+		return;
+	}
+
+	RootSearchResult = FSearchResult(new FFindInMaterialResult(FString("BehaviorTreeRoot")));
+
+	for (auto It(Graph->Nodes.CreateConstIterator()); It; ++It)
+	{
+		UEdGraphNode* Node = *It;
+
+		const FString NodeName = Node->GetNodeTitle(ENodeTitleType::ListView);
+		FSearchResult NodeResult(new FFindInMaterialResult(NodeName, RootSearchResult, Node));
+
+		FString NodeSearchString = NodeName + Node->NodeComment;
+		NodeSearchString = NodeSearchString.Replace(TEXT(" "), TEXT(""));
+
+		bool bNodeMatchesSearch = StringMatchesSearchTokens(Tokens, NodeSearchString);
+
+		// Use old Material Expression search functions too
+		if (!bNodeMatchesSearch)
+		{
+			bool bMatchesAllTokens = true;
+			if (UMaterialGraphNode* MatNode = Cast<UMaterialGraphNode>(Node))
+			{
+				for (int32 Index = 0; Index < Tokens.Num(); ++Index)
+				{
+					if (!MatNode->MaterialExpression->MatchesSearchQuery(*Tokens[Index]))
+					{
+						bMatchesAllTokens = false;
+						break;
+					}
+				}
+			}
+			else if (UMaterialGraphNode_Comment* MatComment = Cast<UMaterialGraphNode_Comment>(Node))
+			{
+				for (int32 Index = 0; Index < Tokens.Num(); ++Index)
+				{
+					if (!MatComment->MaterialExpressionComment->MatchesSearchQuery(*Tokens[Index]))
+					{
+						bMatchesAllTokens = false;
+						break;
+					}
+				}
+			}
+			else
+			{
+				bMatchesAllTokens = false;
+			}
+			if (bMatchesAllTokens)
+			{
+				bNodeMatchesSearch = true;
+			}
+		}
+
+		for (TArray<UEdGraphPin*>::TIterator PinIt(Node->Pins); PinIt; ++PinIt)
+		{
+			UEdGraphPin* Pin = *PinIt;
+			if (Pin && Pin->PinFriendlyName != TEXT(" "))
+			{
+				FString PinName = Pin->GetSchema()->GetPinDisplayName(Pin);
+				FString PinSearchString = Pin->PinName + Pin->PinFriendlyName + Pin->DefaultValue + Pin->PinType.PinCategory + Pin->PinType.PinSubCategory + (Pin->PinType.PinSubCategoryObject.IsValid() ? Pin->PinType.PinSubCategoryObject.Get()->GetFullName() : TEXT(""));
+				PinSearchString = PinSearchString.Replace(TEXT(" "), TEXT(""));
+				if (StringMatchesSearchTokens(Tokens, PinSearchString))
+				{
+					FSearchResult PinResult(new FFindInMaterialResult(PinName, NodeResult, Pin));
+					NodeResult->Children.Add(PinResult);
+				}
+			}
+		}
+
+		if ((NodeResult->Children.Num() > 0) || bNodeMatchesSearch)
+		{
+			ItemsFound.Add(NodeResult);
+		}
+	}
+}
+
+TSharedRef<ITableRow> SFindInMaterial::OnGenerateRow(FSearchResult InItem, const TSharedRef<STableViewBase>& OwnerTable)
+{
+	bool bIsACategoryWidget = false;//!bIsInFindWithinBlueprintMode && !InItem->Parent.IsValid();
+
+	if (bIsACategoryWidget)
+	{
+		return SNew(STableRow< TSharedPtr<FFindInMaterialResult> >, OwnerTable)
+			[
+				SNew(SBorder)
+				.VAlign(VAlign_Center)
+				.BorderImage(FEditorStyle::GetBrush("PropertyWindow.CategoryBackground"))
+				.Padding(FMargin(2.0f))
+				.ForegroundColor(FEditorStyle::GetColor("PropertyWindow.CategoryForeground"))
+				[
+					SNew(STextBlock)
+					.Text(InItem->Value)
+					.ToolTipText(LOCTEXT("BlueprintCatSearchToolTip", "Blueprint").ToString())
+				]
+			];
+	}
+	else // Functions/Event/Pin widget
+	{
+		FString CommentText = InItem->GetCommentText();
+
+		return SNew(STableRow< TSharedPtr<FFindInMaterialResult> >, OwnerTable)
+			[
+				SNew(SHorizontalBox)
+				+ SHorizontalBox::Slot()
+				.AutoWidth()
+				[
+					InItem->CreateIcon()
+				]
+				+ SHorizontalBox::Slot()
+					.AutoWidth()
+					.VAlign(VAlign_Center)
+					.Padding(2, 0)
+					[
+						SNew(STextBlock)
+						.Text(InItem->Value)
+						.HighlightText(HighlightText)
+						.ToolTipText(FString::Printf(*LOCTEXT("BlueprintResultSearchToolTip", "%s : %s").ToString(), *InItem->GetCategory(), *InItem->Value))
+					]
+				+ SHorizontalBox::Slot()
+					.FillWidth(1)
+					.HAlign(HAlign_Right)
+					.VAlign(VAlign_Center)
+					.Padding(2, 0)
+					[
+						SNew(STextBlock)
+						.Text(InItem->GetValueText())
+						.HighlightText(HighlightText)
+					]
+				+ SHorizontalBox::Slot()
+					.FillWidth(1)
+					.HAlign(HAlign_Right)
+					.VAlign(VAlign_Center)
+					.Padding(2, 0)
+					[
+						SNew(STextBlock)
+						.Text(CommentText.IsEmpty() ? FString() : FString::Printf(TEXT("Node Comment:[%s]"), *CommentText))
+						.ColorAndOpacity(FLinearColor::Yellow)
+						.HighlightText(HighlightText)
+					]
+			];
+	}
+}
+
+void SFindInMaterial::OnGetChildren(FSearchResult InItem, TArray< FSearchResult >& OutChildren)
+{
+	OutChildren += InItem->Children;
+}
+
+void SFindInMaterial::OnTreeSelectionChanged(FSearchResult Item, ESelectInfo::Type)
+{
+	if (Item.IsValid())
+	{
+		Item->OnClick(MaterialEditorPtr);
+	}
+}
+
+bool SFindInMaterial::StringMatchesSearchTokens(const TArray<FString>& Tokens, const FString& ComparisonString)
+{
+	bool bFoundAllTokens = true;
+	//search the entry for each token, it must have all of them to pass
+	for (auto TokIT(Tokens.CreateConstIterator()); TokIT; ++TokIT)
+	{
+		const FString& Token = *TokIT;
+		if (!ComparisonString.Contains(Token))
+		{
+			bFoundAllTokens = false;
+			break;
+		}
+	}
+	return bFoundAllTokens;
+}
+
+/////////////////////////////////////////////////////
+
+#undef LOCTEXT_NAMESPACE
diff --git a/Engine/Source/Editor/MaterialEditor/Private/MaterialEditor.cpp b/Engine/Source/Editor/MaterialEditor/Private/MaterialEditor.cpp
--- a/Engine/Source/Editor/MaterialEditor/Private/MaterialEditor.cpp
+++ b/Engine/Source/Editor/MaterialEditor/Private/MaterialEditor.cpp
@@ -43,6 +43,7 @@ const FName FMaterialEditor::PropertiesTabId( TEXT( "MaterialEditor_MaterialProp
 const FName FMaterialEditor::HLSLCodeTabId( TEXT( "MaterialEditor_HLSLCode" ) );
 const FName FMaterialEditor::PaletteTabId( TEXT( "MaterialEditor_Palette" ) );
 const FName FMaterialEditor::StatsTabId( TEXT( "MaterialEditor_Stats" ) );
+const FName FMaterialEditor::FindTabId( TEXT( "MaterialEditor_Find" ) );
 
 ///////////////////////////
 // FMatExpressionPreview //
@@ -128,14 +129,18 @@ void FMaterialEditor::RegisterTabSpawners(const TSharedRef<class FTabManager>& T
 		.SetDisplayName( LOCTEXT("DetailsTab", "Details") )
 		.SetGroup( MenuStructure.GetAssetEditorCategory() );
 
-		TabManager->RegisterTabSpawner( PaletteTabId, FOnSpawnTab::CreateSP(this, &FMaterialEditor::SpawnTab_Palette) )
-			.SetDisplayName( LOCTEXT("PaletteTab", "Palette") )
-			.SetGroup( MenuStructure.GetAssetEditorCategory() );
+	TabManager->RegisterTabSpawner( PaletteTabId, FOnSpawnTab::CreateSP(this, &FMaterialEditor::SpawnTab_Palette) )
+		.SetDisplayName( LOCTEXT("PaletteTab", "Palette") )
+		.SetGroup( MenuStructure.GetAssetEditorCategory() );
 
-		TabManager->RegisterTabSpawner( StatsTabId, FOnSpawnTab::CreateSP(this, &FMaterialEditor::SpawnTab_Stats) )
-			.SetDisplayName( LOCTEXT("StatsTab", "Stats") )
-			.SetGroup( MenuStructure.GetAssetEditorCategory() );
+	TabManager->RegisterTabSpawner( StatsTabId, FOnSpawnTab::CreateSP(this, &FMaterialEditor::SpawnTab_Stats) )
+		.SetDisplayName( LOCTEXT("StatsTab", "Stats") )
+		.SetGroup( MenuStructure.GetAssetEditorCategory() );
 	
+	TabManager->RegisterTabSpawner(FindTabId, FOnSpawnTab::CreateSP(this, &FMaterialEditor::SpawnTab_Find))
+		.SetDisplayName(LOCTEXT("FindTab", "Find Results"))
+		.SetGroup(MenuStructure.GetAssetEditorCategory());
+
 	TabManager->RegisterTabSpawner( HLSLCodeTabId, FOnSpawnTab::CreateSP(this, &FMaterialEditor::SpawnTab_HLSLCode) )
 		.SetDisplayName( LOCTEXT("HLSLCodeTab", "HLSL Code") )
 		.SetGroup( MenuStructure.GetAssetEditorCategory() );
@@ -149,8 +154,9 @@ void FMaterialEditor::UnregisterTabSpawners(const TSharedRef<class FTabManager>&
 	TabManager->UnregisterTabSpawner( PreviewTabId );
 	TabManager->UnregisterTabSpawner( GraphCanvasTabId );
 	TabManager->UnregisterTabSpawner( PropertiesTabId );
-		TabManager->UnregisterTabSpawner( PaletteTabId );
-		TabManager->UnregisterTabSpawner( StatsTabId );
+	TabManager->UnregisterTabSpawner( PaletteTabId );
+	TabManager->UnregisterTabSpawner( StatsTabId );
+	TabManager->UnregisterTabSpawner( FindTabId );
 	TabManager->UnregisterTabSpawner( HLSLCodeTabId );
 }
 
@@ -216,20 +222,20 @@ void FMaterialEditor::InitMaterialEditor( const EToolkitMode::Type Mode, const T
 	bStatsFromPreviewMaterial = false;
 	ColorPickerObject = NULL;
 
-		// Support undo/redo
-		Material->SetFlags(RF_Transactional);
+	// Support undo/redo
+	Material->SetFlags(RF_Transactional);
 
-		GEditor->RegisterForUndo(this);
+	GEditor->RegisterForUndo(this);
 
-		if (!Material->MaterialGraph)
-		{
-			Material->MaterialGraph = CastChecked<UMaterialGraph>(FBlueprintEditorUtils::CreateNewGraph(Material, NAME_None, UMaterialGraph::StaticClass(), UMaterialGraphSchema::StaticClass()));
-		}
-		Material->MaterialGraph->Material = Material;
-		Material->MaterialGraph->MaterialFunction = MaterialFunction;
-		Material->MaterialGraph->RealtimeDelegate.BindSP(this, &FMaterialEditor::IsToggleRealTimeExpressionsChecked);
-		Material->MaterialGraph->MaterialDirtyDelegate.BindSP(this, &FMaterialEditor::SetMaterialDirty);
-		Material->MaterialGraph->ToggleCollapsedDelegate.BindSP(this, &FMaterialEditor::ToggleCollapsed);
+	if (!Material->MaterialGraph)
+	{
+		Material->MaterialGraph = CastChecked<UMaterialGraph>(FBlueprintEditorUtils::CreateNewGraph(Material, NAME_None, UMaterialGraph::StaticClass(), UMaterialGraphSchema::StaticClass()));
+	}
+	Material->MaterialGraph->Material = Material;
+	Material->MaterialGraph->MaterialFunction = MaterialFunction;
+	Material->MaterialGraph->RealtimeDelegate.BindSP(this, &FMaterialEditor::IsToggleRealTimeExpressionsChecked);
+	Material->MaterialGraph->MaterialDirtyDelegate.BindSP(this, &FMaterialEditor::SetMaterialDirty);
+	Material->MaterialGraph->ToggleCollapsedDelegate.BindSP(this, &FMaterialEditor::ToggleCollapsed);
 
 	// copy material usage
 	for( int32 Usage=0; Usage < MATUSAGE_MAX; Usage++ )
@@ -258,13 +264,13 @@ void FMaterialEditor::InitMaterialEditor( const EToolkitMode::Type Mode, const T
 	CreateInternalWidgets();
 
 	// Do setup previously done in SMaterialEditorCanvas
-		SetPreviewMaterial(Material);
-		Material->bIsPreviewMaterial = true;
-		FMaterialEditorUtilities::InitExpressions(Material);
+	SetPreviewMaterial(Material);
+	Material->bIsPreviewMaterial = true;
+	FMaterialEditorUtilities::InitExpressions(Material);
 
 	BindCommands();
 
-	const TSharedRef<FTabManager::FLayout> StandaloneDefaultLayout = FTabManager::NewLayout("Standalone_MaterialEditor_Layout_v4")
+	const TSharedRef<FTabManager::FLayout> StandaloneDefaultLayout = FTabManager::NewLayout("Standalone_MaterialEditor_Layout_v5")
 	->AddArea
 	(
 		FTabManager::NewPrimaryArea() ->SetOrientation(Orient_Vertical)
@@ -309,6 +315,7 @@ void FMaterialEditor::InitMaterialEditor( const EToolkitMode::Type Mode, const T
 					FTabManager::NewStack()
 					->SetSizeCoefficient( 0.20f )
 					->AddTab( StatsTabId, ETabState::ClosedTab )
+					->AddTab( FindTabId, ETabState::ClosedTab )
 				)
 				
 			)
@@ -405,25 +412,25 @@ void FMaterialEditor::InitMaterialEditor( const EToolkitMode::Type Mode, const T
 						bSetPreviewExpression = true;
 
 						// Preview the last output previewed
-							SetPreviewExpression(FunctionOutput);
-						}
+						SetPreviewExpression(FunctionOutput);
 					}
 				}
+			}
 
 			if (!bSetPreviewExpression && FirstOutput)
 			{
-					SetPreviewExpression(FirstOutput);
-				}
+				SetPreviewExpression(FirstOutput);
 			}
 		}
-
-		// TODO: Remove this once Expressions are saved as flipped
-		FlipExpressionPositions(Material->Expressions, Material->EditorComments, Material);
-		Material->MaterialGraph->RebuildGraph();
-		RecenterEditor();
-		ForceRefreshExpressionPreviews();
 	}
 
+	// TODO: Remove this once Expressions are saved as flipped
+	FlipExpressionPositions(Material->Expressions, Material->EditorComments, Material);
+	Material->MaterialGraph->RebuildGraph();
+	RecenterEditor();
+	ForceRefreshExpressionPreviews();
+}
+
 FMaterialEditor::FMaterialEditor()
 	: bMaterialDirty(false)
 	, bStatsFromPreviewMaterial(false)
@@ -456,16 +463,16 @@ FMaterialEditor::~FMaterialEditor()
 
 	MaterialDetailsView.Reset();
 
-		{
-			SCOPED_SUSPEND_RENDERING_THREAD(true);
+	{
+		SCOPED_SUSPEND_RENDERING_THREAD(true);
 	
-			ExpressionPreviews.Empty();
-		}
+		ExpressionPreviews.Empty();
+	}
 	
-		check( !ScopedTransaction );
+	check( !ScopedTransaction );
 	
-		GEditor->UnregisterForUndo( this );
-	}
+	GEditor->UnregisterForUndo( this );
+}
 
 void FMaterialEditor::GetAllMaterialExpressionGroups(TArray<FString>* OutGroups)
 {
@@ -498,13 +505,13 @@ void FMaterialEditor::CreateInternalWidgets()
 
 	FPropertyEditorModule& PropertyEditorModule = FModuleManager::GetModuleChecked<FPropertyEditorModule>( "PropertyEditor" );
 
-		GraphEditor = CreateGraphEditorWidget();
-		// Manually set zoom level to avoid deferred zooming
-		GraphEditor->SetViewLocation(FVector2D::ZeroVector, 1);
+	GraphEditor = CreateGraphEditorWidget();
+	// Manually set zoom level to avoid deferred zooming
+	GraphEditor->SetViewLocation(FVector2D::ZeroVector, 1);
+
+	const FDetailsViewArgs DetailsViewArgs( false, false, true, false, true, this );
+	MaterialDetailsView = PropertyEditorModule.CreateDetailView( DetailsViewArgs );
 
-		const FDetailsViewArgs DetailsViewArgs( false, false, true, false, true, this );
-		MaterialDetailsView = PropertyEditorModule.CreateDetailView( DetailsViewArgs );
-	
 	FOnGetDetailCustomizationInstance LayoutExpressionParameterDetails = FOnGetDetailCustomizationInstance::CreateStatic(
 		&FMaterialExpressionParameterDetails::MakeInstance, FOnCollectParameterGroups::CreateSP(this, &FMaterialEditor::GetAllMaterialExpressionGroups) );
 
@@ -530,22 +537,25 @@ void FMaterialEditor::CreateInternalWidgets()
 		LayoutCollectionParameterDetails
 		);
 
-		Palette = SNew(SMaterialPalette, SharedThis(this));
+	Palette = SNew(SMaterialPalette, SharedThis(this));
 
-		FMessageLogModule& MessageLogModule = FModuleManager::LoadModuleChecked<FMessageLogModule>("MessageLog");
-		FMessageLogInitializationOptions LogOptions;
-		// Show Pages so that user is never allowed to clear log messages
-		LogOptions.bShowPages = true;
-		LogOptions.MaxPageCount = 1;
-		StatsListing = MessageLogModule.CreateLogListing( "MaterialEditorStats", LogOptions );
+	FMessageLogModule& MessageLogModule = FModuleManager::LoadModuleChecked<FMessageLogModule>("MessageLog");
+	FMessageLogInitializationOptions LogOptions;
+	// Show Pages so that user is never allowed to clear log messages
+	LogOptions.bShowPages = true;
+	LogOptions.MaxPageCount = 1;
+	StatsListing = MessageLogModule.CreateLogListing( "MaterialEditorStats", LogOptions );
 
-		Stats = 
-			SNew(SBorder)
-			.Padding(0.0f)
-			.BorderImage(FEditorStyle::GetBrush("ToolPanel.GroupBorder"))
-			[
-				MessageLogModule.CreateLogListingWidget( StatsListing.ToSharedRef() )
-			];
+	Stats = 
+		SNew(SBorder)
+		.Padding(0.0f)
+		.BorderImage(FEditorStyle::GetBrush("ToolPanel.GroupBorder"))
+		[
+			MessageLogModule.CreateLogListingWidget( StatsListing.ToSharedRef() )
+		];
+
+	FindResults =
+		SNew(SFindInMaterial, SharedThis(this));
 
 	CodeViewUtility =
 		SNew(SVerticalBox)
@@ -621,9 +631,9 @@ FLinearColor FMaterialEditor::GetWorldCentricTabColorScale() const
 
 void FMaterialEditor::Tick( float InDeltaTime )
 {
-		UpdateMaterialInfoList();
-		UpdateGraphNodeStates();
-	}
+	UpdateMaterialInfoList();
+	UpdateGraphNodeStates();
+}
 
 TStatId FMaterialEditor::GetStatId() const
 {
@@ -652,7 +662,7 @@ void FMaterialEditor::ExtendToolbar()
 {
 	struct Local
 	{
-		static void FillToolbar(FToolBarBuilder& ToolbarBuilder, TSharedRef<SWidget> Searchbox)
+		static void FillToolbar(FToolBarBuilder& ToolbarBuilder)
 		{
 			ToolbarBuilder.BeginSection("Apply");
 			{
@@ -660,6 +670,12 @@ void FMaterialEditor::ExtendToolbar()
 			}
 			ToolbarBuilder.EndSection();
 	
+			ToolbarBuilder.BeginSection("Search");
+			{
+				ToolbarBuilder.AddToolBarButton(FMaterialEditorCommands::Get().FindInMaterial);
+			}
+			ToolbarBuilder.EndSection();
+
 			ToolbarBuilder.BeginSection("Graph");
 			{
 				ToolbarBuilder.AddToolBarButton(FMaterialEditorCommands::Get().CameraHome);
@@ -671,39 +687,16 @@ void FMaterialEditor::ExtendToolbar()
 				ToolbarBuilder.AddToolBarButton(FMaterialEditorCommands::Get().ToggleMobileStats);
 			}
 			ToolbarBuilder.EndSection();
-	
-			ToolbarBuilder.BeginSection("Search");
-			{
-				ToolbarBuilder.AddWidget(Searchbox);
-			}
-			ToolbarBuilder.EndSection();
 		}
 	};
 
 	TSharedPtr<FExtender> ToolbarExtender = MakeShareable(new FExtender);
 
-	float SearchWidth = 200.f;
-	TSharedRef<SHorizontalBox> Searchbox = SNew(SHorizontalBox);
-	TSharedRef<SWidget> SearchWidget = Searchbox;
-
-		Searchbox->AddSlot()
-		.MaxWidth(SearchWidth)
-		.FillWidth(1.f)
-		.VAlign(VAlign_Bottom)
-		[
-			SNew( SSearchBox )
-			.ToolTipText( LOCTEXT("GraphSearchHint", "Type here to search expressions in the graph").ToString() )
-			.OnTextChanged( this, &FMaterialEditor::OnGraphSearchChanged )
-			.OnTextCommitted( this, &FMaterialEditor::OnGraphSearchCommitted )
-			.OnSearch( this, &FMaterialEditor::OnSearch )
-			.MinDesiredWidth(SearchWidth)
-		];
-
 	ToolbarExtender->AddToolBarExtension(
 		"Asset",
 		EExtensionHook::After,
 		GetToolkitCommands(),
-		FToolBarExtensionDelegate::CreateStatic( &Local::FillToolbar, SearchWidget )
+		FToolBarExtensionDelegate::CreateStatic( &Local::FillToolbar )
 		);
 	
 	AddToolbarExtender(ToolbarExtender);
@@ -849,40 +842,40 @@ void FMaterialEditor::DrawMaterialInfoStrings(
 
 void FMaterialEditor::DrawMessages( FViewport* InViewport, FCanvas* Canvas )
 {
-		if( PreviewExpression != NULL )
-		{
-			Canvas->PushAbsoluteTransform( FMatrix::Identity );
+	if( PreviewExpression != NULL )
+	{
+		Canvas->PushAbsoluteTransform( FMatrix::Identity );
 
-			// The message to display in the viewport.
-			FString Name = FString::Printf( TEXT("Previewing: %s"), *PreviewExpression->GetName() );
+		// The message to display in the viewport.
+		FString Name = FString::Printf( TEXT("Previewing: %s"), *PreviewExpression->GetName() );
 
-			// Size of the tile we are about to draw.  Should extend the length of the view in X.
-			const FIntPoint TileSize( InViewport->GetSizeXY().X, 25);
+		// Size of the tile we are about to draw.  Should extend the length of the view in X.
+		const FIntPoint TileSize( InViewport->GetSizeXY().X, 25);
 
-			const FColor PreviewColor( 70,100,200 );
-			const FColor FontColor( 255,255,128 );
+		const FColor PreviewColor( 70,100,200 );
+		const FColor FontColor( 255,255,128 );
 
-			UFont* FontToUse = GEditor->EditorFont;
+		UFont* FontToUse = GEditor->EditorFont;
 
-			Canvas->DrawTile(  0.0f, 0.0f, TileSize.X, TileSize.Y, 0.0f, 0.0f, 0.0f, 0.0f, PreviewColor );
+		Canvas->DrawTile(  0.0f, 0.0f, TileSize.X, TileSize.Y, 0.0f, 0.0f, 0.0f, 0.0f, PreviewColor );
 
-			int32 XL, YL;
+		int32 XL, YL;
+		StringSize( FontToUse, XL, YL, *Name );
+		if( XL > TileSize.X )
+		{
+			// There isn't enough room to show the preview expression name
+			Name = TEXT("Previewing");
 			StringSize( FontToUse, XL, YL, *Name );
-			if( XL > TileSize.X )
-			{
-				// There isn't enough room to show the preview expression name
-				Name = TEXT("Previewing");
-				StringSize( FontToUse, XL, YL, *Name );
-			}
+		}
 
-			// Center the string in the middle of the tile.
-			const FIntPoint StringPos( (TileSize.X-XL)/2, ((TileSize.Y-YL)/2)+1 );
-			// Draw the preview message
-			Canvas->DrawShadowedString(  StringPos.X, StringPos.Y, *Name, FontToUse, FontColor );
+		// Center the string in the middle of the tile.
+		const FIntPoint StringPos( (TileSize.X-XL)/2, ((TileSize.Y-YL)/2)+1 );
+		// Draw the preview message
+		Canvas->DrawShadowedString(  StringPos.X, StringPos.Y, *Name, FontToUse, FontColor );
 
-			Canvas->PopTransform();
-		}
+		Canvas->PopTransform();
 	}
+}
 
 void FMaterialEditor::RecenterEditor()
 {
@@ -920,7 +913,7 @@ void FMaterialEditor::RecenterEditor()
 
 	if (FocusNode)
 	{
-		GraphEditor->JumpToNode(FocusNode, false);
+		JumpToNode(FocusNode);
 	}
 	else
 	{
@@ -970,9 +963,9 @@ void FMaterialEditor::LoadEditorSettings()
 {
 	EditorOptions = ConstructObject<UMaterialEditorOptions>( UMaterialEditorOptions::StaticClass() );
 	
-		if (EditorOptions->bHideUnusedConnectors) {OnShowConnectors();}
-		if (EditorOptions->bAlwaysRefreshAllPreviews) {OnAlwaysRefreshAllPreviews();}
-		if (EditorOptions->bRealtimeExpressionViewport) {ToggleRealTimeExpressions();}
+	if (EditorOptions->bHideUnusedConnectors) {OnShowConnectors();}
+	if (EditorOptions->bAlwaysRefreshAllPreviews) {OnAlwaysRefreshAllPreviews();}
+	if (EditorOptions->bRealtimeExpressionViewport) {ToggleRealTimeExpressions();}
 
 	if ( Viewport.IsValid() )
 	{
@@ -1009,29 +1002,15 @@ void FMaterialEditor::SaveEditorSettings()
 		EditorOptions->bShowBackground				= Viewport->IsTogglePreviewBackgroundChecked();
 		EditorOptions->bRealtimeMaterialViewport	= Viewport->IsRealtime();
 		EditorOptions->bShowMobileStats				= bShowMobileStats;
-			EditorOptions->bHideUnusedConnectors		= !IsOnShowConnectorsChecked();
-			EditorOptions->bAlwaysRefreshAllPreviews	= IsOnAlwaysRefreshAllPreviews();
-			EditorOptions->bRealtimeExpressionViewport	= IsToggleRealTimeExpressionsChecked();
+		EditorOptions->bHideUnusedConnectors		= !IsOnShowConnectorsChecked();
+		EditorOptions->bAlwaysRefreshAllPreviews	= IsOnAlwaysRefreshAllPreviews();
+		EditorOptions->bRealtimeExpressionViewport	= IsToggleRealTimeExpressionsChecked();
 		EditorOptions->SaveConfig();
 	}
 
 	GConfig->SetInt(TEXT("MaterialEditor"), TEXT("PrimType"), Viewport->PreviewPrimType, GEditorUserSettingsIni);
 }
 
-void FMaterialEditor::OnGraphSearchChanged( const FText& InFilterText )
-{
-		SearchQuery = InFilterText.ToString();
-		UpdateSearch(true);
-	}
-
-void FMaterialEditor::OnGraphSearchCommitted(const FText& NewTypeInValue, ETextCommit::Type CommitInfo)
-{
-	if (CommitInfo ==  ETextCommit::OnEnter)
-	{
-			OnSearch( SSearchBox::Next );
-		}
-	}
-
 FString FMaterialEditor::GetCodeViewText() const
 {
 	return HLSLCode;
@@ -1103,10 +1082,10 @@ void FMaterialEditor::UpdatePreviewMaterial( )
 {
 	bStatsFromPreviewMaterial = true;
 
-		if( PreviewExpression && ExpressionPreviewMaterial )
-		{
-			PreviewExpression->ConnectToPreviewMaterial(ExpressionPreviewMaterial,0);
-		}
+	if( PreviewExpression && ExpressionPreviewMaterial )
+	{
+		PreviewExpression->ConnectToPreviewMaterial(ExpressionPreviewMaterial,0);
+	}
 
 	if(PreviewExpression)
 	{
@@ -1201,7 +1180,7 @@ void FMaterialEditor::UpdateOriginalMaterial()
 	}
 
 	// Make sure any graph position changes that might not have been copied are taken into account
-		Material->MaterialGraph->LinkMaterialExpressionsFromGraph();
+	Material->MaterialGraph->LinkMaterialExpressionsFromGraph();
 
 	//remove any memory copies of shader files, so they will be reloaded from disk
 	//this way the material editor can be used for quick shader iteration
@@ -1241,8 +1220,8 @@ void FMaterialEditor::UpdateOriginalMaterial()
 		MaterialFunction->ParentFunction->ThumbnailInfo = OriginalThumbnailInfo;
 		MaterialFunction->ThumbnailInfo = ThumbnailInfo;
 
-			// TODO: Remove this once Expressions are saved as flipped
-			FlipExpressionPositions(MaterialFunction->ParentFunction->FunctionExpressions, MaterialFunction->ParentFunction->FunctionEditorComments);
+		// TODO: Remove this once Expressions are saved as flipped
+		FlipExpressionPositions(MaterialFunction->ParentFunction->FunctionExpressions, MaterialFunction->ParentFunction->FunctionEditorComments);
 
 		// Restore RF_Standalone on the original material function, as it had been removed from the preview material so that it could be GC'd.
 		MaterialFunction->ParentFunction->SetFlags( RF_Standalone );
@@ -1372,8 +1351,8 @@ void FMaterialEditor::UpdateOriginalMaterial()
 			OriginalMaterial->ThumbnailInfo = OriginalThumbnailInfo;
 			Material->ThumbnailInfo = ThumbnailInfo;
 
-				// TODO: Remove this once Expressions are saved as flipped
-				FlipExpressionPositions(OriginalMaterial->Expressions, OriginalMaterial->EditorComments, OriginalMaterial);
+			// TODO: Remove this once Expressions are saved as flipped
+			FlipExpressionPositions(OriginalMaterial->Expressions, OriginalMaterial->EditorComments, OriginalMaterial);
 
 			// Change the original material object to the new original material
 			OriginalMaterialObject = OriginalMaterial;
@@ -1591,7 +1570,6 @@ void FMaterialEditor::UpdateGraphNodeStates()
 			MaterialNode->bIsPreviewExpression = (PreviewExpression == MaterialNode->MaterialExpression);
 			MaterialNode->bIsErrorExpression = (ErrorMaterialResource->GetErrorExpressions().Find(MaterialNode->MaterialExpression) != INDEX_NONE)
 												|| (ErrorMaterialResourceES2 && ErrorMaterialResourceES2->GetErrorExpressions().Find(MaterialNode->MaterialExpression) != INDEX_NONE);
-			MaterialNode->bIsCurrentSearchResult = (SelectedSearchResult >= 0 && SelectedSearchResult < SearchResults.Num() && SearchResults[SelectedSearchResult] == MaterialNode->MaterialExpression);
 
 			if (MaterialNode->bIsErrorExpression && !MaterialNode->bHasCompilerMessage)
 			{
@@ -1640,106 +1618,110 @@ void FMaterialEditor::BindCommands()
 		FCanExecuteAction(),
 		FIsActionChecked::CreateSP( Viewport.ToSharedRef(), &SMaterialEditorViewport::IsRealtime ) );
 
-		ToolkitCommands->MapAction(
-			FGenericCommands::Get().Undo,
-			FExecuteAction::CreateSP(this, &FMaterialEditor::UndoGraphAction));
-
-		ToolkitCommands->MapAction(
-			FGenericCommands::Get().Redo,
-			FExecuteAction::CreateSP(this, &FMaterialEditor::RedoGraphAction));
-
-		ToolkitCommands->MapAction(
-			Commands.CameraHome,
-			FExecuteAction::CreateSP(this, &FMaterialEditor::OnCameraHome),
-			FCanExecuteAction() );
-
-		ToolkitCommands->MapAction(
-			Commands.CleanUnusedExpressions,
-			FExecuteAction::CreateSP(this, &FMaterialEditor::CleanUnusedExpressions),
-			FCanExecuteAction() );
-
-		ToolkitCommands->MapAction(
-			Commands.ShowHideConnectors,
-			FExecuteAction::CreateSP(this, &FMaterialEditor::OnShowConnectors),
-			FCanExecuteAction(),
-			FIsActionChecked::CreateSP(this, &FMaterialEditor::IsOnShowConnectorsChecked));
-
-		ToolkitCommands->MapAction(
-			Commands.ToggleRealtimeExpressions,
-			FExecuteAction::CreateSP(this, &FMaterialEditor::ToggleRealTimeExpressions),
-			FCanExecuteAction(),
-			FIsActionChecked::CreateSP(this, &FMaterialEditor::IsToggleRealTimeExpressionsChecked));
-
-		ToolkitCommands->MapAction(
-			Commands.AlwaysRefreshAllPreviews,
-			FExecuteAction::CreateSP(this, &FMaterialEditor::OnAlwaysRefreshAllPreviews),
-			FCanExecuteAction(),
-			FIsActionChecked::CreateSP(this, &FMaterialEditor::IsOnAlwaysRefreshAllPreviews));
-
-		ToolkitCommands->MapAction(
-			Commands.ToggleMaterialStats,
-			FExecuteAction::CreateSP(this, &FMaterialEditor::ToggleStats),
-			FCanExecuteAction(),
-			FIsActionChecked::CreateSP(this, &FMaterialEditor::IsToggleStatsChecked));
-
-		ToolkitCommands->MapAction(
-			Commands.ToggleMobileStats,
-			FExecuteAction::CreateSP(this, &FMaterialEditor::ToggleMobileStats),
-			FCanExecuteAction(),
-			FIsActionChecked::CreateSP(this, &FMaterialEditor::IsToggleMobileStatsChecked));
-
-		ToolkitCommands->MapAction(
-			Commands.UseCurrentTexture,
-			FExecuteAction::CreateSP(this, &FMaterialEditor::OnUseCurrentTexture));
-
-		ToolkitCommands->MapAction(
-			Commands.ConvertObjects,
-			FExecuteAction::CreateSP(this, &FMaterialEditor::OnConvertObjects));
-
-		ToolkitCommands->MapAction(
-			Commands.ConvertToTextureObjects,
-			FExecuteAction::CreateSP(this, &FMaterialEditor::OnConvertTextures));
-
-		ToolkitCommands->MapAction(
-			Commands.ConvertToTextureSamples,
-			FExecuteAction::CreateSP(this, &FMaterialEditor::OnConvertTextures));
-
-		ToolkitCommands->MapAction(
-			Commands.StopPreviewNode,
-			FExecuteAction::CreateSP(this, &FMaterialEditor::OnPreviewNode));
-
-		ToolkitCommands->MapAction(
-			Commands.StartPreviewNode,
-			FExecuteAction::CreateSP(this, &FMaterialEditor::OnPreviewNode));
-
-		ToolkitCommands->MapAction(
-			Commands.EnableRealtimePreviewNode,
-			FExecuteAction::CreateSP(this, &FMaterialEditor::OnToggleRealtimePreview));
-
-		ToolkitCommands->MapAction(
-			Commands.DisableRealtimePreviewNode,
-			FExecuteAction::CreateSP(this, &FMaterialEditor::OnToggleRealtimePreview));
-
-		ToolkitCommands->MapAction(
-			Commands.SelectDownstreamNodes,
-			FExecuteAction::CreateSP(this, &FMaterialEditor::OnSelectDownsteamNodes));
-
-		ToolkitCommands->MapAction(
-			Commands.SelectUpstreamNodes,
-			FExecuteAction::CreateSP(this, &FMaterialEditor::OnSelectUpsteamNodes));
-
-		ToolkitCommands->MapAction(
-			Commands.RemoveFromFavorites,
-			FExecuteAction::CreateSP(this, &FMaterialEditor::RemoveSelectedExpressionFromFavorites));
+	ToolkitCommands->MapAction(
+		FGenericCommands::Get().Undo,
+		FExecuteAction::CreateSP(this, &FMaterialEditor::UndoGraphAction));
+
+	ToolkitCommands->MapAction(
+		FGenericCommands::Get().Redo,
+		FExecuteAction::CreateSP(this, &FMaterialEditor::RedoGraphAction));
+
+	ToolkitCommands->MapAction(
+		Commands.CameraHome,
+		FExecuteAction::CreateSP(this, &FMaterialEditor::OnCameraHome),
+		FCanExecuteAction() );
+
+	ToolkitCommands->MapAction(
+		Commands.CleanUnusedExpressions,
+		FExecuteAction::CreateSP(this, &FMaterialEditor::CleanUnusedExpressions),
+		FCanExecuteAction() );
+
+	ToolkitCommands->MapAction(
+		Commands.ShowHideConnectors,
+		FExecuteAction::CreateSP(this, &FMaterialEditor::OnShowConnectors),
+		FCanExecuteAction(),
+		FIsActionChecked::CreateSP(this, &FMaterialEditor::IsOnShowConnectorsChecked));
+
+	ToolkitCommands->MapAction(
+		Commands.ToggleRealtimeExpressions,
+		FExecuteAction::CreateSP(this, &FMaterialEditor::ToggleRealTimeExpressions),
+		FCanExecuteAction(),
+		FIsActionChecked::CreateSP(this, &FMaterialEditor::IsToggleRealTimeExpressionsChecked));
+
+	ToolkitCommands->MapAction(
+		Commands.AlwaysRefreshAllPreviews,
+		FExecuteAction::CreateSP(this, &FMaterialEditor::OnAlwaysRefreshAllPreviews),
+		FCanExecuteAction(),
+		FIsActionChecked::CreateSP(this, &FMaterialEditor::IsOnAlwaysRefreshAllPreviews));
+
+	ToolkitCommands->MapAction(
+		Commands.ToggleMaterialStats,
+		FExecuteAction::CreateSP(this, &FMaterialEditor::ToggleStats),
+		FCanExecuteAction(),
+		FIsActionChecked::CreateSP(this, &FMaterialEditor::IsToggleStatsChecked));
+
+	ToolkitCommands->MapAction(
+		Commands.ToggleMobileStats,
+		FExecuteAction::CreateSP(this, &FMaterialEditor::ToggleMobileStats),
+		FCanExecuteAction(),
+		FIsActionChecked::CreateSP(this, &FMaterialEditor::IsToggleMobileStatsChecked));
+
+	ToolkitCommands->MapAction(
+		Commands.UseCurrentTexture,
+		FExecuteAction::CreateSP(this, &FMaterialEditor::OnUseCurrentTexture));
+
+	ToolkitCommands->MapAction(
+		Commands.ConvertObjects,
+		FExecuteAction::CreateSP(this, &FMaterialEditor::OnConvertObjects));
+
+	ToolkitCommands->MapAction(
+		Commands.ConvertToTextureObjects,
+		FExecuteAction::CreateSP(this, &FMaterialEditor::OnConvertTextures));
+
+	ToolkitCommands->MapAction(
+		Commands.ConvertToTextureSamples,
+		FExecuteAction::CreateSP(this, &FMaterialEditor::OnConvertTextures));
+
+	ToolkitCommands->MapAction(
+		Commands.StopPreviewNode,
+		FExecuteAction::CreateSP(this, &FMaterialEditor::OnPreviewNode));
+
+	ToolkitCommands->MapAction(
+		Commands.StartPreviewNode,
+		FExecuteAction::CreateSP(this, &FMaterialEditor::OnPreviewNode));
+
+	ToolkitCommands->MapAction(
+		Commands.EnableRealtimePreviewNode,
+		FExecuteAction::CreateSP(this, &FMaterialEditor::OnToggleRealtimePreview));
+
+	ToolkitCommands->MapAction(
+		Commands.DisableRealtimePreviewNode,
+		FExecuteAction::CreateSP(this, &FMaterialEditor::OnToggleRealtimePreview));
+
+	ToolkitCommands->MapAction(
+		Commands.SelectDownstreamNodes,
+		FExecuteAction::CreateSP(this, &FMaterialEditor::OnSelectDownsteamNodes));
+
+	ToolkitCommands->MapAction(
+		Commands.SelectUpstreamNodes,
+		FExecuteAction::CreateSP(this, &FMaterialEditor::OnSelectUpsteamNodes));
+
+	ToolkitCommands->MapAction(
+		Commands.RemoveFromFavorites,
+		FExecuteAction::CreateSP(this, &FMaterialEditor::RemoveSelectedExpressionFromFavorites));
 		
-		ToolkitCommands->MapAction(
-			Commands.AddToFavorites,
-			FExecuteAction::CreateSP(this, &FMaterialEditor::AddSelectedExpressionToFavorites));
+	ToolkitCommands->MapAction(
+		Commands.AddToFavorites,
+		FExecuteAction::CreateSP(this, &FMaterialEditor::AddSelectedExpressionToFavorites));
 
-		ToolkitCommands->MapAction(
-			Commands.ForceRefreshPreviews,
-			FExecuteAction::CreateSP(this, &FMaterialEditor::OnForceRefreshPreviews));
-	}
+	ToolkitCommands->MapAction(
+		Commands.ForceRefreshPreviews,
+		FExecuteAction::CreateSP(this, &FMaterialEditor::OnForceRefreshPreviews));
+
+	ToolkitCommands->MapAction(
+		Commands.FindInMaterial,
+		FExecuteAction::CreateSP(this, &FMaterialEditor::OnFindInMaterial));
+}
 
 void FMaterialEditor::OnApply()
 {
@@ -2246,6 +2228,12 @@ void FMaterialEditor::OnCreateComment()
 	CreateNewMaterialExpressionComment(GraphEditor->GetPasteLocation());
 }
 
+void FMaterialEditor::OnFindInMaterial()
+{
+	TabManager->InvokeTab(FindTabId);
+	FindResults->FocusForUse();
+}
+
 void FMaterialEditor::RenameAssetFromRegistry(const FAssetData& InAddedAssetData, const FString& InNewName)
 {
 	// Grab the asset class, it will be checked for being a material function.
@@ -2253,9 +2241,9 @@ void FMaterialEditor::RenameAssetFromRegistry(const FAssetData& InAddedAssetData
 
 	if(Asset->IsChildOf(UMaterialFunction::StaticClass()))
 	{
-			ForceRefreshExpressionPreviews();
-		}
+		ForceRefreshExpressionPreviews();
 	}
+}
 
 void FMaterialEditor::OnMaterialUsageFlagsChanged(UMaterial* MaterialThatChanged, int32 FlagThatChanged)
 {
@@ -2369,28 +2357,21 @@ TSharedRef<SDockTab> FMaterialEditor::SpawnTab_Stats(const FSpawnTabArgs& Args)
 	return SpawnedTab;
 }
 
-void FMaterialEditor::OnSearch(SSearchBox::SearchDirection Direction)
+TSharedRef<SDockTab> FMaterialEditor::SpawnTab_Find(const FSpawnTabArgs& Args)
 {
-	if (SearchQuery.Len() > 0)
-	{
-		if (Direction == SSearchBox::Previous)
-		{
-			SelectedSearchResult--;
-			if( SelectedSearchResult < 0 )
-			{
-				SelectedSearchResult = FMath::Max<int32>(SearchResults.Num()-1,0);
-			}
-		}
-		else
-		{
-			SelectedSearchResult++;
-			if( SelectedSearchResult >= SearchResults.Num() )
-			{
-				SelectedSearchResult = 0;
-			}
-		}
-		ShowSearchResult();
-	}
+	check(Args.GetTabId() == FindTabId);
+
+	TSharedRef<SDockTab> SpawnedTab = SNew(SDockTab)
+		.Icon(FEditorStyle::GetBrush("Kismet.Tabs.FindResults"))
+		.Label(LOCTEXT("MaterialFindTitle", "Find Results"))
+		[
+			SNew(STutorialWrapper, TEXT("MaterialFind"))
+			[
+				FindResults.ToSharedRef()
+			]
+		];
+
+	return SpawnedTab;
 }
 
 void FMaterialEditor::SetPreviewExpression(UMaterialExpression* NewPreviewExpression)
@@ -2448,6 +2429,11 @@ void FMaterialEditor::SetPreviewExpression(UMaterialExpression* NewPreviewExpres
 	}
 }
 
+void FMaterialEditor::JumpToNode(const UEdGraphNode* Node)
+{
+	GraphEditor->JumpToNode(Node, false);
+}
+
 UMaterialExpression* FMaterialEditor::CreateNewMaterialExpression(UClass* NewExpressionClass, const FVector2D& NodePos, bool bAutoSelect, bool bAutoAssignResource)
 {
 	check( NewExpressionClass->IsChildOf(UMaterialExpression::StaticClass()) );
@@ -2601,7 +2587,6 @@ UMaterialExpression* FMaterialEditor::CreateNewMaterialExpression(UClass* NewExp
 	}
 
 	RegenerateCodeView();
-	UpdateSearch(false);
 
 	// Update the current preview material.
 	UpdatePreviewMaterial();
@@ -2774,7 +2759,6 @@ void FMaterialEditor::DeleteNodes(const TArray<UEdGraphNode*>& NodesToDelete)
 				SetPreviewMaterial( Material );
 			}
 			RegenerateCodeView();
-			UpdateSearch(false);
 		}
 		UpdatePreviewMaterial();
 		Material->MarkPackageDirty();
@@ -3060,134 +3044,128 @@ FMaterialRenderProxy* FMaterialEditor::GetExpressionPreview(UMaterialExpression*
 
 void FMaterialEditor::UndoGraphAction()
 {
-		int32 NumExpressions = Material->Expressions.Num();
-		GEditor->UndoTransaction();
+	int32 NumExpressions = Material->Expressions.Num();
+	GEditor->UndoTransaction();
 
-		if(NumExpressions != Material->Expressions.Num())
-		{
-			Material->BuildEditorParameterList();
-		}
-
-		UpdateSearch(false);
+	if(NumExpressions != Material->Expressions.Num())
+	{
+		Material->BuildEditorParameterList();
+	}
 
-		// Update the current preview material.
-		UpdatePreviewMaterial();
+	// Update the current preview material.
+	UpdatePreviewMaterial();
 
-		RefreshExpressionPreviews();
-		GraphEditor->NotifyGraphChanged();
-		SetMaterialDirty();
-	}
+	RefreshExpressionPreviews();
+	GraphEditor->NotifyGraphChanged();
+	SetMaterialDirty();
+}
 
 void FMaterialEditor::RedoGraphAction()
 {
-		// Clear selection, to avoid holding refs to nodes that go away
-		GraphEditor->ClearSelectionSet();
+	// Clear selection, to avoid holding refs to nodes that go away
+	GraphEditor->ClearSelectionSet();
 
-		int32 NumExpressions = Material->Expressions.Num();
-		GEditor->RedoTransaction();
+	int32 NumExpressions = Material->Expressions.Num();
+	GEditor->RedoTransaction();
 
-		if(NumExpressions != Material->Expressions.Num())
-		{
-			Material->BuildEditorParameterList();
-		}
+	if(NumExpressions != Material->Expressions.Num())
+	{
+		Material->BuildEditorParameterList();
+	}
 
-		// Update the current preview material.
-		UpdatePreviewMaterial();
+	// Update the current preview material.
+	UpdatePreviewMaterial();
 
-		RefreshExpressionPreviews();
-		GraphEditor->NotifyGraphChanged();
-		SetMaterialDirty();
-	}
+	RefreshExpressionPreviews();
+	GraphEditor->NotifyGraphChanged();
+	SetMaterialDirty();
+}
 
 void FMaterialEditor::PostUndo(bool bSuccess)
 {
-		GraphEditor->ClearSelectionSet();
-
-		Material->BuildEditorParameterList();
+	GraphEditor->ClearSelectionSet();
 
-		UpdateSearch(false);
+	Material->BuildEditorParameterList();
 
-		// Update the current preview material.
-		UpdatePreviewMaterial();
+	// Update the current preview material.
+	UpdatePreviewMaterial();
 
-		RefreshExpressionPreviews();
-		GraphEditor->NotifyGraphChanged();
-		SetMaterialDirty();
-	}
+	RefreshExpressionPreviews();
+	GraphEditor->NotifyGraphChanged();
+	SetMaterialDirty();
+}
 
 void FMaterialEditor::NotifyPreChange(UProperty* PropertyAboutToChange)
 {
-		check( !ScopedTransaction );
-		ScopedTransaction = new FScopedTransaction( NSLOCTEXT("UnrealEd", "MaterialEditorEditProperties", "Material Editor: Edit Properties") );
-		FlushRenderingCommands();
-	}
+	check( !ScopedTransaction );
+	ScopedTransaction = new FScopedTransaction( NSLOCTEXT("UnrealEd", "MaterialEditorEditProperties", "Material Editor: Edit Properties") );
+	FlushRenderingCommands();
+}
 
 void FMaterialEditor::NotifyPostChange( const FPropertyChangedEvent& PropertyChangedEvent, UProperty* PropertyThatChanged)
 {
-		check( ScopedTransaction );
+	check( ScopedTransaction );
 
-		if ( PropertyThatChanged )
+	if ( PropertyThatChanged )
+	{
+		const FName NameOfPropertyThatChanged( *PropertyThatChanged->GetName() );
+		if ( NameOfPropertyThatChanged == FName(TEXT("PreviewMesh")) ||
+			NameOfPropertyThatChanged == FName(TEXT("bUsedWithSkeletalMesh")) )
 		{
-			const FName NameOfPropertyThatChanged( *PropertyThatChanged->GetName() );
-			if ( NameOfPropertyThatChanged == FName(TEXT("PreviewMesh")) ||
-				NameOfPropertyThatChanged == FName(TEXT("bUsedWithSkeletalMesh")) )
+			// SetPreviewMesh will return false if the material has bUsedWithSkeletalMesh and
+			// a skeleton was requested, in which case revert to a sphere static mesh.
+			if ( !SetPreviewMesh( *Material->PreviewMesh.AssetLongPathname ) )
 			{
-				// SetPreviewMesh will return false if the material has bUsedWithSkeletalMesh and
-				// a skeleton was requested, in which case revert to a sphere static mesh.
-				if ( !SetPreviewMesh( *Material->PreviewMesh.AssetLongPathname ) )
-				{
-					SetPreviewMesh( GUnrealEd->GetThumbnailManager()->EditorSphere, NULL );
-				}
+				SetPreviewMesh( GUnrealEd->GetThumbnailManager()->EditorSphere, NULL );
 			}
+		}
 
-			FGraphPanelSelectionSet SelectedNodes = GraphEditor->GetSelectedNodes();
+		FGraphPanelSelectionSet SelectedNodes = GraphEditor->GetSelectedNodes();
 
-			for (FGraphPanelSelectionSet::TConstIterator NodeIt(SelectedNodes); NodeIt; ++NodeIt)
-			{
-				UMaterialGraphNode* SelectedNode = Cast<UMaterialGraphNode>(*NodeIt);
+		for (FGraphPanelSelectionSet::TConstIterator NodeIt(SelectedNodes); NodeIt; ++NodeIt)
+		{
+			UMaterialGraphNode* SelectedNode = Cast<UMaterialGraphNode>(*NodeIt);
 
-				if (SelectedNode && SelectedNode->MaterialExpression)
+			if (SelectedNode && SelectedNode->MaterialExpression)
+			{
+				if(NameOfPropertyThatChanged == FName(TEXT("ParameterName")))
 				{
-					if(NameOfPropertyThatChanged == FName(TEXT("ParameterName")))
-					{
-						Material->UpdateExpressionParameterName(SelectedNode->MaterialExpression);
-					}
-					else if (NameOfPropertyThatChanged == FName(TEXT("ParamNames")))
-					{
-						Material->UpdateExpressionDynamicParameterNames(SelectedNode->MaterialExpression);
-					}
-					else
-					{
-						Material->PropagateExpressionParameterChanges(SelectedNode->MaterialExpression);
-					}
+					Material->UpdateExpressionParameterName(SelectedNode->MaterialExpression);
+				}
+				else if (NameOfPropertyThatChanged == FName(TEXT("ParamNames")))
+				{
+					Material->UpdateExpressionDynamicParameterNames(SelectedNode->MaterialExpression);
+				}
+				else
+				{
+					Material->PropagateExpressionParameterChanges(SelectedNode->MaterialExpression);
 				}
 			}
 		}
+	}
 
-		// Prevent constant recompilation of materials while properties are being interacted with
-		if( PropertyChangedEvent.ChangeType != EPropertyChangeType::Interactive )
+	// Prevent constant recompilation of materials while properties are being interacted with
+	if( PropertyChangedEvent.ChangeType != EPropertyChangeType::Interactive )
+	{
+		// Also prevent recompilation when properties have no effect on material output
+		const FName PropertyName = PropertyChangedEvent.MemberProperty ? PropertyChangedEvent.MemberProperty->GetFName() : NAME_None;
+		if (PropertyName != GET_MEMBER_NAME_CHECKED(UMaterialExpressionComment, Text)
+		&& PropertyName != GET_MEMBER_NAME_CHECKED(UMaterialExpressionComment, CommentColor)
+		&& PropertyName != GET_MEMBER_NAME_CHECKED(UMaterialExpression, Desc))
 		{
-			// Also prevent recompilation when properties have no effect on material output
-			const FName PropertyName = PropertyChangedEvent.MemberProperty ? PropertyChangedEvent.MemberProperty->GetFName() : NAME_None;
-			if (PropertyName != GET_MEMBER_NAME_CHECKED(UMaterialExpressionComment, Text)
-			&& PropertyName != GET_MEMBER_NAME_CHECKED(UMaterialExpressionComment, CommentColor)
-			&& PropertyName != GET_MEMBER_NAME_CHECKED(UMaterialExpression, Desc))
-			{
-				// Update the current preview material.
-				UpdatePreviewMaterial();
-				RefreshExpressionPreviews();
-				RegenerateCodeView();
-			}
-
-			UpdateSearch(false);
+			// Update the current preview material.
+			UpdatePreviewMaterial();
+			RefreshExpressionPreviews();
+			RegenerateCodeView();
 		}
+	}
 
-		delete ScopedTransaction;
-		ScopedTransaction = NULL;
+	delete ScopedTransaction;
+	ScopedTransaction = NULL;
 
-		Material->MarkPackageDirty();
-		SetMaterialDirty();
-	}
+	Material->MarkPackageDirty();
+	SetMaterialDirty();
+}
 
 void FMaterialEditor::ToggleCollapsed(UMaterialExpression* MaterialExpression)
 {
@@ -3300,100 +3278,6 @@ FMatExpressionPreview* FMaterialEditor::GetExpressionPreview(UMaterialExpression
 	return NULL;
 }
 
-void FMaterialEditor::UpdateSearch( bool bQueryChanged )
-{
-	SearchResults.Empty();
-
-	if (SearchQuery.Len() == 0)
-	{
-		if (bQueryChanged)
-		{
-			// We just cleared the search
-			SelectedSearchResult = 0;
-		}
-	}
-	else
-	{
-		// Search expressions
-		for (int32 Index = 0; Index < Material->Expressions.Num(); Index++)
-		{
-			if (Material->Expressions[Index]->MatchesSearchQuery(*SearchQuery))
-			{
-				SearchResults.Add(Material->Expressions[Index]);
-			}
-		}
-
-		// Search comments
-		for (int32 Index = 0; Index < Material->EditorComments.Num(); Index++)
-		{
-			if (Material->EditorComments[Index]->MatchesSearchQuery(*SearchQuery))
-			{
-				SearchResults.Add(Material->EditorComments[Index]);
-			}
-		}
-
-		// Comparison function used to sort search results
-		struct FCompareUMaterialExpressionByPos
-		{
-			FORCEINLINE bool operator()( const UMaterialExpression& A, const UMaterialExpression& B ) const
-			{
-				// Divide into grid cells and step horizontally and then vertically.
-				int32 AGridX = A.MaterialExpressionEditorX / 100;
-				int32 AGridY = A.MaterialExpressionEditorY / 100;
-				int32 BGridX = B.MaterialExpressionEditorX / 100;
-				int32 BGridY = B.MaterialExpressionEditorY / 100;
-
-				if( AGridY < BGridY )
-				{
-					return true;
-				}
-				else
-					if( AGridY > BGridY )
-					{
-						return false;
-					}
-					else
-					{
-						return AGridX < BGridX;
-					}
-			}
-		};
-
-		SearchResults.Sort( FCompareUMaterialExpressionByPos() );
-
-		if( bQueryChanged )
-		{
-			// This is a new query rather than a material change, so navigate to first search result.
-			SelectedSearchResult = 0;
-			
-			if( SearchResults.Num() > 0 )
-			{
-				ShowSearchResult();
-			}
-		}
-		else
-		{
-			if( SelectedSearchResult < 0 || SelectedSearchResult >= SearchResults.Num() )
-			{
-				SelectedSearchResult = 0;
-			}
-		}
-	}
-}
-
-void FMaterialEditor::ShowSearchResult()
-{
-	if( SelectedSearchResult >= 0 && SelectedSearchResult < SearchResults.Num() )
-	{
-		UMaterialExpression* Expression = SearchResults[SelectedSearchResult];
-
-		// Select the selected search item
-		GraphEditor->ClearSelectionSet();
-		GraphEditor->SetNodeSelection(Expression->GraphNode, true);
-		GraphEditor->JumpToNode(Expression->GraphNode, false);
-	}
-}
-
 void FMaterialEditor::PreColorPickerCommit(FLinearColor LinearColor)
 {
 	// Begin a property edit transaction.
diff --git a/Engine/Source/Editor/MaterialEditor/Private/MaterialEditor.h b/Engine/Source/Editor/MaterialEditor/Private/MaterialEditor.h
--- a/Engine/Source/Editor/MaterialEditor/Private/MaterialEditor.h
+++ b/Engine/Source/Editor/MaterialEditor/Private/MaterialEditor.h
@@ -328,12 +328,12 @@ public:
 	/** Pushes the PreviewMesh assigned the the material instance to the thumbnail info */
 	static void UpdateThumbnailInfoPreviewMesh(UMaterialInterface* MatInterface);
 
-	/** Callback for when the canvas search next and previous are used */
-	void OnSearch(SSearchBox::SearchDirection Direction);
-
 	/** Sets the expression to be previewed. */
 	void SetPreviewExpression(UMaterialExpression* NewPreviewExpression);
 
+	/** Pan the view to center on a particular node */
+	void JumpToNode(const UEdGraphNode* Node);
+
 	// IMaterial Editor Interface
 	virtual UMaterialExpression* CreateNewMaterialExpression(UClass* NewExpressionClass, const FVector2D& NodePos, bool bAutoSelect, bool bAutoAssignResource) OVERRIDE;
 	virtual UMaterialExpressionComment* CreateNewMaterialExpressionComment(const FVector2D& NodePos) OVERRIDE;
@@ -346,8 +346,7 @@ public:
 	virtual void PasteNodesHere(const FVector2D& Location) OVERRIDE;
 	virtual int32 GetNumberOfSelectedNodes() const OVERRIDE;
 	virtual FMaterialRenderProxy* GetExpressionPreview(UMaterialExpression* InExpression) OVERRIDE;
-	virtual void UpdateSearch( bool bQueryChanged ) OVERRIDE;
-	
+
 public:
 	/** Set to true when modifications have been made to the material */
 	bool bMaterialDirty;
@@ -471,12 +470,6 @@ private:
 public:
 
 private:
-	/** Called when the graph search box changes text */
-	void OnGraphSearchChanged( const FText& InFilterText );
-
-	/** Called when the graph search text is committed */
-	void OnGraphSearchCommitted(const FText& NewTypeInValue, ETextCommit::Type CommitInfo);
-	
 	/**
 	 * Load editor settings from disk (docking state, window pos/size, option state, etc).
 	 */
@@ -540,8 +533,10 @@ private:
 	void OnSelectUpsteamNodes();
 	/** Command to force a refresh of all previews (triggered by space bar) */
 	void OnForceRefreshPreviews();
-	/* Create comment node on graph */
+	/** Create comment node on graph */
 	void OnCreateComment();
+	/** Bring up the search tab */
+	void OnFindInMaterial();
 
 	/** Callback from the Asset Registry when an asset is renamed. */
 	void RenameAssetFromRegistry(const FAssetData& InAddedAssetData, const FString& InNewName);
@@ -582,11 +577,6 @@ private:
 	 */
 	FMatExpressionPreview* GetExpressionPreview(UMaterialExpression* MaterialExpression, bool& bNewlyCreated);
 
-	/**
-	  * Moves the shows the selected result in the viewport
-	  */
-	void ShowSearchResult();
-
 	/** Pointer to the object that the current color picker is working on. Can be NULL and stale. */
 	TWeakObjectPtr<UObject> ColorPickerObject;
 
@@ -634,6 +624,7 @@ private:
 	TSharedRef<SDockTab> SpawnTab_HLSLCode(const FSpawnTabArgs& Args);
 	TSharedRef<SDockTab> SpawnTab_Palette(const FSpawnTabArgs& Args);
 	TSharedRef<SDockTab> SpawnTab_Stats(const FSpawnTabArgs& Args);
+	TSharedRef<SDockTab> SpawnTab_Find(const FSpawnTabArgs& Args);
 private:
 	/** List of open tool panels; used to ensure only one exists at any one time */
 	TMap< FName, TWeakPtr<class SDockableTab> > SpawnedToolPanels;
@@ -662,14 +653,8 @@ private:
 	TSharedPtr<class SWidget> Stats;
 	TSharedPtr<class IMessageLogListing> StatsListing;
 
-	/** Current search query */
-	FString SearchQuery;
-
-	/** Array of expressions matching the current search terms. */
-	TArray<UMaterialExpression*> SearchResults;
-
-	/** Index in the above array of the currently selected search result */
-	int32 SelectedSearchResult;
+	/** Find results log as well as the search filter */
+	TSharedPtr<class SFindInMaterial> FindResults;
 
 	/** The current transaction. */
 	FScopedTransaction* ScopedTransaction;
@@ -705,4 +690,5 @@ private:
 	static const FName HLSLCodeTabId;	
 	static const FName PaletteTabId;
 	static const FName StatsTabId;
+	static const FName FindTabId;
 };
diff --git a/Engine/Source/Editor/MaterialEditor/Private/MaterialEditorActions.cpp b/Engine/Source/Editor/MaterialEditor/Private/MaterialEditorActions.cpp
--- a/Engine/Source/Editor/MaterialEditor/Private/MaterialEditorActions.cpp
+++ b/Engine/Source/Editor/MaterialEditor/Private/MaterialEditorActions.cpp
@@ -47,6 +47,8 @@ void FMaterialEditorCommands::RegisterCommands()
 	UI_COMMAND( AddToFavorites, "Add To Favorites", "Adds this expression to your favorites.", EUserInterfaceActionType::Button, FInputGesture() );
 	UI_COMMAND( BreakLink, "Break Link", "Deletes this link.", EUserInterfaceActionType::Button, FInputGesture() );
 	UI_COMMAND( ForceRefreshPreviews, "Force Refresh Previews", "Forces a refresh of all previews", EUserInterfaceActionType::Button, FInputGesture(EKeys::SpaceBar) );
+
+	UI_COMMAND( FindInMaterial, "Search", "Finds expressions and comments in the current Material", EUserInterfaceActionType::Button, FInputGesture(EModifierKey::Control, EKeys::F));
 }
 
 //////////////////////////////////////////////////////////////////////////
diff --git a/Engine/Source/Editor/MaterialEditor/Public/FindInMaterial.h b/Engine/Source/Editor/MaterialEditor/Public/FindInMaterial.h
new file mode 100644
--- /dev/null
+++ b/Engine/Source/Editor/MaterialEditor/Public/FindInMaterial.h
@@ -0,0 +1,123 @@
+// Copyright 1998-2014 Epic Games, Inc. All Rights Reserved.
+
+#pragma once
+
+/** Item that matched the search results */
+class FFindInMaterialResult
+{
+public:
+	/** Create a root (or only text) result */
+	FFindInMaterialResult(const FString& InValue);
+
+	/** Create a listing for a search result*/
+	FFindInMaterialResult(const FString& InValue, TSharedPtr<FFindInMaterialResult>& InParent, UClass* InClass, int InDuplicationIndex);
+
+	/** Create a listing for a pin result */
+	FFindInMaterialResult(const FString& InValue, TSharedPtr<FFindInMaterialResult>& InParent, UEdGraphPin* InPin);
+
+	/** Create a listing for a node result */
+	FFindInMaterialResult(const FString& InValue, TSharedPtr<FFindInMaterialResult>& InParent, UEdGraphNode* InNode);
+
+	/** Called when user clicks on the search item */
+	FReply OnClick(TWeakPtr<class FMaterialEditor> MaterialEditor);
+
+	/* Get Category for this search result */
+	FString GetCategory() const;
+
+	/** Create an icon to represent the result */
+	TSharedRef<SWidget>	CreateIcon() const;
+
+	/** Gets the comment on this node if any */
+	FString GetCommentText() const;
+
+	/** Gets the value of the pin if any */
+	FString GetValueText() const;
+
+	/** Any children listed under this category */
+	TArray< TSharedPtr<FFindInMaterialResult> > Children;
+
+	/** Search result Parent */
+	TWeakPtr<FFindInMaterialResult> Parent;
+
+	/*The meta string that was stored in the asset registry for this item */
+	FString Value;
+
+	/*The graph may have multiple instances of whatever we are looking for, this tells us which instance # we refer to*/
+	int	DuplicationIndex;
+
+	/*The class this item refers to */
+	UClass* Class;
+
+	/** The pin that this search result refers to */
+	TWeakObjectPtr<UEdGraphPin> Pin;
+
+	/** The graph node that this search result refers to (if not by asset registry or UK2Node) */
+	TWeakObjectPtr<UEdGraphNode> GraphNode;
+
+	/** Display text for comment information */
+	FString CommentText;
+};
+
+
+/** Widget for searching for items that are part of a UEdGraph */
+class SFindInMaterial : public SCompoundWidget
+{
+public:
+	SLATE_BEGIN_ARGS(SFindInMaterial){}
+	SLATE_END_ARGS()
+
+	void Construct(const FArguments& InArgs, TSharedPtr<class FMaterialEditor> InMaterialEditor);
+
+	/** Focuses this widget's search box */
+	void FocusForUse();
+
+protected:
+	typedef TSharedPtr<FFindInMaterialResult> FSearchResult;
+	typedef STreeView<FSearchResult> STreeViewType;
+
+	/** Called when user changes the text they are searching for */
+	void OnSearchTextChanged(const FText& Text);
+
+	/** Called when user changes commits text to the search box */
+	void OnSearchTextCommitted(const FText& Text, ETextCommit::Type CommitType);
+
+	/** Get the children of a row */
+	void OnGetChildren(FSearchResult InItem, TArray<FSearchResult>& OutChildren);
+
+	/** Called when user clicks on a new result */
+	void OnTreeSelectionChanged(FSearchResult Item, ESelectInfo::Type SelectInfo);
+
+	/** Called when a new row is being generated */
+	TSharedRef<ITableRow> OnGenerateRow(FSearchResult InItem, const TSharedRef<STableViewBase>& OwnerTable);
+
+	/** Begins the search based on the SearchValue */
+	void InitiateSearch();
+
+	/** Find any results that contain all of the tokens */
+	void MatchTokens(const TArray<FString>& Tokens);
+
+	/** Determines if a string matches the search tokens */
+	static bool StringMatchesSearchTokens(const TArray<FString>& Tokens, const FString& ComparisonString);
+
+protected:
+	/** Pointer back to the Material editor that owns us */
+	TWeakPtr<class FMaterialEditor> MaterialEditorPtr;
+
+	/** The tree view displays the results */
+	TSharedPtr<STreeViewType> TreeView;
+
+	/** The search text box */
+	TSharedPtr<class SSearchBox> SearchTextField;
+
+	/** This buffer stores the currently displayed results */
+	TArray<FSearchResult> ItemsFound;
+
+	/** we need to keep a handle on the root result, because it won't show up in the tree */
+	FSearchResult RootSearchResult;
+
+	/** The string to highlight in the results */
+	FText HighlightText;
+
+	/** The string to search for */
+	FString	SearchValue;
+};
\ No newline at end of file
diff --git a/Engine/Source/Editor/MaterialEditor/Public/MaterialEditorActions.h b/Engine/Source/Editor/MaterialEditor/Public/MaterialEditorActions.h
--- a/Engine/Source/Editor/MaterialEditor/Public/MaterialEditorActions.h
+++ b/Engine/Source/Editor/MaterialEditor/Public/MaterialEditorActions.h
@@ -143,6 +143,9 @@ public:
 	/** Forces a refresh of all previews */
 	TSharedPtr< FUICommandInfo > ForceRefreshPreviews;
 
+	/** Finds expressions in current material */
+	TSharedPtr< FUICommandInfo > FindInMaterial;
+
 	/**
 	 * Initialize commands
 	 */
diff --git a/Engine/Source/Editor/UnrealEd/Classes/MaterialGraph/MaterialGraphNode.h b/Engine/Source/Editor/UnrealEd/Classes/MaterialGraph/MaterialGraphNode.h
--- a/Engine/Source/Editor/UnrealEd/Classes/MaterialGraph/MaterialGraphNode.h
+++ b/Engine/Source/Editor/UnrealEd/Classes/MaterialGraph/MaterialGraphNode.h
@@ -19,9 +19,6 @@ class UMaterialGraphNode : public UMaterialGraphNode_Base
 	/** Set to true if this expression causes an error in the material */
 	bool bIsErrorExpression;
 
-	/** Set to true if this expression is the current search result */
-	bool bIsCurrentSearchResult;
-
 	/** Set to true if this expression is currently being previewed */
 	bool bIsPreviewExpression;
 
diff --git a/Engine/Source/Editor/UnrealEd/Private/MaterialGraphNode.cpp b/Engine/Source/Editor/UnrealEd/Private/MaterialGraphNode.cpp
--- a/Engine/Source/Editor/UnrealEd/Private/MaterialGraphNode.cpp
+++ b/Engine/Source/Editor/UnrealEd/Private/MaterialGraphNode.cpp
@@ -18,7 +18,6 @@ UMaterialGraphNode::UMaterialGraphNode(const class FPostConstructInitializePrope
 	: Super(PCIP)
 	, bPreviewNeedsUpdate(false)
 	, bIsErrorExpression(false)
-	, bIsCurrentSearchResult(false)
 	, bIsPreviewExpression(false)
 {
 }
@@ -131,7 +130,7 @@ FString UMaterialGraphNode::GetNodeTitle(ENodeTitleType::Type TitleType) const
 	}
 	else if (TitleType == ENodeTitleType::ListView)
 	{
-		return MaterialExpression->GetClass()->GetDescription();
+		return MaterialExpression->GetDescription();
 	}
 	else
 	{
diff --git a/Engine/Source/Runtime/Engine/Classes/Materials/MaterialExpression.h b/Engine/Source/Runtime/Engine/Classes/Materials/MaterialExpression.h
--- a/Engine/Source/Runtime/Engine/Classes/Materials/MaterialExpression.h
+++ b/Engine/Source/Runtime/Engine/Classes/Materials/MaterialExpression.h
@@ -228,6 +228,8 @@ class UMaterialExpression : public UObject
 	 */
 	virtual void GetCaption(TArray<FString>& OutCaptions) const;
 #if WITH_EDITOR
+	/** Get a single line description of the material expression (used for lists) */
+	virtual FString GetDescription() const;
 	/** Get a tooltip for the specified connector. */
 	virtual void GetConnectorToolTip(int32 InputIndex, int32 OutputIndex, TArray<FString>& OutToolTip);
 
diff --git a/Engine/Source/Runtime/Engine/Classes/Materials/MaterialExpressionConstant.h b/Engine/Source/Runtime/Engine/Classes/Materials/MaterialExpressionConstant.h
--- a/Engine/Source/Runtime/Engine/Classes/Materials/MaterialExpressionConstant.h
+++ b/Engine/Source/Runtime/Engine/Classes/Materials/MaterialExpressionConstant.h
@@ -16,6 +16,10 @@ class UMaterialExpressionConstant : public UMaterialExpression
 	// Begin UMaterialExpression Interface
 	virtual int32 Compile(class FMaterialCompiler* Compiler, int32 OutputIndex, int32 MultiplexIndex) OVERRIDE;
 	virtual void GetCaption(TArray<FString>& OutCaptions) const OVERRIDE;
+#if WITH_EDITOR
+	virtual FString GetDescription() const OVERRIDE;
+#endif // WITH_EDITOR
+
 	// End UMaterialExpression Interface
 };
 
diff --git a/Engine/Source/Runtime/Engine/Classes/Materials/MaterialExpressionConstant2Vector.h b/Engine/Source/Runtime/Engine/Classes/Materials/MaterialExpressionConstant2Vector.h
--- a/Engine/Source/Runtime/Engine/Classes/Materials/MaterialExpressionConstant2Vector.h
+++ b/Engine/Source/Runtime/Engine/Classes/Materials/MaterialExpressionConstant2Vector.h
@@ -20,6 +20,7 @@ class UMaterialExpressionConstant2Vector : public UMaterialExpression
 	virtual int32 Compile(class FMaterialCompiler* Compiler, int32 OutputIndex, int32 MultiplexIndex) OVERRIDE;
 	virtual void GetCaption(TArray<FString>& OutCaptions) const OVERRIDE;
 #if WITH_EDITOR
+	virtual FString GetDescription() const OVERRIDE;
 	virtual uint32 GetOutputType(int32 OutputIndex) OVERRIDE {return MCT_Float2;}
 #endif // WITH_EDITOR
 	// End UMaterialExpression Interface
diff --git a/Engine/Source/Runtime/Engine/Classes/Materials/MaterialExpressionConstant3Vector.h b/Engine/Source/Runtime/Engine/Classes/Materials/MaterialExpressionConstant3Vector.h
--- a/Engine/Source/Runtime/Engine/Classes/Materials/MaterialExpressionConstant3Vector.h
+++ b/Engine/Source/Runtime/Engine/Classes/Materials/MaterialExpressionConstant3Vector.h
@@ -29,6 +29,7 @@ class UMaterialExpressionConstant3Vector : public UMaterialExpression
 	virtual int32 Compile(class FMaterialCompiler* Compiler, int32 OutputIndex, int32 MultiplexIndex) OVERRIDE;
 	virtual void GetCaption(TArray<FString>& OutCaptions) const OVERRIDE;
 #if WITH_EDITOR
+	virtual FString GetDescription() const OVERRIDE;
 	virtual uint32 GetOutputType(int32 OutputIndex) OVERRIDE {return MCT_Float3;}
 #endif // WITH_EDITOR
 	// End UMaterialExpression Interface
diff --git a/Engine/Source/Runtime/Engine/Classes/Materials/MaterialExpressionConstant4Vector.h b/Engine/Source/Runtime/Engine/Classes/Materials/MaterialExpressionConstant4Vector.h
--- a/Engine/Source/Runtime/Engine/Classes/Materials/MaterialExpressionConstant4Vector.h
+++ b/Engine/Source/Runtime/Engine/Classes/Materials/MaterialExpressionConstant4Vector.h
@@ -32,6 +32,7 @@ class UMaterialExpressionConstant4Vector : public UMaterialExpression
 	virtual int32 Compile(class FMaterialCompiler* Compiler, int32 OutputIndex, int32 MultiplexIndex) OVERRIDE;
 	virtual void GetCaption(TArray<FString>& OutCaptions) const OVERRIDE;
 #if WITH_EDITOR
+	virtual FString GetDescription() const OVERRIDE;
 	virtual uint32 GetOutputType(int32 OutputIndex) OVERRIDE {return MCT_Float4;}
 #endif // WITH_EDITOR
 	// End UMaterialExpression Interface
diff --git a/Engine/Source/Runtime/Engine/Classes/Materials/MaterialExpressionMaterialFunctionCall.h b/Engine/Source/Runtime/Engine/Classes/Materials/MaterialExpressionMaterialFunctionCall.h
--- a/Engine/Source/Runtime/Engine/Classes/Materials/MaterialExpressionMaterialFunctionCall.h
+++ b/Engine/Source/Runtime/Engine/Classes/Materials/MaterialExpressionMaterialFunctionCall.h
@@ -99,6 +99,7 @@ class UMaterialExpressionMaterialFunctionCall : public UMaterialExpression
 	virtual FString GetInputName(int32 InputIndex) const OVERRIDE;
 	virtual bool IsInputConnectionRequired(int32 InputIndex) const OVERRIDE;
 #if WITH_EDITOR
+	virtual FString GetDescription() const OVERRIDE;
 	virtual void GetConnectorToolTip(int32 InputIndex, int32 OutputIndex, TArray<FString>& OutToolTip) OVERRIDE;
 	virtual void GetExpressionToolTip(TArray<FString>& OutToolTip) OVERRIDE;
 #endif
diff --git a/Engine/Source/Runtime/Engine/Classes/Materials/MaterialExpressionTextureBase.h b/Engine/Source/Runtime/Engine/Classes/Materials/MaterialExpressionTextureBase.h
--- a/Engine/Source/Runtime/Engine/Classes/Materials/MaterialExpressionTextureBase.h
+++ b/Engine/Source/Runtime/Engine/Classes/Materials/MaterialExpressionTextureBase.h
@@ -29,6 +29,12 @@ class UMaterialExpressionTextureBase : public UMaterialExpression
 #endif // WITH_EDITOR
 	// End UObject Interface
 
+	// Begin UMaterialExpression Interface
+#if WITH_EDITOR
+	virtual FString GetDescription() const OVERRIDE;
+#endif
+	// End UMaterialExpression Interface
+
 	/** 
 	 * Callback to get any texture reference this expression emits.
 	 * This is used to link the compiled uniform expressions with their default texture values. 
diff --git a/Engine/Source/Runtime/Engine/Private/Materials/MaterialExpressions.cpp b/Engine/Source/Runtime/Engine/Private/Materials/MaterialExpressions.cpp
--- a/Engine/Source/Runtime/Engine/Private/Materials/MaterialExpressions.cpp
+++ b/Engine/Source/Runtime/Engine/Private/Materials/MaterialExpressions.cpp
@@ -640,6 +640,29 @@ void UMaterialExpression::GetCaption(TArray<FString>& OutCaptions) const
 	OutCaptions.Add(TEXT("Expression"));
 }
 #if WITH_EDITOR
+FString UMaterialExpression::GetDescription() const
+{
+	// Combined captions sufficient for most expressions
+	TArray<FString> Captions;
+	GetCaption(Captions);
+
+	if (Captions.Num() > 1)
+	{
+		FString Result = Captions[0];
+		for (int32 Index = 1; Index < Captions.Num(); ++Index)
+		{
+			Result += TEXT(" ");
+			Result += Captions[Index];
+		}
+
+		return Result;
+	}
+	else
+	{
+		return Captions[0];
+	}
+}
+
 void UMaterialExpression::GetConnectorToolTip(int32 InputIndex, int32 OutputIndex, TArray<FString>& OutToolTip)
 {
 	if (InputIndex != INDEX_NONE)
@@ -874,6 +897,16 @@ void UMaterialExpressionTextureBase::PostEditChangeProperty(FPropertyChangedEven
 		}
 	}
 }
+
+FString UMaterialExpressionTextureBase::GetDescription() const
+{
+	FString Result = Super::GetDescription();
+	Result += TEXT(" (");
+	Result += Texture ? Texture->GetName() : TEXT("None");
+	Result += TEXT(")");
+
+	return Result;
+}
 #endif // WITH_EDITOR
 
 void UMaterialExpressionTextureBase::AutoSetSampleType()
@@ -1971,6 +2004,17 @@ void UMaterialExpressionConstant::GetCaption(TArray<FString>& OutCaptions) const
 	OutCaptions.Add(FString::Printf( TEXT("%.4g"), R ));
 }
 
+#if WITH_EDITOR
+FString UMaterialExpressionConstant::GetDescription() const
+{
+	FString Result = FString(*GetClass()->GetName()).Mid(FCString::Strlen(TEXT("MaterialExpression")));
+	Result += TEXT(" (");
+	Result += Super::GetDescription();
+	Result += TEXT(")");
+	return Result;
+}
+#endif // WITH_EDITOR
+
 UMaterialExpressionConstant2Vector::UMaterialExpressionConstant2Vector(const class FPostConstructInitializeProperties& PCIP)
 	: Super(PCIP)
 {
@@ -2003,6 +2047,17 @@ void UMaterialExpressionConstant2Vector::GetCaption(TArray<FString>& OutCaptions
 	OutCaptions.Add(FString::Printf( TEXT("%.3g,%.3g"), R, G ));
 }
 
+#if WITH_EDITOR
+FString UMaterialExpressionConstant2Vector::GetDescription() const
+{
+	FString Result = FString(*GetClass()->GetName()).Mid(FCString::Strlen(TEXT("MaterialExpression")));
+	Result += TEXT(" (");
+	Result += Super::GetDescription();
+	Result += TEXT(")");
+	return Result;
+}
+#endif // WITH_EDITOR
+
 UMaterialExpressionConstant3Vector::UMaterialExpressionConstant3Vector(const class FPostConstructInitializeProperties& PCIP)
 	: Super(PCIP)
 {
@@ -2047,6 +2102,17 @@ void UMaterialExpressionConstant3Vector::GetCaption(TArray<FString>& OutCaptions
 	OutCaptions.Add(FString::Printf( TEXT("%.3g,%.3g,%.3g"), Constant.R, Constant.G, Constant.B ));
 }
 
+#if WITH_EDITOR
+FString UMaterialExpressionConstant3Vector::GetDescription() const
+{
+	FString Result = FString(*GetClass()->GetName()).Mid(FCString::Strlen(TEXT("MaterialExpression")));
+	Result += TEXT(" (");
+	Result += Super::GetDescription();
+	Result += TEXT(")");
+	return Result;
+}
+#endif // WITH_EDITOR
+
 UMaterialExpressionConstant4Vector::UMaterialExpressionConstant4Vector(const class FPostConstructInitializeProperties& PCIP)
 	: Super(PCIP)
 {
@@ -2093,6 +2159,17 @@ void UMaterialExpressionConstant4Vector::GetCaption(TArray<FString>& OutCaptions
 	OutCaptions.Add(FString::Printf( TEXT("%.2g,%.2g,%.2g,%.2g"), Constant.R, Constant.G, Constant.B, Constant.A ));
 }
 
+#if WITH_EDITOR
+FString UMaterialExpressionConstant4Vector::GetDescription() const
+{
+	FString Result = FString(*GetClass()->GetName()).Mid(FCString::Strlen(TEXT("MaterialExpression")));
+	Result += TEXT(" (");
+	Result += Super::GetDescription();
+	Result += TEXT(")");
+	return Result;
+}
+#endif // WITH_EDITOR
+
 UMaterialExpressionClamp::UMaterialExpressionClamp(const class FPostConstructInitializeProperties& PCIP)
 	: Super(PCIP)
 {
@@ -6454,6 +6531,15 @@ static FString GetInputDefaultValueString(EFunctionInputType InputType, const FV
 }
 
 #if WITH_EDITOR
+FString UMaterialExpressionMaterialFunctionCall::GetDescription() const
+{
+	FString Result = FString(*GetClass()->GetName()).Mid(FCString::Strlen(TEXT("MaterialExpression")));
+	Result += TEXT(" (");
+	Result += Super::GetDescription();
+	Result += TEXT(")");
+	return Result;
+}
+
 void UMaterialExpressionMaterialFunctionCall::GetConnectorToolTip(int32 InputIndex, int32 OutputIndex, TArray<FString>& OutToolTip) 
 {
 	if (MaterialFunction)
