Material overhead stats check/ensure fix.

Fixes the underlying issue with the material overhead stats.

Usage flags for the material were set from outside the material editor. I needed to update the overhead stats materials to match.

diff --git a/Engine/Source/Editor/MaterialEditor/Private/MaterialEditor.cpp b/Engine/Source/Editor/MaterialEditor/Private/MaterialEditor.cpp
--- a/Engine/Source/Editor/MaterialEditor/Private/MaterialEditor.cpp
+++ b/Engine/Source/Editor/MaterialEditor/Private/MaterialEditor.cpp
@@ -2253,6 +2253,7 @@ void FMaterialEditor::OnMaterialUsageFlagsChanged(UMaterial* MaterialThatChanged
 	{
 		bool bNeedsRecompile = false;
 		Material->SetMaterialUsage(bNeedsRecompile, Flag, MaterialThatChanged->GetUsageByFlag(Flag));
+		MaterialDevelopmentOverheadStats.Update(Material);
 	}
 }
 
