Hiding the material dev overhead stats behind the r.ShowMaterialDevOverheadStats CVar. Off by default.

diff --git a/Engine/Source/Editor/MaterialEditor/Private/MaterialEditor.cpp b/Engine/Source/Editor/MaterialEditor/Private/MaterialEditor.cpp
--- a/Engine/Source/Editor/MaterialEditor/Private/MaterialEditor.cpp
+++ b/Engine/Source/Editor/MaterialEditor/Private/MaterialEditor.cpp
@@ -46,6 +46,13 @@ const FName FMaterialEditor::PaletteTabId( TEXT( "MaterialEditor_Palette" ) );
 const FName FMaterialEditor::StatsTabId( TEXT( "MaterialEditor_Stats" ) );
 const FName FMaterialEditor::FindTabId( TEXT( "MaterialEditor_Find" ) );
 
+static int32 GShowDevOverheadStats = 0;
+static FAutoConsoleVariableRef CVarShowDevOverheadStats(
+	TEXT("r.ShowMaterialDevOverheadStats"),
+	GShowDevOverheadStats,
+	TEXT("Show the overhead introduced into materials by the editor.")
+	);
+
 ///////////////////////////
 // FMatExpressionPreview //
 ///////////////////////////
@@ -3751,7 +3758,7 @@ void FMaterialDevelopmentOverheadStats::Init(UMaterial* InMaterial)
 
 void FMaterialDevelopmentOverheadStats::Update(UMaterial* InMaterial)
 {
-	if (NULL == InMaterial)
+	if (NULL == InMaterial || 0 == GShowDevOverheadStats)
 	{
 		//Don't bother clearing the materials as we'll either be reusing them again shortly or exiting the editor.
 		bDisabled = true;
@@ -3795,7 +3802,7 @@ void FMaterialDevelopmentOverheadStats::AddReferencedObjects(FReferenceCollector
 
 bool FMaterialDevelopmentOverheadStats::GetOverheadCounts(TArray<int32>& OverheadCounts, ERHIFeatureLevel::Type FeatureLevel)
 {		
-	if (bDisabled)
+	if (bDisabled || 0 == GShowDevOverheadStats)
 	{
 		return false;
 	}
