#summary Permanently flip and space out coordinates of Material Expression nodes

#ttp 316874 - EDITOR: ME2: Space out the nodes in the new editor

#proj UE4.Engine

#change Added version number for permanent flip and added static function to material that can flip coords of materials and functions. Also scale the position of all expressions by roughly the size difference of the new nodes so that they are hopefully less tightly packed together.

diff --git a/Engine/Source/Editor/MaterialEditor/Private/MaterialEditor.cpp b/Engine/Source/Editor/MaterialEditor/Private/MaterialEditor.cpp
--- a/Engine/Source/Editor/MaterialEditor/Private/MaterialEditor.cpp
+++ b/Engine/Source/Editor/MaterialEditor/Private/MaterialEditor.cpp
@@ -432,8 +432,6 @@ void FMaterialEditor::InitMaterialEditor( const EToolkitMode::Type Mode, const T
 		}
 	}
 
-	// TODO: Remove this once Expressions are saved as flipped
-	FlipExpressionPositions(Material->Expressions, Material->EditorComments, Material);
 	Material->MaterialGraph->RebuildGraph();
 	RecenterEditor();
 	ForceRefreshExpressionPreviews();
@@ -1228,9 +1226,6 @@ void FMaterialEditor::UpdateOriginalMaterial()
 		MaterialFunction->ParentFunction->ThumbnailInfo = OriginalThumbnailInfo;
 		MaterialFunction->ThumbnailInfo = ThumbnailInfo;
 
-		// TODO: Remove this once Expressions are saved as flipped
-		FlipExpressionPositions(MaterialFunction->ParentFunction->FunctionExpressions, MaterialFunction->ParentFunction->FunctionEditorComments);
-
 		// Restore RF_Standalone on the original material function, as it had been removed from the preview material so that it could be GC'd.
 		MaterialFunction->ParentFunction->SetFlags( RF_Standalone );
 
@@ -1359,9 +1354,6 @@ void FMaterialEditor::UpdateOriginalMaterial()
 			OriginalMaterial->ThumbnailInfo = OriginalThumbnailInfo;
 			Material->ThumbnailInfo = ThumbnailInfo;
 
-			// TODO: Remove this once Expressions are saved as flipped
-			FlipExpressionPositions(OriginalMaterial->Expressions, OriginalMaterial->EditorComments, OriginalMaterial);
-
 			// Change the original material object to the new original material
 			OriginalMaterialObject = OriginalMaterial;
 
@@ -3732,24 +3724,6 @@ FReply FMaterialEditor::OnSpawnGraphNodeByShortcut(FInputGesture InGesture, cons
 	return FReply::Unhandled();
 }
 
-void FMaterialEditor::FlipExpressionPositions(const TArray<UMaterialExpression*>& Expressions, const TArray<UMaterialExpressionComment*>& Comments, UMaterial* InMaterial)
-{
-	if (InMaterial)
-	{
-		InMaterial->EditorX = -InMaterial->EditorX;
-	}
-	for (int32 ExpressionIndex = 0; ExpressionIndex < Expressions.Num(); ExpressionIndex++)
-	{
-		UMaterialExpression* Expression = Expressions[ExpressionIndex];
-		Expression->MaterialExpressionEditorX = -Expression->MaterialExpressionEditorX;
-	}
-	for (int32 ExpressionIndex = 0; ExpressionIndex < Comments.Num(); ExpressionIndex++)
-	{
-		UMaterialExpressionComment* Comment = Comments[ExpressionIndex];
-		Comment->MaterialExpressionEditorX = -Comment->MaterialExpressionEditorX - Comment->SizeX;
-	}
-}
-
 void FMaterialDevelopmentOverheadStats::Init(UMaterial* InMaterial)
 {
 	bDisabled = false;
diff --git a/Engine/Source/Editor/MaterialEditor/Private/MaterialEditor.h b/Engine/Source/Editor/MaterialEditor/Private/MaterialEditor.h
--- a/Engine/Source/Editor/MaterialEditor/Private/MaterialEditor.h
+++ b/Engine/Source/Editor/MaterialEditor/Private/MaterialEditor.h
@@ -608,15 +608,6 @@ private:
 	/** Adds the selected expression to the favorites list. */
 	void AddSelectedExpressionToFavorites();
 
-	/**
-	 * Flip the X coordinates of a material's expressions when loading and saving -
-	 * Remove this once new material editor is active and materials are saved inverted.
-	 *
-	 * @param	Expressions	Array of material expressions
-	 * @param	Comments	Array of material expression comments
-	 */
-	void FlipExpressionPositions(const TArray<UMaterialExpression*>& Expressions, const TArray<UMaterialExpressionComment*>& Comments, UMaterial* Material = NULL);
-
 private:
 	TSharedRef<SDockTab> SpawnTab_Preview(const FSpawnTabArgs& Args);
 	TSharedRef<SDockTab> SpawnTab_GraphCanvas(const FSpawnTabArgs& Args);
diff --git a/Engine/Source/Runtime/Core/Public/UObject/ObjectVersion.h b/Engine/Source/Runtime/Core/Public/UObject/ObjectVersion.h
--- a/Engine/Source/Runtime/Core/Public/UObject/ObjectVersion.h
+++ b/Engine/Source/Runtime/Core/Public/UObject/ObjectVersion.h
@@ -518,6 +518,9 @@ enum EUnrealEngineObjectUE4Version
 	VER_UE4_UNDO_BREAK_MATERIALATTRIBUTES_CHANGE,
 	// Now Default custom profile name isn't NONE anymore due to copy/paste not working properly with it
 	VER_UE4_ADD_CUSTOMPROFILENAME_CHANGE, 
+	// Permanently flip and scale material expression coordinates
+	VER_UE4_FLIP_MATERIAL_COORDS,
+
 	// -----<new versions can be added before this line>-------------------------------------------------
 	// - this needs to be the last line (see note below)
 	VER_UE4_AUTOMATIC_VERSION_PLUS_ONE,
diff --git a/Engine/Source/Runtime/Engine/Classes/Materials/Material.h b/Engine/Source/Runtime/Engine/Classes/Materials/Material.h
--- a/Engine/Source/Runtime/Engine/Classes/Materials/Material.h
+++ b/Engine/Source/Runtime/Engine/Classes/Materials/Material.h
@@ -1145,6 +1145,17 @@ public:
 	// For all materials, UMaterial::CacheResourceShadersForRendering
 	static void AllMaterialsCacheResourceShadersForRendering();
 
+#if WITH_EDITORONLY_DATA
+	/**
+	 * Flip the X coordinates of a material's expressions and space them out more
+	 *
+	 * @param	Expressions	Array of material expressions
+	 * @param	Comments	Array of material expression comments
+	 * @param	Material	The Material to flip its home coords (optional)
+	 */
+	static void FlipExpressionPositions(const TArray<UMaterialExpression*>& Expressions, const TArray<UMaterialExpressionComment*>& Comments, UMaterial* Material = NULL);
+#endif //WITH_EDITORONLY_DATA
+
 private:
 	static FMaterialCompilationFinished MaterialCompilationFinishedEvent;
 
diff --git a/Engine/Source/Runtime/Engine/Private/Materials/Material.cpp b/Engine/Source/Runtime/Engine/Private/Materials/Material.cpp
--- a/Engine/Source/Runtime/Engine/Private/Materials/Material.cpp
+++ b/Engine/Source/Runtime/Engine/Private/Materials/Material.cpp
@@ -1998,6 +1998,13 @@ void UMaterial::PostLoad()
 			DefaultMaterialInstances[i]->GameThread_UpdateDistanceFieldPenumbraScale(GetDistanceFieldPenumbraScale());
 		}
 	}
+
+#if WITH_EDITORONLY_DATA
+	if (GetLinkerUE4Version() < VER_UE4_FLIP_MATERIAL_COORDS)
+	{
+		FlipExpressionPositions(Expressions, EditorComments, this);
+	}
+#endif //WITH_EDITORONLY_DATA
 }
 
 void UMaterial::BeginCacheForCookedPlatformData( const ITargetPlatform *TargetPlatform )
@@ -3500,4 +3507,31 @@ bool UMaterial::IsPropertyActive(EMaterialProperty InProperty)const
 		break;
 	}
 	return Active;
-}
\ No newline at end of file
+}
+
+#if WITH_EDITORONLY_DATA
+void UMaterial::FlipExpressionPositions(const TArray<UMaterialExpression*>& Expressions, const TArray<UMaterialExpressionComment*>& Comments, UMaterial* InMaterial)
+{
+	// Rough estimate of average increase in node size for the new editor
+	const float PosScaling = 1.25f;
+
+	if (InMaterial)
+	{
+		InMaterial->EditorX = -InMaterial->EditorX;
+	}
+	for (int32 ExpressionIndex = 0; ExpressionIndex < Expressions.Num(); ExpressionIndex++)
+	{
+		UMaterialExpression* Expression = Expressions[ExpressionIndex];
+		Expression->MaterialExpressionEditorX = -Expression->MaterialExpressionEditorX * PosScaling;
+		Expression->MaterialExpressionEditorY *= PosScaling;
+	}
+	for (int32 ExpressionIndex = 0; ExpressionIndex < Comments.Num(); ExpressionIndex++)
+	{
+		UMaterialExpressionComment* Comment = Comments[ExpressionIndex];
+		Comment->MaterialExpressionEditorX = -Comment->MaterialExpressionEditorX * PosScaling - Comment->SizeX;
+		Comment->MaterialExpressionEditorY *= PosScaling;
+		Comment->SizeX *= PosScaling;
+		Comment->SizeY *= PosScaling;
+	}
+}
+#endif //WITH_EDITORONLY_DATA
diff --git a/Engine/Source/Runtime/Engine/Private/Materials/MaterialExpressions.cpp b/Engine/Source/Runtime/Engine/Private/Materials/MaterialExpressions.cpp
--- a/Engine/Source/Runtime/Engine/Private/Materials/MaterialExpressions.cpp
+++ b/Engine/Source/Runtime/Engine/Private/Materials/MaterialExpressions.cpp
@@ -6040,6 +6040,12 @@ void UMaterialFunction::PostLoad()
 		}
 	}
 #endif
+#if WITH_EDITORONLY_DATA
+	if (GetLinkerUE4Version() < VER_UE4_FLIP_MATERIAL_COORDS)
+	{
+		UMaterial::FlipExpressionPositions(FunctionExpressions, FunctionEditorComments);
+	}
+#endif //WITH_EDITORONLY_DATA
 }
 
 void UMaterialFunction::GetAssetRegistryTags(TArray<FAssetRegistryTag>& OutTags) const
