Altered material editor stats to optionally use relase shaders and to display a "built-in" shader cost.

I've altered the preview material directly to compile with release shaders so that this incurs no overhead.
Refactored the empty material stuff to be neater and to report the base cost of the shader.
Reduced the shaders compiled for this empty stats material to the bare minimum. Typically 3-5, adding a near insignificat cost.

diff --git a/Engine/Source/Editor/MaterialEditor/Private/MaterialEditor.cpp b/Engine/Source/Editor/MaterialEditor/Private/MaterialEditor.cpp
--- a/Engine/Source/Editor/MaterialEditor/Private/MaterialEditor.cpp
+++ b/Engine/Source/Editor/MaterialEditor/Private/MaterialEditor.cpp
@@ -46,13 +46,6 @@ const FName FMaterialEditor::PaletteTabId( TEXT( "MaterialEditor_Palette" ) );
 const FName FMaterialEditor::StatsTabId( TEXT( "MaterialEditor_Stats" ) );
 const FName FMaterialEditor::FindTabId( TEXT( "MaterialEditor_Find" ) );
 
-static int32 GShowDevOverheadStats = 0;
-static FAutoConsoleVariableRef CVarShowDevOverheadStats(
-	TEXT("r.ShowMaterialDevOverheadStats"),
-	GShowDevOverheadStats,
-	TEXT("Show the overhead introduced into materials by the editor.")
-	);
-
 ///////////////////////////
 // FMatExpressionPreview //
 ///////////////////////////
@@ -192,8 +185,6 @@ void FMaterialEditor::InitEditorForMaterial(UMaterial* InMaterial)
 			Material->Expressions.RemoveAt(ExpressionIndex);
 		}
 	}
-
-	MaterialDevelopmentOverheadStats.Init(Material);
 }
 
 void FMaterialEditor::InitEditorForMaterialFunction(UMaterialFunction* InMaterialFunction)
@@ -443,6 +434,7 @@ FMaterialEditor::FMaterialEditor()
 	, Material(NULL)
 	, OriginalMaterial(NULL)
 	, ExpressionPreviewMaterial(NULL)
+	, EmptyMaterial(NULL)
 	, PreviewExpression(NULL)
 	, MaterialFunction(NULL)
 	, OriginalMaterialObject(NULL)
@@ -452,6 +444,7 @@ FMaterialEditor::FMaterialEditor()
 	, bHideUnusedConnectors(false)
 	, bIsRealtime(false)
 	, bShowStats(true)
+	, bShowBuiltinStats(false)
 	, bShowMobileStats(false)
 {
 }
@@ -690,6 +683,8 @@ void FMaterialEditor::ExtendToolbar()
 				ToolbarBuilder.AddToolBarButton(FMaterialEditorCommands::Get().ToggleRealtimeExpressions);
 				ToolbarBuilder.AddToolBarButton(FMaterialEditorCommands::Get().AlwaysRefreshAllPreviews);
 				ToolbarBuilder.AddToolBarButton(FMaterialEditorCommands::Get().ToggleMaterialStats);
+				ToolbarBuilder.AddToolBarButton(FMaterialEditorCommands::Get().ToggleReleaseStats);
+				ToolbarBuilder.AddToolBarButton(FMaterialEditorCommands::Get().ToggleBuiltinStats);
 				ToolbarBuilder.AddToolBarButton(FMaterialEditorCommands::Get().ToggleMobileStats);
 			}
 			ToolbarBuilder.EndSection();
@@ -988,6 +983,16 @@ void FMaterialEditor::LoadEditorSettings()
 		ToggleMobileStats();
 	}
 
+	if (EditorOptions->bReleaseStats)
+	{
+		ToggleReleaseStats();
+	}
+
+	if (EditorOptions->bShowBuiltinStats)
+	{
+		ToggleBuiltinStats();
+	}
+
 	// Primitive type
 	int32 PrimType;
 	if(GConfig->GetInt(TEXT("MaterialEditor"), TEXT("PrimType"), PrimType, GEditorUserSettingsIni))
@@ -1102,9 +1107,6 @@ void FMaterialEditor::UpdatePreviewMaterial( )
 		// If we are previewing an expression, update the expression preview material
 		ExpressionPreviewMaterial->PreEditChange( NULL );
 		ExpressionPreviewMaterial->PostEditChange();
-
-		//Don't give overhead stats for preview expression material.
-		MaterialDevelopmentOverheadStats.Update(NULL);
 	}
 	else 
 	{
@@ -1112,10 +1114,10 @@ void FMaterialEditor::UpdatePreviewMaterial( )
 		Material->PreEditChange( NULL );
 		Material->PostEditChange();
 
+		UpdateStatsMaterials();
+
 		// Null out the expression preview material so they can be GC'ed
 		ExpressionPreviewMaterial = NULL;
-
-		MaterialDevelopmentOverheadStats.Update(Material);
 	}
 
 
@@ -1463,22 +1465,37 @@ void FMaterialEditor::UpdateMaterialInfoList(bool bForceDisplay)
 				// Display any errors and messages in the upper left corner of the viewport.
 				TArray<FString> Descriptions;
 				TArray<int32> InstructionCounts;
+				TArray<FString> EmptyDescriptions;
+				TArray<int32> EmptyInstructionCounts;
+
 				MaterialResource->GetRepresentativeInstructionCounts(Descriptions, InstructionCounts);
 
-				TArray<int32> OverheadCounts;
-				bool bDoOverheadCount = MaterialDevelopmentOverheadStats.GetOverheadCounts(OverheadCounts, FeatureLevel);
-				bDoOverheadCount = bDoOverheadCount && OverheadCounts.Num() > 0 && InstructionCounts.Num() > 0;
+				bool bBuiltinStats = false;
+				const FMaterialResource* EmptyMaterialResource = EmptyMaterial ? EmptyMaterial->GetMaterialResource(FeatureLevel) : NULL;
+				if (bShowBuiltinStats && bStatsFromPreviewMaterial && EmptyMaterialResource && InstructionCounts.Num() > 0)
+				{
+					EmptyMaterialResource->GetRepresentativeInstructionCounts(EmptyDescriptions, EmptyInstructionCounts);
+
+					if (EmptyInstructionCounts.Num() > 0)
+					{
+						//The instruction counts should match. If not, the preview material has been changed without the EmptyMaterial being updated to match.
+						if (ensure(InstructionCounts.Num() == EmptyInstructionCounts.Num()))
+						{
+							bBuiltinStats = true;
+						}
+					}
+				}
 
 				for (int32 InstructionIndex = 0; InstructionIndex < Descriptions.Num(); InstructionIndex++)
 				{
-					FString InstructionCountString = FString::Printf(TEXT("%s: %u instructions"),*Descriptions[InstructionIndex],InstructionCounts[InstructionIndex]);
-					if (bDoOverheadCount && ensure(OverheadCounts.Num() == InstructionCounts.Num()))
+					FString InstructionCountString = FString::Printf(TEXT("%s: %u instructions"),*Descriptions[InstructionIndex], InstructionCounts[InstructionIndex]);
+					if (bBuiltinStats)
 					{
-						InstructionCountString += FString::Printf( TEXT(" - Dev Overhead: %u"), OverheadCounts[InstructionIndex] );
+						InstructionCountString += FString::Printf(TEXT(" - Built-in instructions: %u"), EmptyInstructionCounts[InstructionIndex]);
 					}
 					TempMaterialInfoList.Add(MakeShareable(new FMaterialInfo(InstructionCountString, FLinearColor::Yellow)));
-					TSharedRef<FTokenizedMessage> Line = FTokenizedMessage::Create( EMessageSeverity::Info );
-					Line->AddToken( FTextToken::Create( FText::FromString( InstructionCountString ) ) );
+					TSharedRef<FTokenizedMessage> Line = FTokenizedMessage::Create(EMessageSeverity::Info);
+					Line->AddToken(FTextToken::Create(FText::FromString(InstructionCountString)));
 					Messages.Add(Line);
 				}
 
@@ -1600,7 +1617,7 @@ void FMaterialEditor::AddReferencedObjects( FReferenceCollector& Collector )
 	Collector.AddReferencedObject( OriginalMaterial );
 	Collector.AddReferencedObject( MaterialFunction );
 	Collector.AddReferencedObject( ExpressionPreviewMaterial );
-	MaterialDevelopmentOverheadStats.AddReferencedObjects( Collector );
+	Collector.AddReferencedObject( EmptyMaterial );
 }
 
 void FMaterialEditor::BindCommands()
@@ -1660,6 +1677,18 @@ void FMaterialEditor::BindCommands()
 		FCanExecuteAction(),
 		FIsActionChecked::CreateSP(this, &FMaterialEditor::IsToggleStatsChecked));
 
+	ToolkitCommands->MapAction(
+		Commands.ToggleReleaseStats,
+		FExecuteAction::CreateSP(this, &FMaterialEditor::ToggleReleaseStats),
+		FCanExecuteAction(),
+		FIsActionChecked::CreateSP(this, &FMaterialEditor::IsToggleReleaseStatsChecked));
+
+	ToolkitCommands->MapAction(
+		Commands.ToggleBuiltinStats,
+		FExecuteAction::CreateSP(this, &FMaterialEditor::ToggleBuiltinStats),
+		FCanExecuteAction(),
+		FIsActionChecked::CreateSP(this, &FMaterialEditor::IsToggleBuiltinStatsChecked));
+
 	ToolkitCommands->MapAction(
 		Commands.ToggleMobileStats,
 		FExecuteAction::CreateSP(this, &FMaterialEditor::ToggleMobileStats),
@@ -1787,6 +1816,35 @@ bool FMaterialEditor::IsToggleStatsChecked() const
 	return bShowStats == true;
 }
 
+void FMaterialEditor::ToggleReleaseStats()
+{
+	Material->bAllowDevelopmentShaderCompile = !Material->bAllowDevelopmentShaderCompile;
+	UpdatePreviewMaterial();
+}
+
+bool FMaterialEditor::IsToggleReleaseStatsChecked() const
+{
+	return !Material->bAllowDevelopmentShaderCompile;
+}
+
+void FMaterialEditor::ToggleBuiltinStats()
+{
+	bShowBuiltinStats = !bShowBuiltinStats;
+
+	if (bShowBuiltinStats && !bStatsFromPreviewMaterial)
+	{
+		//Have to be start using the preview material for stats.
+		UpdatePreviewMaterial();
+	}
+
+	UpdateStatsMaterials();
+}
+
+bool FMaterialEditor::IsToggleBuiltinStatsChecked() const
+{
+	return bShowBuiltinStats;
+}
+
 void FMaterialEditor::ToggleMobileStats()
 {
 	// Toggle the showing of material stats each time the user presses the show stats button
@@ -1806,6 +1864,7 @@ void FMaterialEditor::ToggleMobileStats()
 				OriginalMaterial->ForceRecompileForRendering();
 			}
 		}
+		UpdateStatsMaterials();
 		RefreshPreviewViewport();
 	}
 	UpdateMaterialInfoList(bShowMobileStats);
@@ -2252,7 +2311,7 @@ void FMaterialEditor::OnMaterialUsageFlagsChanged(UMaterial* MaterialThatChanged
 	{
 		bool bNeedsRecompile = false;
 		Material->SetMaterialUsage(bNeedsRecompile, Flag, MaterialThatChanged->GetUsageByFlag(Flag));
-		MaterialDevelopmentOverheadStats.Update(Material);
+		UpdateStatsMaterials();
 	}
 }
 
@@ -3724,89 +3783,28 @@ FReply FMaterialEditor::OnSpawnGraphNodeByShortcut(FInputGesture InGesture, cons
 	return FReply::Unhandled();
 }
 
-void FMaterialDevelopmentOverheadStats::Init(UMaterial* InMaterial)
-{
-	bDisabled = false;
-	Update(InMaterial);
-}
-
-void FMaterialDevelopmentOverheadStats::Update(UMaterial* InMaterial)
-{
-	if (NULL == InMaterial || 0 == GShowDevOverheadStats)
-	{
-		//Don't bother clearing the materials as we'll either be reusing them again shortly or exiting the editor.
-		bDisabled = true;
-		return;
-	}
-
-	bDisabled = false;
-
-	//Duplicate the material in place (same name).
-	EmptyMaterial = (UMaterial*)StaticDuplicateObject(InMaterial, GetTransientPackage(), TEXT("EmptyMaterial"), ~RF_Standalone, UPreviewMaterial::StaticClass()); 		
-	EmptyMaterial->Expressions.Empty();		
-	
-	EmptyMaterialWithOverhead = (UMaterial*)StaticDuplicateObject(InMaterial, GetTransientPackage(), TEXT("EmptyOverheadMaterial"), ~RF_Standalone, UPreviewMaterial::StaticClass()); 		
-	EmptyMaterialWithOverhead->Expressions.Empty();	
-
-	//Disconnect all properties from the expressions
-	for( int32 PropIdx = 0; PropIdx < MP_MAX ; ++PropIdx )
-	{
-		FExpressionInput* ExpInput = EmptyMaterial->GetExpressionInputForProperty((EMaterialProperty)PropIdx);
-		ExpInput->Expression = NULL;	
-		ExpInput = EmptyMaterialWithOverhead->GetExpressionInputForProperty((EMaterialProperty)PropIdx);
-		ExpInput->Expression = NULL;
-	}
-
-	EmptyMaterial->bIsMaterialDevelopmentOverheadStatsMaterial = true;
-
-	EmptyMaterial->bAllowDevelopmentShaderCompile = false;
-	EmptyMaterial->PreEditChange( NULL );
-	EmptyMaterial->PostEditChange();
-
-	EmptyMaterialWithOverhead->bIsMaterialDevelopmentOverheadStatsMaterial = true;
-	EmptyMaterialWithOverhead->PreEditChange( NULL );
-	EmptyMaterialWithOverhead->PostEditChange();
-}
-
-void FMaterialDevelopmentOverheadStats::AddReferencedObjects(FReferenceCollector& Collector)
+void FMaterialEditor::UpdateStatsMaterials()
 {
- 	Collector.AddReferencedObject(EmptyMaterial);
- 	Collector.AddReferencedObject(EmptyMaterialWithOverhead);
-}
+	if (bShowBuiltinStats && bStatsFromPreviewMaterial)
+	{	
+		UMaterial* StatsMaterial = Material;
+		FString EmptyMaterialName = FString(TEXT("MEStatsMaterial_Empty_")) + Material->GetName();
+		EmptyMaterial = (UMaterial*)StaticDuplicateObject(Material, GetTransientPackage(), *EmptyMaterialName, ~RF_Standalone, UPreviewMaterial::StaticClass());
 
-bool FMaterialDevelopmentOverheadStats::GetOverheadCounts(TArray<int32>& OverheadCounts, ERHIFeatureLevel::Type FeatureLevel)
-{		
-	if (bDisabled || 0 == GShowDevOverheadStats)
-	{
-		return false;
-	}
+		EmptyMaterial->SetFeatureLevelToCompile(ERHIFeatureLevel::ES2, bShowMobileStats);
 
-	const FMaterialResource* EmptyMaterialResource = EmptyMaterial ? EmptyMaterial->GetMaterialResource(FeatureLevel) : NULL;
-	const FMaterialResource* EmptyMaterialResourceWithOverhead = EmptyMaterialWithOverhead ? EmptyMaterialWithOverhead->GetMaterialResource(FeatureLevel) : NULL;
-	
-	if( !EmptyMaterialResource || !EmptyMaterialResourceWithOverhead )
-	{
-		return false;
-	}
-		
-	TArray<FString> Descriptions;
-	TArray<int32> InstructionCounts;
-	TArray<FString> OverheadDescriptions;
-	TArray<int32> OverheadInstructionCounts;
-
-	EmptyMaterialResource->GetRepresentativeInstructionCounts(Descriptions, InstructionCounts);
-	EmptyMaterialResourceWithOverhead->GetRepresentativeInstructionCounts(OverheadDescriptions, OverheadInstructionCounts);
-
-	if( InstructionCounts.Num() == 0 || OverheadInstructionCounts.Num() == 0 )
-	{
-		return false;
-	}
+		EmptyMaterial->Expressions.Empty();
 
-	for( int32 i=0 ; i < InstructionCounts.Num() ; ++i )
-	{
-		OverheadCounts.Add( OverheadInstructionCounts[i] - InstructionCounts[i] );
+		//Disconnect all properties from the expressions
+		for (int32 PropIdx = 0; PropIdx < MP_MAX; ++PropIdx)
+		{
+			FExpressionInput* ExpInput = EmptyMaterial->GetExpressionInputForProperty((EMaterialProperty)PropIdx);
+			ExpInput->Expression = NULL;
+		}
+		EmptyMaterial->bAllowDevelopmentShaderCompile = Material->bAllowDevelopmentShaderCompile;
+		EmptyMaterial->PreEditChange(NULL);
+		EmptyMaterial->PostEditChange();
 	}
-	return true;
 }
 
 #undef LOCTEXT_NAMESPACE
diff --git a/Engine/Source/Editor/MaterialEditor/Private/MaterialEditor.h b/Engine/Source/Editor/MaterialEditor/Private/MaterialEditor.h
--- a/Engine/Source/Editor/MaterialEditor/Private/MaterialEditor.h
+++ b/Engine/Source/Editor/MaterialEditor/Private/MaterialEditor.h
@@ -174,33 +174,6 @@ struct FMaterialInfo
 	{}
 };
 
-/**
-This class allows for the display of stats on the overhead introduce into shaders because they have been compiled for development.
-It compiles two empty materials with the same set up as the current preview material. One with the development overhead and one without.
-The difference in instruction counts between the two is reported to the user as the development overhead.
-*/
-class FMaterialDevelopmentOverheadStats
-{
-public:
-	FMaterialDevelopmentOverheadStats()
-		: EmptyMaterialWithOverhead(NULL)
-		, EmptyMaterial(NULL)
-		, bDisabled(false)
-	{
-
-	}
-	void Init( UMaterial* InMaterial );
-	void Update( UMaterial* InMaterial );
-	bool GetOverheadCounts( TArray<int32>& OverheadCounts, ERHIFeatureLevel::Type FeatureLevel  );
-	void AddReferencedObjects( FReferenceCollector& Collector );
-
-private:
-	UMaterial* EmptyMaterialWithOverhead;
-	UMaterial* EmptyMaterial;
-
-	bool bDisabled;
-};
-
 /**
  * Material Editor class
  */
@@ -347,6 +320,9 @@ public:
 	virtual int32 GetNumberOfSelectedNodes() const OVERRIDE;
 	virtual FMaterialRenderProxy* GetExpressionPreview(UMaterialExpression* InExpression) OVERRIDE;
 
+
+	void UpdateStatsMaterials();
+
 public:
 	/** Set to true when modifications have been made to the material */
 	bool bMaterialDirty;
@@ -363,8 +339,8 @@ public:
 	/** The material applied to the preview mesh when previewing an expression. */
 	UMaterial* ExpressionPreviewMaterial;
 
-	/** An helper class containing empty materials to calculate editor overhead on this material. */
-	FMaterialDevelopmentOverheadStats MaterialDevelopmentOverheadStats;
+	/** An empty copy of the preview material. Allows displaying of stats about the built in cost of the current material. */
+	UMaterial* EmptyMaterial;
 
 	/** The expression currently being previewed.  This is NULL when not in expression preview mode. */
 	UMaterialExpression* PreviewExpression;
@@ -514,6 +490,11 @@ private:
 	/** Command for the stats button */
 	void ToggleStats();
 	bool IsToggleStatsChecked() const;
+	void ToggleReleaseStats();
+	bool IsToggleReleaseStatsChecked() const;
+	void ToggleBuiltinStats();
+	bool IsToggleBuiltinStatsChecked() const;
+
 	/** Mobile stats button. */
 	void ToggleMobileStats();
 	bool IsToggleMobileStatsChecked() const;
@@ -668,6 +649,9 @@ private:
 	/** If true, show material stats like number of shader instructions. */
 	bool bShowStats;
 
+	/** If true, show stats for an empty material. Helps artists to judge the cost of their changes to the graph. */
+	bool bShowBuiltinStats;
+
 	/** If true, show material stats and errors for mobile. */
 	bool bShowMobileStats;
 
diff --git a/Engine/Source/Editor/MaterialEditor/Private/MaterialEditorActions.cpp b/Engine/Source/Editor/MaterialEditor/Private/MaterialEditorActions.cpp
--- a/Engine/Source/Editor/MaterialEditor/Private/MaterialEditorActions.cpp
+++ b/Engine/Source/Editor/MaterialEditor/Private/MaterialEditorActions.cpp
@@ -26,7 +26,9 @@ void FMaterialEditorCommands::RegisterCommands()
 	UI_COMMAND( ShowHideConnectors, "Connectors", "Show or Hide Unused Connectors", EUserInterfaceActionType::ToggleButton, FInputGesture() );
 	UI_COMMAND( ToggleRealtimeExpressions, "Live Nodes", "Toggles real time update of the graph canvas.", EUserInterfaceActionType::ToggleButton, FInputGesture() );
 	UI_COMMAND( AlwaysRefreshAllPreviews, "Live Update", "All nodes are previewed live.", EUserInterfaceActionType::ToggleButton, FInputGesture() );
-	UI_COMMAND( ToggleMaterialStats, "Stats", "Toggles displaying of the material's stats.", EUserInterfaceActionType::ToggleButton, FInputGesture() );
+	UI_COMMAND( ToggleMaterialStats, "Stats", "Toggles displaying of the material's stats.", EUserInterfaceActionType::ToggleButton, FInputGesture());
+	UI_COMMAND( ToggleReleaseStats, "Release Stats", "Toggles use of release or development shaders in stats reporting.", EUserInterfaceActionType::ToggleButton, FInputGesture());
+	UI_COMMAND( ToggleBuiltinStats, "Built-In Stats", "Toggles displaying of stats for an empty material. Helps identify cost of changes in the graph.", EUserInterfaceActionType::ToggleButton, FInputGesture());
 	UI_COMMAND( ToggleMobileStats, "Mobile Stats", "Toggles material stats and compilation errors for mobile.", EUserInterfaceActionType::ToggleButton, FInputGesture() );
 	UI_COMMAND( NewComment, "New Comment", "Creates a new comment node.", EUserInterfaceActionType::Button, FInputGesture() );
 	UI_COMMAND( MatertialPasteHere, "Paste Here", "Pastes copied items at this location.", EUserInterfaceActionType::Button, FInputGesture() );
diff --git a/Engine/Source/Editor/MaterialEditor/Public/MaterialEditorActions.h b/Engine/Source/Editor/MaterialEditor/Public/MaterialEditorActions.h
--- a/Engine/Source/Editor/MaterialEditor/Public/MaterialEditorActions.h
+++ b/Engine/Source/Editor/MaterialEditor/Public/MaterialEditorActions.h
@@ -85,6 +85,12 @@ public:
 	/** Toggles the material stats on the canvas pane */
 	TSharedPtr< FUICommandInfo > ToggleMaterialStats;
 
+	/** Toggles whether the reported stats are for release shaders or development shaders. */
+	TSharedPtr< FUICommandInfo > ToggleReleaseStats;
+
+	/** Toggles the reporting of instructions for an empty material. i.e. The built in cost of that material. */
+	TSharedPtr< FUICommandInfo > ToggleBuiltinStats;
+
 	/** Shows material stats and errors when compiled for mobile. */
 	TSharedPtr< FUICommandInfo > ToggleMobileStats;
 	
diff --git a/Engine/Source/Editor/UnrealEd/Classes/Preferences/MaterialEditorOptions.h b/Engine/Source/Editor/UnrealEd/Classes/Preferences/MaterialEditorOptions.h
--- a/Engine/Source/Editor/UnrealEd/Classes/Preferences/MaterialEditorOptions.h
+++ b/Engine/Source/Editor/UnrealEd/Classes/Preferences/MaterialEditorOptions.h
@@ -54,6 +54,14 @@ class UNREALED_API UMaterialEditorOptions : public UObject
 	UPROPERTY(EditAnywhere, config, Category=Options)
 	uint32 bShowMobileStats:1;
 
+	/** If true, stats use release shaders. */
+	UPROPERTY(EditAnywhere, config, Category = Options)
+	uint32 bReleaseStats : 1;
+
+	/** If true, show stats for an blank copy (no graph) of the preview material. Helps identify cost of material graph. Incurs slight cost to build/iteration times. */
+	UPROPERTY(EditAnywhere, config, Category = Options)
+	uint32 bShowBuiltinStats : 1;
+
 	/** The users favorite material expressions. */
 	UPROPERTY(EditAnywhere, config, Category=Options)
 	TArray<FString> FavoriteExpressions;
diff --git a/Engine/Source/Editor/UnrealEd/Private/PreviewMaterial.cpp b/Engine/Source/Editor/UnrealEd/Private/PreviewMaterial.cpp
--- a/Engine/Source/Editor/UnrealEd/Private/PreviewMaterial.cpp
+++ b/Engine/Source/Editor/UnrealEd/Private/PreviewMaterial.cpp
@@ -35,6 +35,8 @@ public:
 		// only generate the needed shaders (which should be very restrictive for fast recompiling during editing)
 		// @todo: Add a FindShaderType by fname or something
 
+		bool bEditorStatsMaterial = Material->bIsMaterialEditorStatsMaterial;
+
 		// Always allow HitProxy shaders.
 		if (FCString::Stristr(ShaderType->GetName(), TEXT("HitProxy")))
 		{
@@ -46,7 +48,8 @@ public:
 		{
 			//cache for gpu skinned vertex factory if the material allows it
 			//this way we can have a preview skeletal mesh
-			if (!IsUsedWithSkeletalMesh() ||
+			if (bEditorStatsMaterial ||
+				!IsUsedWithSkeletalMesh() ||
 				(VertexFactoryType != FindVertexFactoryType(FName(TEXT("TGPUSkinVertexFactoryfalse"), FNAME_Find)) &&
 				VertexFactoryType != FindVertexFactoryType(FName(TEXT("TGPUSkinVertexFactorytrue"), FNAME_Find))))
 			{
@@ -54,6 +57,16 @@ public:
 			}
 		}
 
+		if (bEditorStatsMaterial)
+		{
+			TArray<FString> ShaderTypeNames;
+			TArray<FString> ShaderTypeDescriptions;
+			GetRepresentativeShaderTypesAndDescriptions(ShaderTypeNames, ShaderTypeDescriptions);
+
+			//Only allow shaders that are used in the stats.
+			return ShaderTypeNames.Find(ShaderType->GetName()) != INDEX_NONE;
+		}
+
 		// look for any of the needed type
 		bool bShaderTypeMatches = false;
 
diff --git a/Engine/Source/Runtime/Engine/Classes/Materials/Material.h b/Engine/Source/Runtime/Engine/Classes/Materials/Material.h
--- a/Engine/Source/Runtime/Engine/Classes/Materials/Material.h
+++ b/Engine/Source/Runtime/Engine/Classes/Materials/Material.h
@@ -647,9 +647,9 @@ public:
 	UPROPERTY(transient, duplicatetransient)
 	uint32 bAllowDevelopmentShaderCompile:1;
 
-	/** true if this is one of the materials used in development overhead calculations in the material editor. */
+	/** true if this is a special material used for stats by the material editor. */
 	UPROPERTY(transient, duplicatetransient)
-	uint32 bIsMaterialDevelopmentOverheadStatsMaterial:1;
+	uint32 bIsMaterialEditorStatsMaterial:1;
 
 	/** true if we have printed a warning about material usage for a given usage flag. */
 	UPROPERTY(transient, duplicatetransient)
diff --git a/Engine/Source/Runtime/Engine/Private/Materials/Material.cpp b/Engine/Source/Runtime/Engine/Private/Materials/Material.cpp
--- a/Engine/Source/Runtime/Engine/Private/Materials/Material.cpp
+++ b/Engine/Source/Runtime/Engine/Private/Materials/Material.cpp
@@ -517,7 +517,7 @@ UMaterial::UMaterial(const class FPostConstructInitializeProperties& PCIP)
 	MaterialDecalResponse = MDR_ColorNormalRoughness;
 
 	bAllowDevelopmentShaderCompile = true;
-	bIsMaterialDevelopmentOverheadStatsMaterial = false;
+	bIsMaterialEditorStatsMaterial = false;
 
 #if WITH_EDITORONLY_DATA
 	MaterialGraph = NULL;
@@ -1856,6 +1856,11 @@ void UMaterial::PostLoad()
 		AssertDefaultMaterialsPostLoaded();
 	}	
 
+	if ( GIsEditor && GetOuter() == GetTransientPackage() && FCString::Strstr(*GetName(), TEXT("MEStatsMaterial_")))
+	{
+		bIsMaterialEditorStatsMaterial = true;
+	}
+
 	// Ensure expressions have been postloaded before we use them for compiling
 	// Any UObjects used by material compilation must be postloaded here
 	for (int32 ExpressionIndex = 0; ExpressionIndex < Expressions.Num(); ExpressionIndex++)
@@ -1978,7 +1983,8 @@ void UMaterial::PostLoad()
 	{
 		SCOPE_SECONDS_COUNTER(MaterialLoadTime);
 
-		if (FApp::CanEverRender()) 
+		//Don't compile shaders in post load for dev overhead materials.
+		if (FApp::CanEverRender() && !bIsMaterialEditorStatsMaterial)
 		{
 			CacheResourceShadersForRendering(false);
 		}
@@ -2208,7 +2214,7 @@ void UMaterial::PostEditChangeProperty(FPropertyChangedEvent& PropertyChangedEve
 		// Ensure that any components with static elements using this material are reregistered so changes
 		// are propagated to them. The preview material is only applied to the preview mesh component,
 		// and that reregister is handled by the material editor.
-		if( !bIsPreviewMaterial && !bIsMaterialDevelopmentOverheadStatsMaterial )
+		if (!bIsPreviewMaterial && !bIsMaterialEditorStatsMaterial)
 		{
 			FGlobalComponentReregisterContext RecreateComponents;
 		}
diff --git a/Engine/Source/Runtime/Engine/Private/Materials/MaterialShared.cpp b/Engine/Source/Runtime/Engine/Private/Materials/MaterialShared.cpp
--- a/Engine/Source/Runtime/Engine/Private/Materials/MaterialShared.cpp
+++ b/Engine/Source/Runtime/Engine/Private/Materials/MaterialShared.cpp
@@ -858,73 +858,12 @@ void FMaterialResource::GetRepresentativeInstructionCounts(TArray<FString> &Desc
 	TArray<FString> ShaderTypeNames;
 	TArray<FString> ShaderTypeDescriptions;
 
-	static auto* MobileHDR = IConsoleManager::Get().FindTConsoleVariableDataInt(TEXT("r.MobileHDR"));
-	bool bMobileHDR = MobileHDR && MobileHDR->GetValueOnAnyThread() == 1;
-
 	//when adding a shader type here be sure to update FPreviewMaterial::ShouldCache()
 	//so the shader type will get compiled with preview materials
 	const FMaterialShaderMap* MaterialShaderMap = GetGameThreadShaderMap();
 	if (MaterialShaderMap && MaterialShaderMap->IsCompilationFinalized())
 	{
-		if (GetFeatureLevel() > ERHIFeatureLevel::ES2)
-		{
-			if (GetLightingModel() == MLM_Unlit)
-			{
-				//unlit materials are never lightmapped
-				new (ShaderTypeNames) FString(TEXT("TBasePassPSFNoLightMapPolicy"));
-				new (ShaderTypeDescriptions) FString(TEXT("Base pass shader without light map"));
-			}
-			else
-			{
-				if (IsUsedWithStaticLighting())
-				{
-					//lit materials are usually lightmapped
-					new (ShaderTypeNames) FString(TEXT("TBasePassPSTDistanceFieldShadowsAndLightMapPolicyHQ"));
-					new (ShaderTypeDescriptions) FString(TEXT("Base pass shader with static lighting"));
-				}
-
-				//also show a dynamically lit shader
-				new (ShaderTypeNames) FString(TEXT("TBasePassPSFNoLightMapPolicy"));
-				new (ShaderTypeDescriptions) FString(TEXT("Base pass shader with only dynamic lighting"));
-
-				if (IsTranslucentBlendMode(GetBlendMode()))
-				{
-					new (ShaderTypeNames) FString(TEXT("TBasePassPSFSelfShadowedTranslucencyPolicy"));
-					new (ShaderTypeDescriptions) FString(TEXT("Base pass shader for self shadowed translucency"));
-				}
-			}
-			
-			new (ShaderTypeNames) FString(TEXT("TBasePassVSFNoLightMapPolicy"));
-			new (ShaderTypeDescriptions) FString(TEXT("Vertex shader"));
-		}
-		else
-		{
-			const TCHAR* ShaderSuffix = bMobileHDR ? TEXT("HDRLinear64") : TEXT("LDRGamma32");
-			const TCHAR* DescSuffix = bMobileHDR ? TEXT(" (HDR)") : TEXT(" (LDR)");
-
-			if (GetLightingModel() == MLM_Unlit)
-			{
-				//unlit materials are never lightmapped
-				new (ShaderTypeNames) FString(FString::Printf(TEXT("TBasePassForForwardShadingPSFNoLightMapPolicy%s"),ShaderSuffix));
-				new (ShaderTypeDescriptions) FString(FString::Printf(TEXT("Mobile base pass shader without light map%s"),DescSuffix));
-			}
-			else
-			{
-				if (IsUsedWithStaticLighting())
-				{
-					//lit materials are usually lightmapped
-					new (ShaderTypeNames) FString(FString::Printf(TEXT("TBasePassForForwardShadingPSTDistanceFieldShadowsAndLightMapPolicyLQ%s"),ShaderSuffix));
-					new (ShaderTypeDescriptions) FString(FString::Printf(TEXT("Mobile base pass shader with static lighting%s"),DescSuffix));
-				}
-
-				//also show a dynamically lit shader
-				new (ShaderTypeNames) FString(FString::Printf(TEXT("TBasePassForForwardShadingPSFSimpleDirectionalLightAndSHIndirectPolicy%s"),ShaderSuffix));
-				new (ShaderTypeDescriptions) FString(FString::Printf(TEXT("Mobile base pass shader with only dynamic lighting%s"),DescSuffix));
-			}
-			
-			new (ShaderTypeNames) FString(FString::Printf(TEXT("TBasePassForForwardShadingVSFNoLightMapPolicy%s"),ShaderSuffix));
-			new (ShaderTypeDescriptions) FString(FString::Printf(TEXT("Mobile base pass vertex shader%s"),DescSuffix));
-		}
+		GetRepresentativeShaderTypesAndDescriptions(ShaderTypeNames, ShaderTypeDescriptions);
 
 		const FMeshMaterialShaderMap* MeshShaderMap = MaterialShaderMap->GetMeshShaderMap(&FLocalVertexFactory::StaticType);
 		if (MeshShaderMap)
@@ -952,6 +891,72 @@ void FMaterialResource::GetRepresentativeInstructionCounts(TArray<FString> &Desc
 	check(Descriptions.Num() == InstructionCounts.Num());
 }
 
+void FMaterialResource::GetRepresentativeShaderTypesAndDescriptions(TArray<FString> &ShaderTypeNames, TArray<FString> &ShaderTypeDescriptions) const
+{
+	static auto* MobileHDR = IConsoleManager::Get().FindTConsoleVariableDataInt(TEXT("r.MobileHDR"));
+	bool bMobileHDR = MobileHDR && MobileHDR->GetValueOnAnyThread() == 1;
+
+	if (GetFeatureLevel() > ERHIFeatureLevel::ES2)
+	{
+		if (GetLightingModel() == MLM_Unlit)
+		{
+			//unlit materials are never lightmapped
+			new (ShaderTypeNames)FString(TEXT("TBasePassPSFNoLightMapPolicy"));
+			new (ShaderTypeDescriptions)FString(TEXT("Base pass shader without light map"));
+		}
+		else
+		{
+			if (IsUsedWithStaticLighting())
+			{
+				//lit materials are usually lightmapped
+				new (ShaderTypeNames)FString(TEXT("TBasePassPSTDistanceFieldShadowsAndLightMapPolicyHQ"));
+				new (ShaderTypeDescriptions)FString(TEXT("Base pass shader with static lighting"));
+			}
+
+			//also show a dynamically lit shader
+			new (ShaderTypeNames)FString(TEXT("TBasePassPSFNoLightMapPolicy"));
+			new (ShaderTypeDescriptions)FString(TEXT("Base pass shader with only dynamic lighting"));
+
+			if (IsTranslucentBlendMode(GetBlendMode()))
+			{
+				new (ShaderTypeNames)FString(TEXT("TBasePassPSFSelfShadowedTranslucencyPolicy"));
+				new (ShaderTypeDescriptions)FString(TEXT("Base pass shader for self shadowed translucency"));
+			}
+		}
+
+		new (ShaderTypeNames)FString(TEXT("TBasePassVSFNoLightMapPolicy"));
+		new (ShaderTypeDescriptions)FString(TEXT("Vertex shader"));
+	}
+	else
+	{
+		const TCHAR* ShaderSuffix = bMobileHDR ? TEXT("HDRLinear64") : TEXT("LDRGamma32");
+		const TCHAR* DescSuffix = bMobileHDR ? TEXT(" (HDR)") : TEXT(" (LDR)");
+
+		if (GetLightingModel() == MLM_Unlit)
+		{
+			//unlit materials are never lightmapped
+			new (ShaderTypeNames)FString(FString::Printf(TEXT("TBasePassForForwardShadingPSFNoLightMapPolicy%s"), ShaderSuffix));
+			new (ShaderTypeDescriptions)FString(FString::Printf(TEXT("Mobile base pass shader without light map%s"), DescSuffix));
+		}
+		else
+		{
+			if (IsUsedWithStaticLighting())
+			{
+				//lit materials are usually lightmapped
+				new (ShaderTypeNames)FString(FString::Printf(TEXT("TBasePassForForwardShadingPSTDistanceFieldShadowsAndLightMapPolicyLQ%s"), ShaderSuffix));
+				new (ShaderTypeDescriptions)FString(FString::Printf(TEXT("Mobile base pass shader with static lighting%s"), DescSuffix));
+			}
+
+			//also show a dynamically lit shader
+			new (ShaderTypeNames)FString(FString::Printf(TEXT("TBasePassForForwardShadingPSFSimpleDirectionalLightAndSHIndirectPolicy%s"), ShaderSuffix));
+			new (ShaderTypeDescriptions)FString(FString::Printf(TEXT("Mobile base pass shader with only dynamic lighting%s"), DescSuffix));
+		}
+
+		new (ShaderTypeNames)FString(FString::Printf(TEXT("TBasePassForForwardShadingVSFNoLightMapPolicy%s"), ShaderSuffix));
+		new (ShaderTypeDescriptions)FString(FString::Printf(TEXT("Mobile base pass vertex shader%s"), DescSuffix));
+	}
+}
+
 SIZE_T FMaterialResource::GetResourceSizeInclusive()
 {
 	SIZE_T ResourceSize = 0;
diff --git a/Engine/Source/Runtime/Engine/Public/MaterialShared.h b/Engine/Source/Runtime/Engine/Public/MaterialShared.h
--- a/Engine/Source/Runtime/Engine/Public/MaterialShared.h
+++ b/Engine/Source/Runtime/Engine/Public/MaterialShared.h
@@ -1770,6 +1770,8 @@ public:
 	 */
 	ENGINE_API void GetRepresentativeInstructionCounts(TArray<FString> &Descriptions, TArray<int32> &InstructionCounts) const;
 
+	ENGINE_API void GetRepresentativeShaderTypesAndDescriptions(TArray<FString> &ShaderTypeNames, TArray<FString> &ShaderTypeDescriptions) const;
+
 	ENGINE_API SIZE_T GetResourceSizeInclusive();
 
 	ENGINE_API virtual void LegacySerialize(FArchive& Ar);
