Added shortcut for component mask on Shift+C.
#ttp 332908 - Material editor: shortcut for component mask
#branch UE4
#proj Editor.MaterialEditor
#add Added new material editor action for CreateComponentMaskNode.
#add Created callback for creating a component mask node at the cursor.

diff --git a/Engine/Source/Editor/MaterialEditor/Private/MaterialEditor.cpp b/Engine/Source/Editor/MaterialEditor/Private/MaterialEditor.cpp
--- a/Engine/Source/Editor/MaterialEditor/Private/MaterialEditor.cpp
+++ b/Engine/Source/Editor/MaterialEditor/Private/MaterialEditor.cpp
@@ -2287,6 +2287,11 @@ void FMaterialEditor::OnCreateComment()
 	CreateNewMaterialExpressionComment(GraphEditor->GetPasteLocation());
 }
 
+void FMaterialEditor::OnCreateComponentMaskNode()
+{
+	CreateNewMaterialExpression(UMaterialExpressionComponentMask::StaticClass(), GraphEditor->GetPasteLocation(), true, false);
+}
+
 void FMaterialEditor::OnFindInMaterial()
 {
 	TabManager->InvokeTab(FindTabId);
@@ -3466,6 +3471,10 @@ TSharedRef<SGraphEditor> FMaterialEditor::CreateGraphEditorWidget()
 		GraphEditorCommands->MapAction( FMaterialEditorCommands::Get().ForceRefreshPreviews,
 			FExecuteAction::CreateSP(this, &FMaterialEditor::OnForceRefreshPreviews)
 			);
+
+		GraphEditorCommands->MapAction( FMaterialEditorCommands::Get().CreateComponentMaskNode,
+			FExecuteAction::CreateSP(this, &FMaterialEditor::OnCreateComponentMaskNode)
+			);
 	}
 
 	FGraphAppearanceInfo AppearanceInfo;
diff --git a/Engine/Source/Editor/MaterialEditor/Private/MaterialEditor.h b/Engine/Source/Editor/MaterialEditor/Private/MaterialEditor.h
--- a/Engine/Source/Editor/MaterialEditor/Private/MaterialEditor.h
+++ b/Engine/Source/Editor/MaterialEditor/Private/MaterialEditor.h
@@ -516,6 +516,8 @@ private:
 	void OnForceRefreshPreviews();
 	/** Create comment node on graph */
 	void OnCreateComment();
+	/** Create ComponentMask node on graph */
+	void OnCreateComponentMaskNode();
 	/** Bring up the search tab */
 	void OnFindInMaterial();
 
diff --git a/Engine/Source/Editor/MaterialEditor/Private/MaterialEditorActions.cpp b/Engine/Source/Editor/MaterialEditor/Private/MaterialEditorActions.cpp
--- a/Engine/Source/Editor/MaterialEditor/Private/MaterialEditorActions.cpp
+++ b/Engine/Source/Editor/MaterialEditor/Private/MaterialEditorActions.cpp
@@ -49,6 +49,7 @@ void FMaterialEditorCommands::RegisterCommands()
 	UI_COMMAND( AddToFavorites, "Add To Favorites", "Adds this expression to your favorites.", EUserInterfaceActionType::Button, FInputGesture() );
 	UI_COMMAND( BreakLink, "Break Link", "Deletes this link.", EUserInterfaceActionType::Button, FInputGesture() );
 	UI_COMMAND( ForceRefreshPreviews, "Force Refresh Previews", "Forces a refresh of all previews", EUserInterfaceActionType::Button, FInputGesture(EKeys::SpaceBar) );
+	UI_COMMAND( CreateComponentMaskNode, "Create ComponentMask Node", "Creates a ComponentMask node at the current cursor position.", EUserInterfaceActionType::Button, FInputGesture( EModifierKey::Shift, EKeys::C ) );
 
 	UI_COMMAND( FindInMaterial, "Search", "Finds expressions and comments in the current Material", EUserInterfaceActionType::Button, FInputGesture(EModifierKey::Control, EKeys::F));
 }
diff --git a/Engine/Source/Editor/MaterialEditor/Public/MaterialEditorActions.h b/Engine/Source/Editor/MaterialEditor/Public/MaterialEditorActions.h
--- a/Engine/Source/Editor/MaterialEditor/Public/MaterialEditorActions.h
+++ b/Engine/Source/Editor/MaterialEditor/Public/MaterialEditorActions.h
@@ -152,6 +152,9 @@ public:
 	/** Finds expressions in current material */
 	TSharedPtr< FUICommandInfo > FindInMaterial;
 
+	/** Create component mask node */
+	TSharedPtr< FUICommandInfo > CreateComponentMaskNode;
+
 	/**
 	 * Initialize commands
 	 */
