UE4 Refactoring. Changed OVERRIDE and FINAL macros to keywords override and final.

diff --git a/Engine/Source/Editor/LandscapeEditor/Private/LandscapeEdModeSplineTools.cpp b/Engine/Source/Editor/LandscapeEditor/Private/LandscapeEdModeSplineTools.cpp
--- a/Engine/Source/Editor/LandscapeEditor/Private/LandscapeEdModeSplineTools.cpp
+++ b/Engine/Source/Editor/LandscapeEditor/Private/LandscapeEdModeSplineTools.cpp
@@ -54,13 +54,13 @@ public:
 		}
 	}
 
-	virtual const TCHAR* GetToolName() OVERRIDE { return TEXT("Splines"); }
-	virtual FText GetDisplayName() OVERRIDE { return NSLOCTEXT("UnrealEd", "LandscapeMode_Splines", "Splines"); };
+	virtual const TCHAR* GetToolName() override { return TEXT("Splines"); }
+	virtual FText GetDisplayName() override { return NSLOCTEXT("UnrealEd", "LandscapeMode_Splines", "Splines"); };
 
-	virtual void SetEditRenderType() OVERRIDE { GLandscapeEditRenderMode = ELandscapeEditRenderMode::None | (GLandscapeEditRenderMode & ELandscapeEditRenderMode::BitMaskForMask); }
-	virtual bool SupportsMask() OVERRIDE { return false; }
+	virtual void SetEditRenderType() override { GLandscapeEditRenderMode = ELandscapeEditRenderMode::None | (GLandscapeEditRenderMode & ELandscapeEditRenderMode::BitMaskForMask); }
+	virtual bool SupportsMask() override { return false; }
 
-	virtual bool IsValidForTarget(const FLandscapeToolTarget& Target) OVERRIDE
+	virtual bool IsValidForTarget(const FLandscapeToolTarget& Target) override
 	{
 		return true; // applied to all...
 	}
@@ -622,7 +622,7 @@ public:
 		}
 	}
 
-	virtual bool BeginTool( FLevelEditorViewportClient* ViewportClient, const FLandscapeToolTarget& InTarget, const FVector& InHitLocation) OVERRIDE
+	virtual bool BeginTool( FLevelEditorViewportClient* ViewportClient, const FLandscapeToolTarget& InTarget, const FVector& InHitLocation) override
 	{
 		LandscapeInfo = InTarget.LandscapeInfo.Get();
 
@@ -646,12 +646,12 @@ public:
 		return true;
 	}
 
-	virtual void EndTool(FLevelEditorViewportClient* ViewportClient) OVERRIDE
+	virtual void EndTool(FLevelEditorViewportClient* ViewportClient) override
 	{
 		LandscapeInfo = NULL;
 	}
 
-	virtual bool MouseMove( FLevelEditorViewportClient* ViewportClient, FViewport* Viewport, int32 x, int32 y ) OVERRIDE
+	virtual bool MouseMove( FLevelEditorViewportClient* ViewportClient, FViewport* Viewport, int32 x, int32 y ) override
 	{
 		FVector HitLocation;
 		if (EdMode->LandscapeMouseTrace(ViewportClient, x, y, HitLocation))
@@ -670,7 +670,7 @@ public:
 	{
 	}
 
-	virtual bool HandleClick( HHitProxy* HitProxy, const FViewportClick& Click ) OVERRIDE
+	virtual bool HandleClick( HHitProxy* HitProxy, const FViewportClick& Click ) override
 	{
 		if ( (!HitProxy || !HitProxy->IsA(HWidgetAxis::StaticGetType()))
 			&& !Click.IsShiftDown())
@@ -765,7 +765,7 @@ public:
 		return false;
 	}
 
-	virtual bool InputKey( FLevelEditorViewportClient* InViewportClient, FViewport* InViewport, FKey InKey, EInputEvent InEvent ) OVERRIDE
+	virtual bool InputKey( FLevelEditorViewportClient* InViewportClient, FViewport* InViewport, FKey InKey, EInputEvent InEvent ) override
 	{
 		if (InKey == EKeys::F4 && InEvent == IE_Pressed)
 		{
@@ -1061,7 +1061,7 @@ public:
 		return false;
 	}
 
-	virtual bool InputDelta( FLevelEditorViewportClient* InViewportClient, FViewport* InViewport, FVector& InDrag, FRotator& InRot, FVector& InScale ) OVERRIDE
+	virtual bool InputDelta( FLevelEditorViewportClient* InViewportClient, FViewport* InViewport, FVector& InDrag, FRotator& InRot, FVector& InScale ) override
 	{
 		FVector Drag = InDrag;
 
@@ -1228,7 +1228,7 @@ public:
 		UpdatePropertiesWindows();
 	}
 
-	virtual void EnterTool() OVERRIDE
+	virtual void EnterTool() override
 	{
 		GEditor->SelectNone(true, true, false);
 
@@ -1242,7 +1242,7 @@ public:
 		}
 	}
 
-	virtual void ExitTool() OVERRIDE
+	virtual void ExitTool() override
 	{
 		ClearSelection();
 		UpdatePropertiesWindows();
@@ -1257,7 +1257,7 @@ public:
 		}
 	}
 
-	virtual void Render(const FSceneView* View, FViewport* Viewport, FPrimitiveDrawInterface* PDI) OVERRIDE
+	virtual void Render(const FSceneView* View, FViewport* Viewport, FPrimitiveDrawInterface* PDI) override
 	{
 		for (ULandscapeSplineControlPoint* ControlPoint : SelectedSplineControlPoints)
 		{
@@ -1309,12 +1309,12 @@ public:
 		}
 	}
 
-	virtual bool OverrideSelection() const OVERRIDE
+	virtual bool OverrideSelection() const override
 	{
 		return true;
 	}
 
-	virtual bool IsSelectionAllowed( AActor* InActor, bool bInSelection ) const OVERRIDE
+	virtual bool IsSelectionAllowed( AActor* InActor, bool bInSelection ) const override
 	{
 		// Only filter selection not deselection
 		if (bInSelection)
@@ -1325,7 +1325,7 @@ public:
 		return true;
 	}
 
-	virtual bool UsesTransformWidget() const OVERRIDE
+	virtual bool UsesTransformWidget() const override
 	{
 		if (SelectedSplineControlPoints.Num() > 0)
 		{
@@ -1335,7 +1335,7 @@ public:
 		return false;
 	}
 
-	virtual EAxisList::Type GetWidgetAxisToDraw(FWidget::EWidgetMode CheckMode) const OVERRIDE
+	virtual EAxisList::Type GetWidgetAxisToDraw(FWidget::EWidgetMode CheckMode) const override
 	{
 		if (SelectedSplineControlPoints.Num() > 0)
 		{
@@ -1358,7 +1358,7 @@ public:
 		return EAxisList::None;
 	}
 
-	virtual FVector GetWidgetLocation() const OVERRIDE
+	virtual FVector GetWidgetLocation() const override
 	{
 		if (SelectedSplineControlPoints.Num() > 0)
 		{
@@ -1370,7 +1370,7 @@ public:
 		return FVector::ZeroVector;
 	}
 
-	virtual FMatrix GetWidgetRotation() const OVERRIDE
+	virtual FMatrix GetWidgetRotation() const override
 	{
 		if (SelectedSplineControlPoints.Num() > 0)
 		{
@@ -1382,7 +1382,7 @@ public:
 		return FMatrix::Identity;
 	}
 
-	virtual EEditAction::Type GetActionEditDuplicate() OVERRIDE
+	virtual EEditAction::Type GetActionEditDuplicate() override
 	{
 		if (SelectedSplineControlPoints.Num() > 0 || SelectedSplineSegments.Num() > 0)
 		{
@@ -1392,7 +1392,7 @@ public:
 		return EEditAction::Skip;
 	}
 
-	virtual EEditAction::Type GetActionEditDelete() OVERRIDE
+	virtual EEditAction::Type GetActionEditDelete() override
 	{
 		if (SelectedSplineControlPoints.Num() > 0 || SelectedSplineSegments.Num() > 0)
 		{
@@ -1402,7 +1402,7 @@ public:
 		return EEditAction::Skip;
 	}
 
-	virtual EEditAction::Type GetActionEditCut() OVERRIDE
+	virtual EEditAction::Type GetActionEditCut() override
 	{
 		if (SelectedSplineControlPoints.Num() > 0 || SelectedSplineSegments.Num() > 0)
 		{
@@ -1412,7 +1412,7 @@ public:
 		return EEditAction::Skip;
 	}
 
-	virtual EEditAction::Type GetActionEditCopy() OVERRIDE
+	virtual EEditAction::Type GetActionEditCopy() override
 	{
 		if (SelectedSplineControlPoints.Num() > 0 || SelectedSplineSegments.Num() > 0)
 		{
@@ -1422,7 +1422,7 @@ public:
 		return EEditAction::Skip;
 	}
 
-	virtual EEditAction::Type GetActionEditPaste() OVERRIDE
+	virtual EEditAction::Type GetActionEditPaste() override
 	{
 		FString PasteString;
 		FPlatformMisc::ClipboardPaste(PasteString);
@@ -1434,31 +1434,31 @@ public:
 		return EEditAction::Skip;
 	}
 
-	virtual bool ProcessEditDuplicate() OVERRIDE
+	virtual bool ProcessEditDuplicate() override
 	{
 		InternalProcessEditDuplicate();
 		return true;
 	}
 
-	virtual bool ProcessEditDelete() OVERRIDE
+	virtual bool ProcessEditDelete() override
 	{
 		InternalProcessEditDelete();
 		return true;
 	}
 
-	virtual bool ProcessEditCut() OVERRIDE
+	virtual bool ProcessEditCut() override
 	{
 		InternalProcessEditCut();
 		return true;
 	}
 
-	virtual bool ProcessEditCopy() OVERRIDE
+	virtual bool ProcessEditCopy() override
 	{
 		InternalProcessEditCopy();
 		return true;
 	}
 
-	virtual bool ProcessEditPaste() OVERRIDE
+	virtual bool ProcessEditPaste() override
 	{
 		InternalProcessEditPaste();
 		return true;
@@ -1597,8 +1597,8 @@ public:
 
 protected:
 	// Begin FEditorUndoClient
-	virtual void PostUndo(bool bSuccess) OVERRIDE { OnUndo(); }
-	virtual void PostRedo(bool bSuccess) OVERRIDE { PostUndo(bSuccess); }
+	virtual void PostUndo(bool bSuccess) override { OnUndo(); }
+	virtual void PostRedo(bool bSuccess) override { PostUndo(bSuccess); }
 	// End of FEditorUndoClient
 
 protected:
diff --git a/Engine/Source/Editor/LandscapeEditor/Private/LandscapeEdModeTools.h b/Engine/Source/Editor/LandscapeEditor/Private/LandscapeEdModeTools.h
--- a/Engine/Source/Editor/LandscapeEditor/Private/LandscapeEdModeTools.h
+++ b/Engine/Source/Editor/LandscapeEditor/Private/LandscapeEdModeTools.h
@@ -1157,7 +1157,7 @@ public:
 	,	ToolStroke(NULL)
 	{}
 
-	virtual bool BeginTool( FLevelEditorViewportClient* ViewportClient, const FLandscapeToolTarget& InTarget, const FVector& InHitLocation ) OVERRIDE
+	virtual bool BeginTool( FLevelEditorViewportClient* ViewportClient, const FLandscapeToolTarget& InTarget, const FVector& InHitLocation ) override
 	{
 		if( !ensure(MousePositions.Num()==0) ) 
 		{
@@ -1175,7 +1175,7 @@ public:
 		return true;
 	}
 
-	virtual void Tick(FLevelEditorViewportClient* ViewportClient,float DeltaTime) OVERRIDE 
+	virtual void Tick(FLevelEditorViewportClient* ViewportClient,float DeltaTime) override 
 	{
 		if( bToolActive && MousePositions.Num() )
 		{
@@ -1184,7 +1184,7 @@ public:
 		}
 	}
 
-	virtual void EndTool(FLevelEditorViewportClient* ViewportClient) OVERRIDE 
+	virtual void EndTool(FLevelEditorViewportClient* ViewportClient) override 
 	{
 		if( bToolActive && MousePositions.Num() )
 		{
@@ -1198,7 +1198,7 @@ public:
 		EdMode->CurrentBrush->EndStroke();
 	}
 
-	virtual bool MouseMove( FLevelEditorViewportClient* ViewportClient, FViewport* Viewport, int32 x, int32 y ) OVERRIDE 
+	virtual bool MouseMove( FLevelEditorViewportClient* ViewportClient, FViewport* Viewport, int32 x, int32 y ) override 
 	{
 		FVector HitLocation;
 		if( EdMode->LandscapeMouseTrace(ViewportClient, x, y, HitLocation)  )
diff --git a/Engine/Source/Editor/LandscapeEditor/Private/LandscapeEdModeXYOffset.cpp b/Engine/Source/Editor/LandscapeEditor/Private/LandscapeEdModeXYOffset.cpp
--- a/Engine/Source/Editor/LandscapeEditor/Private/LandscapeEdModeXYOffset.cpp
+++ b/Engine/Source/Editor/LandscapeEditor/Private/LandscapeEdModeXYOffset.cpp
@@ -553,8 +553,8 @@ public:
 		: FLandscapeToolBase<FLandscapeToolStrokeRetopologize>(InEdMode)
 	{}
 
-	virtual const TCHAR* GetToolName() OVERRIDE { return TEXT("Retopologize"); }
-	virtual FText GetDisplayName() OVERRIDE { return NSLOCTEXT("UnrealEd", "LandscapeMode_Retopologize", "Retopologize"); }
+	virtual const TCHAR* GetToolName() override { return TEXT("Retopologize"); }
+	virtual FText GetDisplayName() override { return NSLOCTEXT("UnrealEd", "LandscapeMode_Retopologize", "Retopologize"); }
 
 	virtual bool IsValidForTarget(const FLandscapeToolTarget& Target)
 	{
diff --git a/Engine/Source/Editor/LandscapeEditor/Private/LandscapeEditorDetailCustomization_AlphaBrush.cpp b/Engine/Source/Editor/LandscapeEditor/Private/LandscapeEditorDetailCustomization_AlphaBrush.cpp
--- a/Engine/Source/Editor/LandscapeEditor/Private/LandscapeEditorDetailCustomization_AlphaBrush.cpp
+++ b/Engine/Source/Editor/LandscapeEditor/Private/LandscapeEditorDetailCustomization_AlphaBrush.cpp
@@ -29,7 +29,7 @@ public:
 	}
 
 	/** FViewportClient interface */
-	virtual void Draw(FViewport* Viewport, FCanvas* Canvas) OVERRIDE;
+	virtual void Draw(FViewport* Viewport, FCanvas* Canvas) override;
 
 private:
 	TWeakPtr<STextureMaskThumbnail> Parent;
diff --git a/Engine/Source/Editor/LandscapeEditor/Private/LandscapeEditorDetailCustomization_AlphaBrush.h b/Engine/Source/Editor/LandscapeEditor/Private/LandscapeEditorDetailCustomization_AlphaBrush.h
--- a/Engine/Source/Editor/LandscapeEditor/Private/LandscapeEditorDetailCustomization_AlphaBrush.h
+++ b/Engine/Source/Editor/LandscapeEditor/Private/LandscapeEditorDetailCustomization_AlphaBrush.h
@@ -17,7 +17,7 @@ public:
 	static TSharedRef<IDetailCustomization> MakeInstance();
 
 	/** IDetailCustomization interface */
-	virtual void CustomizeDetails(IDetailLayoutBuilder& DetailBuilder) OVERRIDE;
+	virtual void CustomizeDetails(IDetailLayoutBuilder& DetailBuilder) override;
 
 protected:
 	static bool OnAssetDraggedOver(const UObject* InObject);
diff --git a/Engine/Source/Editor/LandscapeEditor/Private/LandscapeEditorDetailCustomization_Base.h b/Engine/Source/Editor/LandscapeEditor/Private/LandscapeEditorDetailCustomization_Base.h
--- a/Engine/Source/Editor/LandscapeEditor/Private/LandscapeEditorDetailCustomization_Base.h
+++ b/Engine/Source/Editor/LandscapeEditor/Private/LandscapeEditorDetailCustomization_Base.h
@@ -12,7 +12,7 @@ class FLandscapeEditorDetailCustomization_Base : public IDetailCustomization //,
 {
 public:
 	/** IDetailCustomization interface */
-	virtual void CustomizeDetails(IDetailLayoutBuilder& DetailBuilder) OVERRIDE = 0;
+	virtual void CustomizeDetails(IDetailLayoutBuilder& DetailBuilder) override = 0;
 
 protected:
 	static class FEdModeLandscape* GetEditorMode();
diff --git a/Engine/Source/Editor/LandscapeEditor/Private/LandscapeEditorDetailCustomization_CopyPaste.h b/Engine/Source/Editor/LandscapeEditor/Private/LandscapeEditorDetailCustomization_CopyPaste.h
--- a/Engine/Source/Editor/LandscapeEditor/Private/LandscapeEditorDetailCustomization_CopyPaste.h
+++ b/Engine/Source/Editor/LandscapeEditor/Private/LandscapeEditorDetailCustomization_CopyPaste.h
@@ -16,7 +16,7 @@ public:
 	static TSharedRef<IDetailCustomization> MakeInstance();
 
 	/** IDetailCustomization interface */
-	virtual void CustomizeDetails(IDetailLayoutBuilder& DetailBuilder) OVERRIDE;
+	virtual void CustomizeDetails(IDetailLayoutBuilder& DetailBuilder) override;
 
 public:
 	static FReply OnCopyToGizmoButtonClicked();
@@ -35,8 +35,8 @@ class FLandscapeEditorStructCustomization_FGizmoImportLayer : public FLandscapeE
 public:
 	static TSharedRef<IPropertyTypeCustomization> MakeInstance();
 
-	virtual void CustomizeHeader(TSharedRef<IPropertyHandle> StructPropertyHandle, class FDetailWidgetRow& HeaderRow, IPropertyTypeCustomizationUtils& StructCustomizationUtils) OVERRIDE;
-	virtual void CustomizeChildren(TSharedRef<IPropertyHandle> StructPropertyHandle, class IDetailChildrenBuilder& ChildBuilder, IPropertyTypeCustomizationUtils& StructCustomizationUtils) OVERRIDE;
+	virtual void CustomizeHeader(TSharedRef<IPropertyHandle> StructPropertyHandle, class FDetailWidgetRow& HeaderRow, IPropertyTypeCustomizationUtils& StructCustomizationUtils) override;
+	virtual void CustomizeChildren(TSharedRef<IPropertyHandle> StructPropertyHandle, class IDetailChildrenBuilder& ChildBuilder, IPropertyTypeCustomizationUtils& StructCustomizationUtils) override;
 
 public:
 	static FReply OnGizmoImportLayerFilenameButtonClicked(TSharedRef<IPropertyHandle> PropertyHandle_LayerFilename);
diff --git a/Engine/Source/Editor/LandscapeEditor/Private/LandscapeEditorDetailCustomization_NewLandscape.h b/Engine/Source/Editor/LandscapeEditor/Private/LandscapeEditorDetailCustomization_NewLandscape.h
--- a/Engine/Source/Editor/LandscapeEditor/Private/LandscapeEditorDetailCustomization_NewLandscape.h
+++ b/Engine/Source/Editor/LandscapeEditor/Private/LandscapeEditorDetailCustomization_NewLandscape.h
@@ -22,7 +22,7 @@ public:
 	static TSharedRef<IDetailCustomization> MakeInstance();
 
 	/** IDetailCustomization interface */
-	virtual void CustomizeDetails(IDetailLayoutBuilder& DetailBuilder) OVERRIDE;
+	virtual void CustomizeDetails(IDetailLayoutBuilder& DetailBuilder) override;
 
 public:
 	static void SetScale(float NewValue, ETextCommit::Type, TSharedRef<IPropertyHandle> PropertyHandle);
@@ -86,8 +86,8 @@ class FLandscapeEditorStructCustomization_FLandscapeImportLayer : public FLandsc
 public:
 	static TSharedRef<IPropertyTypeCustomization> MakeInstance();
 
-	virtual void CustomizeHeader(TSharedRef<IPropertyHandle> StructPropertyHandle, class FDetailWidgetRow& HeaderRow, IPropertyTypeCustomizationUtils& StructCustomizationUtils) OVERRIDE;
-	virtual void CustomizeChildren(TSharedRef<IPropertyHandle> StructPropertyHandle, class IDetailChildrenBuilder& ChildBuilder, IPropertyTypeCustomizationUtils& StructCustomizationUtils) OVERRIDE;
+	virtual void CustomizeHeader(TSharedRef<IPropertyHandle> StructPropertyHandle, class FDetailWidgetRow& HeaderRow, IPropertyTypeCustomizationUtils& StructCustomizationUtils) override;
+	virtual void CustomizeChildren(TSharedRef<IPropertyHandle> StructPropertyHandle, class IDetailChildrenBuilder& ChildBuilder, IPropertyTypeCustomizationUtils& StructCustomizationUtils) override;
 
 public:
 	static FReply OnLayerFilenameButtonClicked(TSharedRef<IPropertyHandle> PropertyHandle_LayerFilename);
diff --git a/Engine/Source/Editor/LandscapeEditor/Private/LandscapeEditorDetailCustomization_TargetLayers.h b/Engine/Source/Editor/LandscapeEditor/Private/LandscapeEditorDetailCustomization_TargetLayers.h
--- a/Engine/Source/Editor/LandscapeEditor/Private/LandscapeEditorDetailCustomization_TargetLayers.h
+++ b/Engine/Source/Editor/LandscapeEditor/Private/LandscapeEditorDetailCustomization_TargetLayers.h
@@ -16,7 +16,7 @@ public:
 	static TSharedRef<IDetailCustomization> MakeInstance();
 
 	/** IDetailCustomization interface */
-	virtual void CustomizeDetails(IDetailLayoutBuilder& DetailBuilder) OVERRIDE;
+	virtual void CustomizeDetails(IDetailLayoutBuilder& DetailBuilder) override;
 
 protected:
 	static bool ShouldShowTargetLayers();
@@ -28,13 +28,13 @@ public:
 	FLandscapeEditorCustomNodeBuilder_TargetLayers(TSharedRef<FAssetThumbnailPool> ThumbnailPool);
 	~FLandscapeEditorCustomNodeBuilder_TargetLayers();
 
-	virtual void SetOnRebuildChildren( FSimpleDelegate InOnRegenerateChildren ) OVERRIDE;
-	virtual void GenerateHeaderRowContent( FDetailWidgetRow& NodeRow ) OVERRIDE;
-	virtual void GenerateChildContent( IDetailChildrenBuilder& ChildrenBuilder ) OVERRIDE;
-	virtual void Tick( float DeltaTime ) OVERRIDE {}
-	virtual bool RequiresTick() const OVERRIDE { return false; }
-	virtual bool InitiallyCollapsed() const OVERRIDE { return false; }
-	virtual FName GetName() const OVERRIDE { return "TargetLayers"; }
+	virtual void SetOnRebuildChildren( FSimpleDelegate InOnRegenerateChildren ) override;
+	virtual void GenerateHeaderRowContent( FDetailWidgetRow& NodeRow ) override;
+	virtual void GenerateChildContent( IDetailChildrenBuilder& ChildrenBuilder ) override;
+	virtual void Tick( float DeltaTime ) override {}
+	virtual bool RequiresTick() const override { return false; }
+	virtual bool InitiallyCollapsed() const override { return false; }
+	virtual FName GetName() const override { return "TargetLayers"; }
 
 protected:
 	TSharedRef<FAssetThumbnailPool> ThumbnailPool;
diff --git a/Engine/Source/Editor/LandscapeEditor/Private/LandscapeEditorDetailWidgets.cpp b/Engine/Source/Editor/LandscapeEditor/Private/LandscapeEditorDetailWidgets.cpp
--- a/Engine/Source/Editor/LandscapeEditor/Private/LandscapeEditorDetailWidgets.cpp
+++ b/Engine/Source/Editor/LandscapeEditor/Private/LandscapeEditorDetailWidgets.cpp
@@ -11,7 +11,7 @@ public:
 	FToolSelector(const FUIAction& InAction, const FOnGetContent& InMenuContentGenerator, const TAttribute<FText>& InLabel = TAttribute<FText>(), const TAttribute<FText>& InSmallText = TAttribute<FText>(), const TAttribute<FText>& InToolTip = TAttribute<FText>(), const TAttribute<FSlateIcon>& InIcon = TAttribute<FSlateIcon>());
 
 protected:
-	virtual void CreateMenuEntry(class FMenuBuilder& MenuBuilder) const OVERRIDE;
+	virtual void CreateMenuEntry(class FMenuBuilder& MenuBuilder) const override;
 
 	virtual TSharedRef<class IMultiBlockBaseWidget> ConstructWidget() const;
 
@@ -37,7 +37,7 @@ public:
 
 		void Construct(const FArguments& InArgs);
 
-	virtual void BuildMultiBlockWidget(const ISlateStyle* StyleSet, const FName& StyleName) OVERRIDE;
+	virtual void BuildMultiBlockWidget(const ISlateStyle* StyleSet, const FName& StyleName) override;
 
 protected:
 	TSharedRef<SWidget> OnGetMenuContent();
diff --git a/Engine/Source/Editor/LandscapeEditor/Private/LandscapeSplineImportExport.h b/Engine/Source/Editor/LandscapeEditor/Private/LandscapeSplineImportExport.h
--- a/Engine/Source/Editor/LandscapeEditor/Private/LandscapeSplineImportExport.h
+++ b/Engine/Source/Editor/LandscapeEditor/Private/LandscapeSplineImportExport.h
@@ -14,6 +14,6 @@ public:
 protected:
 	TArray<UObject*> OutObjects;
 
-	virtual void ProcessConstructedObject(UObject* CreatedObject) OVERRIDE;
-	virtual bool CanCreateClass(UClass* ObjectClass) const OVERRIDE;
+	virtual void ProcessConstructedObject(UObject* CreatedObject) override;
+	virtual bool CanCreateClass(UClass* ObjectClass) const override;
 };
diff --git a/Engine/Source/Editor/LandscapeEditor/Private/SLandscapeEditor.h b/Engine/Source/Editor/LandscapeEditor/Private/SLandscapeEditor.h
--- a/Engine/Source/Editor/LandscapeEditor/Private/SLandscapeEditor.h
+++ b/Engine/Source/Editor/LandscapeEditor/Private/SLandscapeEditor.h
@@ -33,17 +33,17 @@ private:
 class FLandscapeToolKit : public FModeToolkit
 {
 public:
-	virtual void RegisterTabSpawners(const TSharedRef<class FTabManager>& TabManager) OVERRIDE;
-	virtual void UnregisterTabSpawners(const TSharedRef<class FTabManager>& TabManager) OVERRIDE;
+	virtual void RegisterTabSpawners(const TSharedRef<class FTabManager>& TabManager) override;
+	virtual void UnregisterTabSpawners(const TSharedRef<class FTabManager>& TabManager) override;
 
 	/** Initializes the geometry mode toolkit */
 	virtual void Init(const TSharedPtr< class IToolkitHost >& InitToolkitHost);
 
 	/** IToolkit interface */
-	virtual FName GetToolkitFName() const OVERRIDE;
-	virtual FText GetBaseToolkitName() const OVERRIDE;
-	virtual class FEdModeLandscape* GetEditorMode() const OVERRIDE;
-	virtual TSharedPtr<class SWidget> GetInlineContent() const OVERRIDE;
+	virtual FName GetToolkitFName() const override;
+	virtual FText GetBaseToolkitName() const override;
+	virtual class FEdModeLandscape* GetEditorMode() const override;
+	virtual TSharedPtr<class SWidget> GetInlineContent() const override;
 
 	void NotifyToolChanged();
 	void NotifyBrushChanged();
diff --git a/Engine/Source/Editor/LandscapeEditor/Public/LandscapeEdMode.h b/Engine/Source/Editor/LandscapeEditor/Public/LandscapeEdMode.h
--- a/Engine/Source/Editor/LandscapeEditor/Public/LandscapeEdMode.h
+++ b/Engine/Source/Editor/LandscapeEditor/Public/LandscapeEdMode.h
@@ -53,7 +53,7 @@ public:
 	virtual EBrushType GetBrushType() { return BT_Normal; }
 	
 	// FGCObject interface
-	virtual void AddReferencedObjects(FReferenceCollector& Collector) OVERRIDE {}
+	virtual void AddReferencedObjects(FReferenceCollector& Collector) override {}
 };
 
 struct FLandscapeBrushSet
@@ -522,18 +522,18 @@ public:
 	virtual ~FEdModeLandscape();
 
 	/** FGCObject interface */
-	virtual void AddReferencedObjects( FReferenceCollector& Collector ) OVERRIDE;
+	virtual void AddReferencedObjects( FReferenceCollector& Collector ) override;
 
-	virtual bool UsesToolkits() const OVERRIDE;
+	virtual bool UsesToolkits() const override;
 
 	/** FEdMode: Called when the mode is entered */
-	virtual void Enter() OVERRIDE;
+	virtual void Enter() override;
 
 	/** FEdMode: Called when the mode is exited */
-	virtual void Exit() OVERRIDE;
+	virtual void Exit() override;
 
 	/** FEdMode: Called when the mouse is moved over the viewport */
-	virtual bool MouseMove( FLevelEditorViewportClient* ViewportClient, FViewport* Viewport, int32 x, int32 y ) OVERRIDE;
+	virtual bool MouseMove( FLevelEditorViewportClient* ViewportClient, FViewport* Viewport, int32 x, int32 y ) override;
 
 	/**
 	 * FEdMode: Called when the mouse is moved while a window input capture is in effect
@@ -545,71 +545,71 @@ public:
 	 *
 	 * @return	true if input was handled
 	 */
-	virtual bool CapturedMouseMove( FLevelEditorViewportClient* InViewportClient, FViewport* InViewport, int32 InMouseX, int32 InMouseY ) OVERRIDE;
+	virtual bool CapturedMouseMove( FLevelEditorViewportClient* InViewportClient, FViewport* InViewport, int32 InMouseX, int32 InMouseY ) override;
 
 	/** FEdMode: Called when a mouse button is pressed */
-	virtual bool StartTracking(FLevelEditorViewportClient* InViewportClient, FViewport* InViewport) OVERRIDE;
+	virtual bool StartTracking(FLevelEditorViewportClient* InViewportClient, FViewport* InViewport) override;
 
 	/** FEdMode: Called when a mouse button is released */
-	virtual bool EndTracking(FLevelEditorViewportClient* InViewportClient, FViewport* InViewport) OVERRIDE;
+	virtual bool EndTracking(FLevelEditorViewportClient* InViewportClient, FViewport* InViewport) override;
 
 	/** FEdMode: Allow us to disable mouse delta tracking during painting */
-	virtual bool DisallowMouseDeltaTracking() const OVERRIDE;
+	virtual bool DisallowMouseDeltaTracking() const override;
 
 	/** FEdMode: Called once per frame */
-	virtual void Tick(FLevelEditorViewportClient* ViewportClient,float DeltaTime) OVERRIDE;
+	virtual void Tick(FLevelEditorViewportClient* ViewportClient,float DeltaTime) override;
 
 	/** FEdMode: Called when clicking on a hit proxy */
-	virtual bool HandleClick(FLevelEditorViewportClient* InViewportClient, HHitProxy* HitProxy, const FViewportClick& Click) OVERRIDE;
+	virtual bool HandleClick(FLevelEditorViewportClient* InViewportClient, HHitProxy* HitProxy, const FViewportClick& Click) override;
 
 	/** FEdMode: Called when a key is pressed */
-	virtual bool InputKey( FLevelEditorViewportClient* InViewportClient, FViewport* InViewport, FKey InKey, EInputEvent InEvent ) OVERRIDE;
+	virtual bool InputKey( FLevelEditorViewportClient* InViewportClient, FViewport* InViewport, FKey InKey, EInputEvent InEvent ) override;
 
 	/** FEdMode: Called when mouse drag input is applied */
-	virtual bool InputDelta( FLevelEditorViewportClient* InViewportClient, FViewport* InViewport, FVector& InDrag, FRotator& InRot, FVector& InScale ) OVERRIDE;
+	virtual bool InputDelta( FLevelEditorViewportClient* InViewportClient, FViewport* InViewport, FVector& InDrag, FRotator& InRot, FVector& InScale ) override;
 
 	/** FEdMode: Render elements for the landscape tool */
-	virtual void Render( const FSceneView* View, FViewport* Viewport, FPrimitiveDrawInterface* PDI ) OVERRIDE;
+	virtual void Render( const FSceneView* View, FViewport* Viewport, FPrimitiveDrawInterface* PDI ) override;
 
 	/** FEdMode: Render HUD elements for this tool */
-	virtual void DrawHUD( FLevelEditorViewportClient* ViewportClient, FViewport* Viewport, const FSceneView* View, FCanvas* Canvas ) OVERRIDE;
+	virtual void DrawHUD( FLevelEditorViewportClient* ViewportClient, FViewport* Viewport, const FSceneView* View, FCanvas* Canvas ) override;
 
 	/** FEdMode: Handling SelectActor */
-	virtual bool Select( AActor* InActor, bool bInSelected ) OVERRIDE;
+	virtual bool Select( AActor* InActor, bool bInSelected ) override;
 
 	/** FEdMode: Check to see if an actor can be selected in this mode - no side effects */
-	virtual bool IsSelectionAllowed( AActor* InActor, bool bInSelection ) const OVERRIDE;
+	virtual bool IsSelectionAllowed( AActor* InActor, bool bInSelection ) const override;
 
 	/** FEdMode: Called when the currently selected actor has changed */
-	virtual void ActorSelectionChangeNotify() OVERRIDE;
+	virtual void ActorSelectionChangeNotify() override;
 
-	virtual void ActorMoveNotify() OVERRIDE;
+	virtual void ActorMoveNotify() override;
 
-	virtual EEditAction::Type GetActionEditDuplicate() OVERRIDE;
-	virtual EEditAction::Type GetActionEditDelete() OVERRIDE;
-	virtual EEditAction::Type GetActionEditCut() OVERRIDE;
-	virtual EEditAction::Type GetActionEditCopy() OVERRIDE;
-	virtual EEditAction::Type GetActionEditPaste() OVERRIDE;
-	virtual bool ProcessEditDuplicate() OVERRIDE;
-	virtual bool ProcessEditDelete() OVERRIDE;
-	virtual bool ProcessEditCut() OVERRIDE;
-	virtual bool ProcessEditCopy() OVERRIDE;
-	virtual bool ProcessEditPaste() OVERRIDE;
+	virtual EEditAction::Type GetActionEditDuplicate() override;
+	virtual EEditAction::Type GetActionEditDelete() override;
+	virtual EEditAction::Type GetActionEditCut() override;
+	virtual EEditAction::Type GetActionEditCopy() override;
+	virtual EEditAction::Type GetActionEditPaste() override;
+	virtual bool ProcessEditDuplicate() override;
+	virtual bool ProcessEditDelete() override;
+	virtual bool ProcessEditCut() override;
+	virtual bool ProcessEditCopy() override;
+	virtual bool ProcessEditPaste() override;
 
 	/** FEdMode: If the EdMode is handling InputDelta (ie returning true from it), this allows a mode to indicated whether or not the Widget should also move. */
-	virtual bool AllowWidgetMove() OVERRIDE { return true; }
+	virtual bool AllowWidgetMove() override { return true; }
 
 	/** FEdMode: Draw the transform widget while in this mode? */
-	virtual bool ShouldDrawWidget() const OVERRIDE;
+	virtual bool ShouldDrawWidget() const override;
 
 	/** FEdMode: Returns true if this mode uses the transform widget */
-	virtual bool UsesTransformWidget() const OVERRIDE;
+	virtual bool UsesTransformWidget() const override;
 
-	virtual EAxisList::Type GetWidgetAxisToDraw( FWidget::EWidgetMode InWidgetMode ) const OVERRIDE;
+	virtual EAxisList::Type GetWidgetAxisToDraw( FWidget::EWidgetMode InWidgetMode ) const override;
 
-	virtual FVector GetWidgetLocation() const OVERRIDE;
-	virtual bool GetCustomDrawingCoordinateSystem( FMatrix& InMatrix, void* InData ) OVERRIDE;
-	virtual bool GetCustomInputCoordinateSystem( FMatrix& InMatrix, void* InData ) OVERRIDE;
+	virtual FVector GetWidgetLocation() const override;
+	virtual bool GetCustomDrawingCoordinateSystem( FMatrix& InMatrix, void* InData ) override;
+	virtual bool GetCustomInputCoordinateSystem( FMatrix& InMatrix, void* InData ) override;
 
 	/** Forces real-time perspective viewports */
 	void ForceRealTimeViewports( const bool bEnable, const bool bStoreCurrentState );
diff --git a/Engine/Source/Editor/Layers/Private/ActorsAssignedToSpecificLayersFilter.h b/Engine/Source/Editor/Layers/Private/ActorsAssignedToSpecificLayersFilter.h
--- a/Engine/Source/Editor/Layers/Private/ActorsAssignedToSpecificLayersFilter.h
+++ b/Engine/Source/Editor/Layers/Private/ActorsAssignedToSpecificLayersFilter.h
@@ -7,7 +7,7 @@ class FActorsAssignedToSpecificLayersFilter : public IFilter< const AActor* cons
 public:
 
 	DECLARE_DERIVED_EVENT(FActorsAssignedToSpecificLayersFilter, IFilter< const AActor* const >::FChangedEvent, FChangedEvent);
-	virtual FChangedEvent& OnChanged() OVERRIDE { return ChangedEvent; }
+	virtual FChangedEvent& OnChanged() override { return ChangedEvent; }
 
 	/** 
 	 * Returns whether the specified Item passes the Filter's text restrictions 
@@ -15,7 +15,7 @@ public:
 	 *	@param	InItem	The Item to check 
 	 *	@return			Whether the specified Item passed the filter
 	 */
-	virtual bool PassesFilter( const AActor* const InActor ) const OVERRIDE
+	virtual bool PassesFilter( const AActor* const InActor ) const override
 	{
 		if( LayerNames.Num() == 0 )
 		{
diff --git a/Engine/Source/Editor/Layers/Private/LayerCollectionViewModel.h b/Engine/Source/Editor/Layers/Private/LayerCollectionViewModel.h
--- a/Engine/Source/Editor/Layers/Private/LayerCollectionViewModel.h
+++ b/Engine/Source/Editor/Layers/Private/LayerCollectionViewModel.h
@@ -219,8 +219,8 @@ private:
 	bool RequestRenameLayer_CanExecute() const;
 
 	// Begin FEditorUndoClient
-	virtual void PostUndo(bool bSuccess) OVERRIDE { Refresh(); }
-	virtual void PostRedo(bool bSuccess) OVERRIDE { PostUndo(bSuccess); }
+	virtual void PostUndo(bool bSuccess) override { Refresh(); }
+	virtual void PostRedo(bool bSuccess) override { PostUndo(bSuccess); }
 	// End of FEditorUndoClient
 
 private:
diff --git a/Engine/Source/Editor/Layers/Private/LayerViewModel.h b/Engine/Source/Editor/Layers/Private/LayerViewModel.h
--- a/Engine/Source/Editor/Layers/Private/LayerViewModel.h
+++ b/Engine/Source/Editor/Layers/Private/LayerViewModel.h
@@ -197,8 +197,8 @@ private:
 	void RefreshActorStats();
 
 	// Begin FEditorUndoClient
-	virtual void PostUndo(bool bSuccess) OVERRIDE { Refresh(); }
-	virtual void PostRedo(bool bSuccess) OVERRIDE { PostUndo(bSuccess); }
+	virtual void PostUndo(bool bSuccess) override { Refresh(); }
+	virtual void PostRedo(bool bSuccess) override { PostUndo(bSuccess); }
 	// End of FEditorUndoClient
 
 private:
diff --git a/Engine/Source/Editor/Layers/Private/SLayerBrowser.h b/Engine/Source/Editor/Layers/Private/SLayerBrowser.h
--- a/Engine/Source/Editor/Layers/Private/SLayerBrowser.h
+++ b/Engine/Source/Editor/Layers/Private/SLayerBrowser.h
@@ -206,7 +206,7 @@ protected:
 	 *
 	 * @param DragDropEvent   The drag and drop event.
 	 */
-	virtual void OnDragLeave( const FDragDropEvent& DragDropEvent ) OVERRIDE
+	virtual void OnDragLeave( const FDragDropEvent& DragDropEvent ) override
 	{
 		TSharedPtr< FActorDragDropGraphEdOp > DragActorOp = DragDropEvent.GetOperationAs< FActorDragDropGraphEdOp >();
 		if(DragActorOp.IsValid())
@@ -223,7 +223,7 @@ protected:
 	 *
 	 * @return A reply that indicated whether this event was handled.
 	 */
-	virtual FReply OnDragOver( const FGeometry& MyGeometry, const FDragDropEvent& DragDropEvent ) OVERRIDE
+	virtual FReply OnDragOver( const FGeometry& MyGeometry, const FDragDropEvent& DragDropEvent ) override
 	{
 		TSharedPtr< FActorDragDropGraphEdOp > DragActorOp = DragDropEvent.GetOperationAs< FActorDragDropGraphEdOp >();
 		if (!SelectedLayerViewModel->GetDataSource().IsValid() || !DragActorOp.IsValid())
@@ -282,7 +282,7 @@ protected:
 	 *
 	 * @return A reply that indicated whether this event was handled.
 	 */
-	virtual FReply OnDrop( const FGeometry& MyGeometry, const FDragDropEvent& DragDropEvent ) OVERRIDE
+	virtual FReply OnDrop( const FGeometry& MyGeometry, const FDragDropEvent& DragDropEvent ) override
 	{
 		TSharedPtr< FActorDragDropGraphEdOp > DragActorOp = DragDropEvent.GetOperationAs< FActorDragDropGraphEdOp >();
 
diff --git a/Engine/Source/Editor/Layers/Private/SLayersView.h b/Engine/Source/Editor/Layers/Private/SLayersView.h
--- a/Engine/Source/Editor/Layers/Private/SLayersView.h
+++ b/Engine/Source/Editor/Layers/Private/SLayersView.h
@@ -126,7 +126,7 @@ protected:
 	 *
 	 * @return  True if this widget can take keyboard focus
 	 */
-	virtual bool SupportsKeyboardFocus() const OVERRIDE
+	virtual bool SupportsKeyboardFocus() const override
 	{
 		return true;
 	}
@@ -139,7 +139,7 @@ protected:
 	 *
 	 * @return  Returns whether the event was handled, along with other possible actions
 	 */
-	virtual FReply OnKeyDown( const FGeometry& MyGeometry, const FKeyboardEvent& InKeyboardEvent ) OVERRIDE
+	virtual FReply OnKeyDown( const FGeometry& MyGeometry, const FKeyboardEvent& InKeyboardEvent ) override
 	{
 		return ViewModel->GetCommandList()->ProcessCommandBindings( InKeyboardEvent ) ? FReply::Handled() : FReply::Unhandled();
 	}
@@ -149,7 +149,7 @@ protected:
 	 *
 	 * @param DragDropEvent   The drag and drop event.
 	 */
-	virtual void OnDragLeave( const FDragDropEvent& DragDropEvent ) OVERRIDE
+	virtual void OnDragLeave( const FDragDropEvent& DragDropEvent ) override
 	{
 		TSharedPtr< FActorDragDropGraphEdOp > DragActorOp = DragDropEvent.GetOperationAs< FActorDragDropGraphEdOp >();
 		if (DragActorOp.IsValid())
@@ -166,7 +166,7 @@ protected:
 	 *
 	 * @return A reply that indicated whether this event was handled.
 	 */
-	virtual FReply OnDragOver( const FGeometry& MyGeometry, const FDragDropEvent& DragDropEvent ) OVERRIDE
+	virtual FReply OnDragOver( const FGeometry& MyGeometry, const FDragDropEvent& DragDropEvent ) override
 	{
 		TSharedPtr< FActorDragDropGraphEdOp > DragActorOp = DragDropEvent.GetOperationAs< FActorDragDropGraphEdOp >();
 		if (DragActorOp.IsValid())
@@ -186,7 +186,7 @@ protected:
 	 *
 	 * @return A reply that indicated whether this event was handled.
 	 */
-	virtual FReply OnDrop( const FGeometry& MyGeometry, const FDragDropEvent& DragDropEvent ) OVERRIDE
+	virtual FReply OnDrop( const FGeometry& MyGeometry, const FDragDropEvent& DragDropEvent ) override
 	{
 		TSharedPtr< FActorDragDropGraphEdOp > DragActorOp = DragDropEvent.GetOperationAs< FActorDragDropGraphEdOp >();
 		if ( !DragActorOp.IsValid() )	
diff --git a/Engine/Source/Editor/Layers/Private/SLayersViewRow.h b/Engine/Source/Editor/Layers/Private/SLayersViewRow.h
--- a/Engine/Source/Editor/Layers/Private/SLayersViewRow.h
+++ b/Engine/Source/Editor/Layers/Private/SLayersViewRow.h
@@ -55,7 +55,7 @@ protected:
 	 *
 	 * @return a widget to represent the contents of a cell in this row of a TableView. 
 	 */
-	virtual TSharedRef< SWidget > GenerateWidgetForColumn( const FName & ColumnID ) OVERRIDE
+	virtual TSharedRef< SWidget > GenerateWidgetForColumn( const FName & ColumnID ) override
 	{
 		TSharedPtr< SWidget > TableRowContent;
 
@@ -143,7 +143,7 @@ protected:
 	 *
 	 * @param DragDropEvent   The drag and drop event.
 	 */
-	virtual void OnDragLeave( const FDragDropEvent& DragDropEvent ) OVERRIDE
+	virtual void OnDragLeave( const FDragDropEvent& DragDropEvent ) override
 	{
 		TSharedPtr< FActorDragDropGraphEdOp > DragActorOp = DragDropEvent.GetOperationAs< FActorDragDropGraphEdOp >();
 		if (DragActorOp.IsValid())
@@ -160,7 +160,7 @@ protected:
 	 *
 	 * @return A reply that indicated whether this event was handled.
 	 */
-	virtual FReply OnDragOver( const FGeometry& MyGeometry, const FDragDropEvent& DragDropEvent ) OVERRIDE
+	virtual FReply OnDragOver( const FGeometry& MyGeometry, const FDragDropEvent& DragDropEvent ) override
 	{
 		TSharedPtr< FActorDragDropGraphEdOp > DragActorOp = DragDropEvent.GetOperationAs< FActorDragDropGraphEdOp >();
 		if (!DragActorOp.IsValid())
@@ -199,7 +199,7 @@ protected:
 	 *
 	 * @return A reply that indicated whether this event was handled.
 	 */
-	virtual FReply OnDrop( const FGeometry& MyGeometry, const FDragDropEvent& DragDropEvent ) OVERRIDE
+	virtual FReply OnDrop( const FGeometry& MyGeometry, const FDragDropEvent& DragDropEvent ) override
 	{
 		TSharedPtr< FActorDragDropGraphEdOp > DragActorOp = DragDropEvent.GetOperationAs< FActorDragDropGraphEdOp >();
 		if (!DragActorOp.IsValid())
diff --git a/Engine/Source/Editor/Layers/Private/SceneOutlinerLayerContentsColumn.h b/Engine/Source/Editor/Layers/Private/SceneOutlinerLayerContentsColumn.h
--- a/Engine/Source/Editor/Layers/Private/SceneOutlinerLayerContentsColumn.h
+++ b/Engine/Source/Editor/Layers/Private/SceneOutlinerLayerContentsColumn.h
@@ -21,19 +21,19 @@ public:
 	//////////////////////////////////////////////////////////////////////////
 	// Begin ISceneOutlinerColumn Implementation
 
-	virtual FName GetColumnID() OVERRIDE;
+	virtual FName GetColumnID() override;
 
-	virtual SHeaderRow::FColumn::FArguments ConstructHeaderRowColumn() OVERRIDE;
+	virtual SHeaderRow::FColumn::FArguments ConstructHeaderRowColumn() override;
 
-	virtual const TSharedRef< SWidget > ConstructRowWidget( const TSharedRef<SceneOutliner::TOutlinerTreeItem> TreeItem ) OVERRIDE;
+	virtual const TSharedRef< SWidget > ConstructRowWidget( const TSharedRef<SceneOutliner::TOutlinerTreeItem> TreeItem ) override;
 
 	virtual bool ProvidesSearchStrings();
 
-	virtual void PopulateActorSearchStrings( const AActor* const InActor, OUT TArray< FString >& OutSearchStrings ) const OVERRIDE;
+	virtual void PopulateActorSearchStrings( const AActor* const InActor, OUT TArray< FString >& OutSearchStrings ) const override;
 	
-	virtual bool SupportsSorting() const OVERRIDE { return false; }
+	virtual bool SupportsSorting() const override { return false; }
 
-	virtual void SortItems(TArray<TSharedPtr<SceneOutliner::TOutlinerTreeItem>>& RootItems, const EColumnSortMode::Type SortMode) const OVERRIDE {}
+	virtual void SortItems(TArray<TSharedPtr<SceneOutliner::TOutlinerTreeItem>>& RootItems, const EColumnSortMode::Type SortMode) const override {}
 
 	// End ISceneOutlinerColumn Implementation
 	//////////////////////////////////////////////////////////////////////////
diff --git a/Engine/Source/Editor/LevelEditor/Private/LevelEditorGenericDetails.h b/Engine/Source/Editor/LevelEditor/Private/LevelEditorGenericDetails.h
--- a/Engine/Source/Editor/LevelEditor/Private/LevelEditorGenericDetails.h
+++ b/Engine/Source/Editor/LevelEditor/Private/LevelEditorGenericDetails.h
@@ -12,7 +12,7 @@ public:
 	static TSharedRef<IDetailCustomization> MakeInstance();
 
 	/** IDetailCustomization interface */
-	virtual void CustomizeDetails( class IDetailLayoutBuilder& DetailLayout ) OVERRIDE;
+	virtual void CustomizeDetails( class IDetailLayoutBuilder& DetailLayout ) override;
 
 private:
 	/**
diff --git a/Engine/Source/Editor/LevelEditor/Private/LevelViewportLayout2x2.h b/Engine/Source/Editor/LevelEditor/Private/LevelViewportLayout2x2.h
--- a/Engine/Source/Editor/LevelEditor/Private/LevelViewportLayout2x2.h
+++ b/Engine/Source/Editor/LevelEditor/Private/LevelViewportLayout2x2.h
@@ -14,17 +14,17 @@ public:
 	/**
 	 * Saves viewport layout information between editor sessions
 	 */
-	virtual void SaveLayoutString(const FString& LayoutString) const OVERRIDE;
+	virtual void SaveLayoutString(const FString& LayoutString) const override;
 
-	virtual const FName& GetLayoutTypeName() const OVERRIDE { return LevelViewportConfigurationNames::FourPanes2x2; }
+	virtual const FName& GetLayoutTypeName() const override { return LevelViewportConfigurationNames::FourPanes2x2; }
 protected:
 	/**
 	 * Creates the viewports and splitter for the 2x2 layout                   
 	 */
-	virtual TSharedRef<SWidget> MakeViewportLayout(const FString& LayoutString) OVERRIDE;
+	virtual TSharedRef<SWidget> MakeViewportLayout(const FString& LayoutString) override;
 
 	/** Overridden from FLevelViewportLayout */
-	virtual void ReplaceWidget( TSharedRef< SWidget > Source, TSharedRef< SWidget > Replacement ) OVERRIDE;
+	virtual void ReplaceWidget( TSharedRef< SWidget > Source, TSharedRef< SWidget > Replacement ) override;
 
 
 private:
diff --git a/Engine/Source/Editor/LevelEditor/Private/LevelViewportLayoutFourPanes.h b/Engine/Source/Editor/LevelEditor/Private/LevelViewportLayoutFourPanes.h
--- a/Engine/Source/Editor/LevelEditor/Private/LevelViewportLayoutFourPanes.h
+++ b/Engine/Source/Editor/LevelEditor/Private/LevelViewportLayoutFourPanes.h
@@ -14,12 +14,12 @@ public:
 	/**
 	 * Saves viewport layout information between editor sessions
 	 */
-	virtual void SaveLayoutString(const FString& LayoutString) const OVERRIDE;
+	virtual void SaveLayoutString(const FString& LayoutString) const override;
 protected:
 	/**
 	 * Creates the viewports and splitter for the four-pane layout              
 	 */
-	virtual TSharedRef<SWidget> MakeViewportLayout(const FString& LayoutString) OVERRIDE;
+	virtual TSharedRef<SWidget> MakeViewportLayout(const FString& LayoutString) override;
 
 	virtual TSharedRef<SWidget> MakeFourPanelWidget(
 		TArray<TSharedPtr< SLevelViewport >>& ViewportWidgets,
@@ -27,7 +27,7 @@ protected:
 		float PrimarySplitterPercentage, float SecondarySplitterPercentage0, float SecondarySplitterPercentage1) = 0;
 
 	/** Overridden from FLevelViewportLayout */
-	virtual void ReplaceWidget( TSharedRef< SWidget > Source, TSharedRef< SWidget > Replacement ) OVERRIDE;
+	virtual void ReplaceWidget( TSharedRef< SWidget > Source, TSharedRef< SWidget > Replacement ) override;
 
 
 protected:
@@ -42,7 +42,7 @@ protected:
 class FLevelViewportLayoutFourPanesLeft : public FLevelViewportLayoutFourPanes
 {
 public:
-	virtual const FName& GetLayoutTypeName() const OVERRIDE { return LevelViewportConfigurationNames::FourPanesLeft; }
+	virtual const FName& GetLayoutTypeName() const override { return LevelViewportConfigurationNames::FourPanesLeft; }
 
 	virtual TSharedRef<SWidget> MakeFourPanelWidget(
 		TArray<TSharedPtr< SLevelViewport >>& ViewportWidgets,
@@ -56,7 +56,7 @@ public:
 class FLevelViewportLayoutFourPanesRight : public FLevelViewportLayoutFourPanes
 {
 public:
-	virtual const FName& GetLayoutTypeName() const OVERRIDE { return LevelViewportConfigurationNames::FourPanesRight; }
+	virtual const FName& GetLayoutTypeName() const override { return LevelViewportConfigurationNames::FourPanesRight; }
 
 	virtual TSharedRef<SWidget> MakeFourPanelWidget(
 		TArray<TSharedPtr< SLevelViewport >>& ViewportWidgets,
@@ -70,7 +70,7 @@ public:
 class FLevelViewportLayoutFourPanesTop : public FLevelViewportLayoutFourPanes
 {
 public:
-	virtual const FName& GetLayoutTypeName() const OVERRIDE { return LevelViewportConfigurationNames::FourPanesTop; }
+	virtual const FName& GetLayoutTypeName() const override { return LevelViewportConfigurationNames::FourPanesTop; }
 
 	virtual TSharedRef<SWidget> MakeFourPanelWidget(
 		TArray<TSharedPtr< SLevelViewport >>& ViewportWidgets,
@@ -84,7 +84,7 @@ public:
 class FLevelViewportLayoutFourPanesBottom : public FLevelViewportLayoutFourPanes
 {
 public:
-	virtual const FName& GetLayoutTypeName() const OVERRIDE { return LevelViewportConfigurationNames::FourPanesBottom; }
+	virtual const FName& GetLayoutTypeName() const override { return LevelViewportConfigurationNames::FourPanesBottom; }
 
 	virtual TSharedRef<SWidget> MakeFourPanelWidget(
 		TArray<TSharedPtr< SLevelViewport >>& ViewportWidgets,
diff --git a/Engine/Source/Editor/LevelEditor/Private/LevelViewportLayoutOnePane.h b/Engine/Source/Editor/LevelEditor/Private/LevelViewportLayoutOnePane.h
--- a/Engine/Source/Editor/LevelEditor/Private/LevelViewportLayoutOnePane.h
+++ b/Engine/Source/Editor/LevelEditor/Private/LevelViewportLayoutOnePane.h
@@ -14,17 +14,17 @@ public:
 	/**
 	* Saves viewport layout information between editor sessions
 	*/
-	virtual void SaveLayoutString(const FString& LayoutString) const OVERRIDE;
+	virtual void SaveLayoutString(const FString& LayoutString) const override;
 
-	virtual const FName& GetLayoutTypeName() const OVERRIDE{ return LevelViewportConfigurationNames::OnePane; }
+	virtual const FName& GetLayoutTypeName() const override{ return LevelViewportConfigurationNames::OnePane; }
 protected:
 	/**
 	* Creates the viewport for the single pane
 	*/
-	virtual TSharedRef<SWidget> MakeViewportLayout(const FString& LayoutString) OVERRIDE;
+	virtual TSharedRef<SWidget> MakeViewportLayout(const FString& LayoutString) override;
 
 	/** Overridden from FLevelViewportLayout */
-	virtual void ReplaceWidget(TSharedRef< SWidget > Source, TSharedRef< SWidget > Replacement) OVERRIDE;
+	virtual void ReplaceWidget(TSharedRef< SWidget > Source, TSharedRef< SWidget > Replacement) override;
 
 protected:
 	/** The viewport widget parent box */
diff --git a/Engine/Source/Editor/LevelEditor/Private/LevelViewportLayoutThreePanes.h b/Engine/Source/Editor/LevelEditor/Private/LevelViewportLayoutThreePanes.h
--- a/Engine/Source/Editor/LevelEditor/Private/LevelViewportLayoutThreePanes.h
+++ b/Engine/Source/Editor/LevelEditor/Private/LevelViewportLayoutThreePanes.h
@@ -14,12 +14,12 @@ public:
 	/**
 	 * Saves viewport layout information between editor sessions
 	 */
-	virtual void SaveLayoutString(const FString& LayoutString) const OVERRIDE;
+	virtual void SaveLayoutString(const FString& LayoutString) const override;
 protected:
 	/**
 	 * Creates the viewports and splitter for the two panes vertical layout                   
 	 */
-	virtual TSharedRef<SWidget> MakeViewportLayout(const FString& LayoutString) OVERRIDE;
+	virtual TSharedRef<SWidget> MakeViewportLayout(const FString& LayoutString) override;
 
 	virtual TSharedRef<SWidget> MakeThreePanelWidget(
 		TArray<TSharedPtr< SLevelViewport >>& ViewportWidgets,
@@ -27,7 +27,7 @@ protected:
 		float PrimarySplitterPercentage, float SecondarySplitterPercentage) = 0;
 
 	/** Overridden from FLevelViewportLayout */
-	virtual void ReplaceWidget( TSharedRef< SWidget > Source, TSharedRef< SWidget > Replacement ) OVERRIDE;
+	virtual void ReplaceWidget( TSharedRef< SWidget > Source, TSharedRef< SWidget > Replacement ) override;
 
 
 protected:
@@ -41,7 +41,7 @@ protected:
 class FLevelViewportLayoutThreePanesLeft : public FLevelViewportLayoutThreePanes
 {
 public:
-	virtual const FName& GetLayoutTypeName() const OVERRIDE { return LevelViewportConfigurationNames::ThreePanesLeft; }
+	virtual const FName& GetLayoutTypeName() const override { return LevelViewportConfigurationNames::ThreePanesLeft; }
 
 	virtual TSharedRef<SWidget> MakeThreePanelWidget(
 		TArray<TSharedPtr< SLevelViewport >>& ViewportWidgets,
@@ -55,7 +55,7 @@ public:
 class FLevelViewportLayoutThreePanesRight : public FLevelViewportLayoutThreePanes
 {
 public:
-	virtual const FName& GetLayoutTypeName() const OVERRIDE { return LevelViewportConfigurationNames::ThreePanesRight; }
+	virtual const FName& GetLayoutTypeName() const override { return LevelViewportConfigurationNames::ThreePanesRight; }
 
 	virtual TSharedRef<SWidget> MakeThreePanelWidget(
 		TArray<TSharedPtr< SLevelViewport >>& ViewportWidgets,
@@ -69,7 +69,7 @@ public:
 class FLevelViewportLayoutThreePanesTop : public FLevelViewportLayoutThreePanes
 {
 public:
-	virtual const FName& GetLayoutTypeName() const OVERRIDE { return LevelViewportConfigurationNames::ThreePanesTop; }
+	virtual const FName& GetLayoutTypeName() const override { return LevelViewportConfigurationNames::ThreePanesTop; }
 
 	virtual TSharedRef<SWidget> MakeThreePanelWidget(
 		TArray<TSharedPtr< SLevelViewport >>& ViewportWidgets,
@@ -83,7 +83,7 @@ public:
 class FLevelViewportLayoutThreePanesBottom : public FLevelViewportLayoutThreePanes
 {
 public:
-	virtual const FName& GetLayoutTypeName() const OVERRIDE { return LevelViewportConfigurationNames::ThreePanesBottom; }
+	virtual const FName& GetLayoutTypeName() const override { return LevelViewportConfigurationNames::ThreePanesBottom; }
 
 	virtual TSharedRef<SWidget> MakeThreePanelWidget(
 		TArray<TSharedPtr< SLevelViewport >>& ViewportWidgets,
diff --git a/Engine/Source/Editor/LevelEditor/Private/LevelViewportLayoutTwoPanes.h b/Engine/Source/Editor/LevelEditor/Private/LevelViewportLayoutTwoPanes.h
--- a/Engine/Source/Editor/LevelEditor/Private/LevelViewportLayoutTwoPanes.h
+++ b/Engine/Source/Editor/LevelEditor/Private/LevelViewportLayoutTwoPanes.h
@@ -15,15 +15,15 @@ public:
 	/**
 	 * Saves viewport layout information between editor sessions
 	 */
-	virtual void SaveLayoutString(const FString& LayoutString) const OVERRIDE;
+	virtual void SaveLayoutString(const FString& LayoutString) const override;
 protected:
 	/**
 	 * Creates the viewports and splitter for the two panes vertical layout                   
 	 */
-	virtual TSharedRef<SWidget> MakeViewportLayout(const FString& LayoutString) OVERRIDE;
+	virtual TSharedRef<SWidget> MakeViewportLayout(const FString& LayoutString) override;
 
 	/** Overridden from FLevelViewportLayout */
-	virtual void ReplaceWidget( TSharedRef< SWidget > Source, TSharedRef< SWidget > Replacement ) OVERRIDE;
+	virtual void ReplaceWidget( TSharedRef< SWidget > Source, TSharedRef< SWidget > Replacement ) override;
 
 
 private:
@@ -37,7 +37,7 @@ private:
 class FLevelViewportLayoutTwoPanesVert : public TLevelViewportLayoutTwoPanes<EOrientation::Orient_Vertical>
 {
 public:
-	virtual const FName& GetLayoutTypeName() const OVERRIDE { return LevelViewportConfigurationNames::TwoPanesVert; }
+	virtual const FName& GetLayoutTypeName() const override { return LevelViewportConfigurationNames::TwoPanesVert; }
 };
 
 
@@ -46,7 +46,7 @@ public:
 class FLevelViewportLayoutTwoPanesHoriz : public TLevelViewportLayoutTwoPanes<EOrientation::Orient_Horizontal>
 {
 public:
-	virtual const FName& GetLayoutTypeName() const OVERRIDE { return LevelViewportConfigurationNames::TwoPanesHoriz; }
+	virtual const FName& GetLayoutTypeName() const override { return LevelViewportConfigurationNames::TwoPanesHoriz; }
 };
 
 #endif
diff --git a/Engine/Source/Editor/LevelEditor/Private/SCaptureRegionWidget.h b/Engine/Source/Editor/LevelEditor/Private/SCaptureRegionWidget.h
--- a/Engine/Source/Editor/LevelEditor/Private/SCaptureRegionWidget.h
+++ b/Engine/Source/Editor/LevelEditor/Private/SCaptureRegionWidget.h
@@ -30,9 +30,9 @@ public:
 	void Deactivate(bool bKeepChanges);
 	void Reset();
 
-	FReply OnMouseButtonDown( const FGeometry& MyGeometry, const FPointerEvent& MouseEvent ) OVERRIDE;
-	FReply OnMouseButtonUp( const FGeometry& MyGeometry, const FPointerEvent& MouseEvent ) OVERRIDE;
-	FReply OnMouseMove( const FGeometry& MyGeometry, const FPointerEvent& MouseEvent ) OVERRIDE;
+	FReply OnMouseButtonDown( const FGeometry& MyGeometry, const FPointerEvent& MouseEvent ) override;
+	FReply OnMouseButtonUp( const FGeometry& MyGeometry, const FPointerEvent& MouseEvent ) override;
+	FReply OnMouseMove( const FGeometry& MyGeometry, const FPointerEvent& MouseEvent ) override;
 
 private:
 
diff --git a/Engine/Source/Editor/LevelEditor/Private/SLevelEditor.h b/Engine/Source/Editor/LevelEditor/Private/SLevelEditor.h
--- a/Engine/Source/Editor/LevelEditor/Private/SLevelEditor.h
+++ b/Engine/Source/Editor/LevelEditor/Private/SLevelEditor.h
@@ -65,11 +65,11 @@ public:
 	TSharedPtr<class FLevelViewportTabContent> GetActiveViewportTab();
 
 	/** ILevelEditor interface */
-	virtual void SummonLevelViewportContextMenu() OVERRIDE;
-	virtual const TArray< TSharedPtr< class IToolkit > >& GetHostedToolkits() const OVERRIDE;
-	virtual TArray< TSharedPtr< ILevelViewport > > GetViewports() const OVERRIDE;
-	virtual TSharedPtr< class FAssetThumbnailPool > GetThumbnailPool() const OVERRIDE;
-	virtual void AppendCommands( const TSharedRef<FUICommandList>& InCommandsToAppend ) OVERRIDE;
+	virtual void SummonLevelViewportContextMenu() override;
+	virtual const TArray< TSharedPtr< class IToolkit > >& GetHostedToolkits() const override;
+	virtual TArray< TSharedPtr< ILevelViewport > > GetViewports() const override;
+	virtual TSharedPtr< class FAssetThumbnailPool > GetThumbnailPool() const override;
+	virtual void AppendCommands( const TSharedRef<FUICommandList>& InCommandsToAppend ) override;
 
 	/**
 	 * Given a tab ID, summons a new tab in the tab stack specified.
@@ -108,20 +108,20 @@ public:
 	}
 
 	/** IToolKitHost interface */
-	virtual TSharedRef< class SWidget > GetParentWidget() OVERRIDE;
-	virtual void BringToFront() OVERRIDE;
-	virtual TSharedRef< class SDockTabStack > GetTabSpot( const EToolkitTabSpot::Type TabSpot ) OVERRIDE;
-	virtual void OnToolkitHostingStarted( const TSharedRef< class IToolkit >& Toolkit ) OVERRIDE;
-	virtual void OnToolkitHostingFinished( const TSharedRef< class IToolkit >& Toolkit ) OVERRIDE;
-	virtual UWorld* GetWorld() const OVERRIDE;
+	virtual TSharedRef< class SWidget > GetParentWidget() override;
+	virtual void BringToFront() override;
+	virtual TSharedRef< class SDockTabStack > GetTabSpot( const EToolkitTabSpot::Type TabSpot ) override;
+	virtual void OnToolkitHostingStarted( const TSharedRef< class IToolkit >& Toolkit ) override;
+	virtual void OnToolkitHostingFinished( const TSharedRef< class IToolkit >& Toolkit ) override;
+	virtual UWorld* GetWorld() const override;
 	
 	/** SWidget overrides */
-	virtual bool SupportsKeyboardFocus() const OVERRIDE
+	virtual bool SupportsKeyboardFocus() const override
 	{
 		return true;
 	}
 
-	virtual void Tick( const FGeometry& AllottedGeometry, const double InCurrentTime, const float InDeltaTime ) OVERRIDE;
+	virtual void Tick( const FGeometry& AllottedGeometry, const double InCurrentTime, const float InDeltaTime ) override;
 
 	// Tab Management
 	TSharedRef<FTabManager> GetTabManager() const;
diff --git a/Engine/Source/Editor/LevelEditor/Private/SLevelEditorBuildAndSubmit.cpp b/Engine/Source/Editor/LevelEditor/Private/SLevelEditorBuildAndSubmit.cpp
--- a/Engine/Source/Editor/LevelEditor/Private/SLevelEditorBuildAndSubmit.cpp
+++ b/Engine/Source/Editor/LevelEditor/Private/SLevelEditorBuildAndSubmit.cpp
@@ -35,7 +35,7 @@ public:
 	 *
 	 * @param InArgs   A declaration from which to construct the widget
 	 */
-	virtual TSharedRef< SWidget > GenerateWidgetForColumn( const FName& ColumnName ) OVERRIDE
+	virtual TSharedRef< SWidget > GenerateWidgetForColumn( const FName& ColumnName ) override
 	{
 		if (NAME_NameColumn == ColumnName)
 		{
diff --git a/Engine/Source/Editor/LevelEditor/Private/SLevelViewport.cpp b/Engine/Source/Editor/LevelEditor/Private/SLevelViewport.cpp
--- a/Engine/Source/Editor/LevelEditor/Private/SLevelViewport.cpp
+++ b/Engine/Source/Editor/LevelEditor/Private/SLevelViewport.cpp
@@ -2063,12 +2063,12 @@ TSharedRef< ISceneOutlinerColumn > SLevelViewport::CreateActorLockSceneOutlinerC
 		//////////////////////////////////////////////////////////////////////////
 		// Begin ISceneOutlinerColumn Implementation
 
-		virtual FName GetColumnID() OVERRIDE
+		virtual FName GetColumnID() override
 		{
 			return FName( "LockedToViewport" );
 		}
 
-		virtual SHeaderRow::FColumn::FArguments ConstructHeaderRowColumn() OVERRIDE
+		virtual SHeaderRow::FColumn::FArguments ConstructHeaderRowColumn() override
 		{
 			return SHeaderRow::Column( GetColumnID() )
 				[
@@ -2076,7 +2076,7 @@ TSharedRef< ISceneOutlinerColumn > SLevelViewport::CreateActorLockSceneOutlinerC
 				];
 		}
 
-		virtual const TSharedRef< SWidget > ConstructRowWidget( const TSharedRef<SceneOutliner::TOutlinerTreeItem> TreeItem ) OVERRIDE
+		virtual const TSharedRef< SWidget > ConstructRowWidget( const TSharedRef<SceneOutliner::TOutlinerTreeItem> TreeItem ) override
 		{
 			if (TreeItem->Type == SceneOutliner::TOutlinerTreeItem::Actor)
 			{
@@ -2099,11 +2099,11 @@ TSharedRef< ISceneOutlinerColumn > SLevelViewport::CreateActorLockSceneOutlinerC
 
 		virtual bool ProvidesSearchStrings() { return false; }
 
-		virtual void PopulateActorSearchStrings( const AActor* const InActor, OUT TArray< FString >& OutSearchStrings ) const OVERRIDE {}
+		virtual void PopulateActorSearchStrings( const AActor* const InActor, OUT TArray< FString >& OutSearchStrings ) const override {}
 	
-		virtual bool SupportsSorting() const OVERRIDE { return false; }
+		virtual bool SupportsSorting() const override { return false; }
 
-		virtual void SortItems(TArray<TSharedPtr<SceneOutliner::TOutlinerTreeItem>>& RootItems, const EColumnSortMode::Type SortMode) const OVERRIDE {}
+		virtual void SortItems(TArray<TSharedPtr<SceneOutliner::TOutlinerTreeItem>>& RootItems, const EColumnSortMode::Type SortMode) const override {}
 
 		// End ISceneOutlinerColumn Implementation
 		//////////////////////////////////////////////////////////////////////////
@@ -2294,8 +2294,8 @@ public:
 
 
 	/** SWidget overrides */
-	virtual void OnMouseEnter( const FGeometry& MyGeometry, const FPointerEvent& MouseEvent ) OVERRIDE;
-	virtual void OnMouseLeave( const FPointerEvent& MouseEvent ) OVERRIDE;
+	virtual void OnMouseEnter( const FGeometry& MyGeometry, const FPointerEvent& MouseEvent ) override;
+	virtual void OnMouseLeave( const FPointerEvent& MouseEvent ) override;
 
 
 private:
diff --git a/Engine/Source/Editor/LevelEditor/Private/SLevelViewportToolBar.cpp b/Engine/Source/Editor/LevelEditor/Private/SLevelViewportToolBar.cpp
--- a/Engine/Source/Editor/LevelEditor/Private/SLevelViewportToolBar.cpp
+++ b/Engine/Source/Editor/LevelEditor/Private/SLevelViewportToolBar.cpp
@@ -27,7 +27,7 @@
 class SLevelEditorViewportViewMenu : public SEditorViewportViewMenu
 {
 public:
-	virtual TSharedRef<SWidget> GenerateViewMenuContent() const OVERRIDE
+	virtual TSharedRef<SWidget> GenerateViewMenuContent() const override
 	{
 		SLevelViewport* LevelViewport = static_cast<SLevelViewport*>(Viewport.Pin().Get());
 		LevelViewport->OnFloatingButtonClicked();
diff --git a/Engine/Source/Editor/LevelEditor/Private/SSurfaceProperties.h b/Engine/Source/Editor/LevelEditor/Private/SSurfaceProperties.h
--- a/Engine/Source/Editor/LevelEditor/Private/SSurfaceProperties.h
+++ b/Engine/Source/Editor/LevelEditor/Private/SSurfaceProperties.h
@@ -34,10 +34,10 @@ private:
 	};
 	
 	// FNotifyHook Interface
-	virtual void NotifyPostChange( const FPropertyChangedEvent& PropertyChangedEvent, UProperty* PropertyThatChanged ) OVERRIDE;
+	virtual void NotifyPostChange( const FPropertyChangedEvent& PropertyChangedEvent, UProperty* PropertyThatChanged ) override;
 	
 	// FGCObject Interface
-	virtual void AddReferencedObjects(FReferenceCollector& Collector) OVERRIDE;
+	virtual void AddReferencedObjects(FReferenceCollector& Collector) override;
 
 	/** 
 	 * Creates Pan texture controls
diff --git a/Engine/Source/Editor/LevelEditor/Public/LevelEditor.h b/Engine/Source/Editor/LevelEditor/Public/LevelEditor.h
--- a/Engine/Source/Editor/LevelEditor/Public/LevelEditor.h
+++ b/Engine/Source/Editor/LevelEditor/Public/LevelEditor.h
@@ -24,29 +24,29 @@ public:
 	/**
 	 * Called right after the module DLL has been loaded and the module object has been created
 	 */
-	virtual void StartupModule() OVERRIDE;
+	virtual void StartupModule() override;
 
 	/**
 	 * Called before the module is unloaded, right before the module object is destroyed.
 	 */
-	virtual void ShutdownModule() OVERRIDE;
+	virtual void ShutdownModule() override;
 
 	/**
 	 * Called before the module has been unloaded
 	 */
-	virtual void PreUnloadCallback() OVERRIDE;
+	virtual void PreUnloadCallback() override;
 
 	/**
 	 * Called after the module has been reloaded
 	 */
-	virtual void PostLoadCallback() OVERRIDE;
+	virtual void PostLoadCallback() override;
 
 	/**
 	 * Override this to set whether your module is allowed to be unloaded on the fly
 	 *
 	 * @return Whether the module supports shutdown separate from the rest of the engine.
 	 */
-	virtual bool SupportsDynamicReloading() OVERRIDE
+	virtual bool SupportsDynamicReloading() override
 	{
 		// @todo: Eventually, this should probably not be allowed.
 		return true;
diff --git a/Engine/Source/Editor/LevelEditor/Public/SLevelViewport.h b/Engine/Source/Editor/LevelEditor/Public/SLevelViewport.h
--- a/Engine/Source/Editor/LevelEditor/Public/SLevelViewport.h
+++ b/Engine/Source/Editor/LevelEditor/Public/SLevelViewport.h
@@ -56,7 +56,7 @@ public:
 		return *LevelViewportClient;
 	}
 
-	virtual FLevelEditorViewportClient& GetLevelViewportClient() OVERRIDE
+	virtual FLevelEditorViewportClient& GetLevelViewportClient() override
 	{		
 		return *LevelViewportClient;
 	}
@@ -75,19 +75,19 @@ public:
 	void SaveConfig(const FString& ConfigName);
 
 	/** ILevelViewport Interface */
-	virtual void StartPlayInEditorSession( UGameViewportClient* PlayClient, const bool bInSimulateInEditor ) OVERRIDE;
-	virtual void EndPlayInEditorSession() OVERRIDE;
-	virtual void SwapViewportsForSimulateInEditor() OVERRIDE;
-	virtual void SwapViewportsForPlayInEditor() OVERRIDE;
-	virtual void OnSimulateSessionStarted() OVERRIDE;
-	virtual void OnSimulateSessionFinished() OVERRIDE;
-	virtual void RegisterGameViewportIfPIE() OVERRIDE;
-	virtual bool HasPlayInEditorViewport() const OVERRIDE; 
-	virtual FViewport* GetActiveViewport() OVERRIDE;
-	virtual TSharedRef< const SWidget> AsWidget() OVERRIDE { return AsShared(); }
+	virtual void StartPlayInEditorSession( UGameViewportClient* PlayClient, const bool bInSimulateInEditor ) override;
+	virtual void EndPlayInEditorSession() override;
+	virtual void SwapViewportsForSimulateInEditor() override;
+	virtual void SwapViewportsForPlayInEditor() override;
+	virtual void OnSimulateSessionStarted() override;
+	virtual void OnSimulateSessionFinished() override;
+	virtual void RegisterGameViewportIfPIE() override;
+	virtual bool HasPlayInEditorViewport() const override; 
+	virtual FViewport* GetActiveViewport() override;
+	virtual TSharedRef< const SWidget> AsWidget() override { return AsShared(); }
 
 	/** SEditorViewport Interface */
-	virtual void OnFocusViewportToSelection() OVERRIDE;
+	virtual void OnFocusViewportToSelection() override;
 
 	/**
 	 * Called when the maximize command is executed                   
@@ -105,12 +105,12 @@ public:
 	 * @param	bWantImmersive Whether to switch to immersive mode, or switch back to normal mode
 	 * @param	bAllowAnimation	True to allow animation when transitioning, otherwise false
 	 */
-	void MakeImmersive( const bool bWantImmersive, const bool bAllowAnimation ) OVERRIDE;
+	void MakeImmersive( const bool bWantImmersive, const bool bAllowAnimation ) override;
 
 	/**
 	 * @return true if this viewport is in immersive mode, false otherwise
 	 */
-	bool IsImmersive () const OVERRIDE;
+	bool IsImmersive () const override;
 	
 	/**
 	 * Called to get the visibility of the viewport's 'Restore from Immersive' button. Returns EVisibility::Collapsed when not in immersive mode
@@ -142,7 +142,7 @@ public:
 	/**
 	 * Called when game view should be toggled
 	 */
-	void ToggleGameView() OVERRIDE;
+	void ToggleGameView() override;
 
 	/**
 	 * @return true if we can toggle game view
@@ -152,7 +152,7 @@ public:
 	/**
 	 * @return true if we are in game view                   
 	 */
-	bool IsInGameView() const OVERRIDE;
+	bool IsInGameView() const override;
 
 	/**
 	 * Toggles layer visibility in this viewport
@@ -255,17 +255,17 @@ public:
 
 protected:
 	/** SEditorViewport interface */
-	virtual TSharedRef<FEditorViewportClient> MakeEditorViewportClient() OVERRIDE;
-	virtual TSharedPtr<SWidget> MakeViewportToolbar() OVERRIDE;
-
-	virtual void OnIncrementPositionGridSize() OVERRIDE;
-	virtual void OnDecrementPositionGridSize() OVERRIDE;
-	virtual void OnIncrementRotationGridSize() OVERRIDE;
-	virtual void OnDecrementRotationGridSize() OVERRIDE;
-	virtual const FSlateBrush* OnGetViewportBorderBrush() const OVERRIDE;
-	virtual FSlateColor OnGetViewportBorderColorAndOpacity() const OVERRIDE;
-	virtual EVisibility OnGetViewportContentVisibility() const OVERRIDE;
-	virtual void BindCommands() OVERRIDE;
+	virtual TSharedRef<FEditorViewportClient> MakeEditorViewportClient() override;
+	virtual TSharedPtr<SWidget> MakeViewportToolbar() override;
+
+	virtual void OnIncrementPositionGridSize() override;
+	virtual void OnDecrementPositionGridSize() override;
+	virtual void OnIncrementRotationGridSize() override;
+	virtual void OnDecrementRotationGridSize() override;
+	virtual const FSlateBrush* OnGetViewportBorderBrush() const override;
+	virtual FSlateColor OnGetViewportBorderColorAndOpacity() const override;
+	virtual EVisibility OnGetViewportContentVisibility() const override;
+	virtual void BindCommands() override;
 private:
 	/** Loads this viewport's config from the ini file */
 	FLevelEditorViewportInstanceSettings LoadLegacyConfigFromIni(const FString& ConfigKey, const FLevelEditorViewportInstanceSettings& InDefaultSettings);
@@ -398,7 +398,7 @@ private:
 	 *
 	 * @param CommandName				Name of the command
 	 */
-	virtual void ToggleStatCommand(FString CommandName) OVERRIDE;
+	virtual void ToggleStatCommand(FString CommandName) override;
 
 	/**
 	 * Called when show flags for this viewport should be reset to default
@@ -495,12 +495,12 @@ private:
 	bool HandlePlaceDraggedObjects(const FDragDropEvent& DragDropEvent, bool bCreateDropPreview);
 
 	/** SWidget Interface */
-	virtual FReply OnKeyDown( const FGeometry& MyGeometry, const FKeyboardEvent& InKeyboardEvent ) OVERRIDE;
-	virtual void OnDragEnter( const FGeometry& MyGeometry, const FDragDropEvent& DragDropEvent ) OVERRIDE;
-	virtual void OnDragLeave( const FDragDropEvent& DragDropEvent ) OVERRIDE;
-	virtual FReply OnDragOver( const FGeometry& MyGeometry, const FDragDropEvent& DragDropEvent ) OVERRIDE;
-	virtual FReply OnDrop( const FGeometry& MyGeometry, const FDragDropEvent& DragDropEvent ) OVERRIDE;
-	virtual void Tick( const FGeometry& AllottedGeometry, const double InCurrentTime, const float InDeltaTime ) OVERRIDE;
+	virtual FReply OnKeyDown( const FGeometry& MyGeometry, const FKeyboardEvent& InKeyboardEvent ) override;
+	virtual void OnDragEnter( const FGeometry& MyGeometry, const FDragDropEvent& DragDropEvent ) override;
+	virtual void OnDragLeave( const FDragDropEvent& DragDropEvent ) override;
+	virtual FReply OnDragOver( const FGeometry& MyGeometry, const FDragDropEvent& DragDropEvent ) override;
+	virtual FReply OnDrop( const FGeometry& MyGeometry, const FDragDropEvent& DragDropEvent ) override;
+	virtual void Tick( const FGeometry& AllottedGeometry, const double InCurrentTime, const float InDeltaTime ) override;
 	/** End of SWidget interface */
 
 	/**
@@ -586,10 +586,10 @@ private:
 	void RemoveActorPreview( int32 PreviewIndex );
 	
 	/** Adds a widget overlaid over the viewport */
-	virtual void AddOverlayWidget(TSharedRef<SWidget> OverlaidWidget) OVERRIDE;
+	virtual void AddOverlayWidget(TSharedRef<SWidget> OverlaidWidget) override;
 
 	/** Removes a widget that was previously overlaid on to this viewport */
-	virtual void RemoveOverlayWidget(TSharedRef<SWidget> OverlaidWidget) OVERRIDE;
+	virtual void RemoveOverlayWidget(TSharedRef<SWidget> OverlaidWidget) override;
 
 	/** Returns true if this viewport is the active viewport and can process UI commands */
 	bool CanProduceActionForCommand(const TSharedRef<const FUICommandInfo>& Command) const;
diff --git a/Engine/Source/Editor/Levels/Private/ActorsAssignedToSpecificLevelsFilter.h b/Engine/Source/Editor/Levels/Private/ActorsAssignedToSpecificLevelsFilter.h
--- a/Engine/Source/Editor/Levels/Private/ActorsAssignedToSpecificLevelsFilter.h
+++ b/Engine/Source/Editor/Levels/Private/ActorsAssignedToSpecificLevelsFilter.h
@@ -7,7 +7,7 @@ class FActorsAssignedToSpecificLevelsFilter : public IFilter< const AActor* cons
 public:
 
 	DECLARE_DERIVED_EVENT(FActorsAssignedToSpecificLevelsFilter, IFilter< const AActor* const >::FChangedEvent, FChangedEvent);
-	virtual FChangedEvent& OnChanged() OVERRIDE { return ChangedEvent; }
+	virtual FChangedEvent& OnChanged() override { return ChangedEvent; }
 
 	/** 
 	 * Returns whether the specified Item passes the Filter's text restrictions 
@@ -15,7 +15,7 @@ public:
 	 *	@param	InItem	The Item to check 
 	 *	@return			Whether the specified Item passed the filter
 	 */
-	virtual bool PassesFilter( const AActor* const InActor ) const OVERRIDE
+	virtual bool PassesFilter( const AActor* const InActor ) const override
 	{
 		if( LevelNames.Num() == 0 )
 		{
diff --git a/Engine/Source/Editor/Levels/Private/LevelCollectionViewModel.h b/Engine/Source/Editor/Levels/Private/LevelCollectionViewModel.h
--- a/Engine/Source/Editor/Levels/Private/LevelCollectionViewModel.h
+++ b/Engine/Source/Editor/Levels/Private/LevelCollectionViewModel.h
@@ -34,9 +34,9 @@ public:
 
 public:
 	// Begin FTickableEditorObject Interface
-	virtual bool IsTickable(void) const OVERRIDE { return true; }
-	virtual void Tick( float DeltaTime ) OVERRIDE;
-	virtual TStatId GetStatId() const OVERRIDE;
+	virtual bool IsTickable(void) const override { return true; }
+	virtual void Tick( float DeltaTime ) override;
+	virtual TStatId GetStatId() const override;
 	// End of FTickableEditorObject
 
 	/** Handler for when a world has been added. */
@@ -245,8 +245,8 @@ private:
 	void AppendSelectLevelNames( TArray< FName >& OutLevelNames ) const;
 
 	// Begin FEditorUndoClient Interface
-	virtual void PostUndo(bool bSuccess) OVERRIDE { Refresh(); }
-	virtual void PostRedo(bool bSuccess) OVERRIDE { PostUndo(bSuccess); }
+	virtual void PostUndo(bool bSuccess) override { Refresh(); }
+	virtual void PostRedo(bool bSuccess) override { PostUndo(bSuccess); }
 	// End of FEditorUndoClient
 
 	/**
diff --git a/Engine/Source/Editor/Levels/Private/SLevelBrowser.h b/Engine/Source/Editor/Levels/Private/SLevelBrowser.h
--- a/Engine/Source/Editor/Levels/Private/SLevelBrowser.h
+++ b/Engine/Source/Editor/Levels/Private/SLevelBrowser.h
@@ -158,7 +158,7 @@ protected:
 	 *
 	 * @param DragDropEvent   The drag and drop event.
 	 */
-	virtual void OnDragLeave( const FDragDropEvent& DragDropEvent ) OVERRIDE
+	virtual void OnDragLeave( const FDragDropEvent& DragDropEvent ) override
 	{
 		TSharedPtr< FActorDragDropGraphEdOp > DragActorOp = DragDropEvent.GetOperationAs< FActorDragDropGraphEdOp >();
 		if (DragActorOp.IsValid())
@@ -175,7 +175,7 @@ protected:
 	 *
 	 * @return A reply that indicated whether this event was handled.
 	 */
-	virtual FReply OnDragOver( const FGeometry& MyGeometry, const FDragDropEvent& DragDropEvent ) OVERRIDE
+	virtual FReply OnDragOver( const FGeometry& MyGeometry, const FDragDropEvent& DragDropEvent ) override
 	{
 		//@todo: add drag + drop support
 		return FReply::Unhandled();
@@ -189,7 +189,7 @@ protected:
 	 *
 	 * @return A reply that indicated whether this event was handled.
 	 */
-	virtual FReply OnDrop( const FGeometry& MyGeometry, const FDragDropEvent& DragDropEvent ) OVERRIDE
+	virtual FReply OnDrop( const FGeometry& MyGeometry, const FDragDropEvent& DragDropEvent ) override
 	{
 		//@todo: add drag + drop support
 		return FReply::Unhandled();
diff --git a/Engine/Source/Editor/Levels/Private/SLevelsView.h b/Engine/Source/Editor/Levels/Private/SLevelsView.h
--- a/Engine/Source/Editor/Levels/Private/SLevelsView.h
+++ b/Engine/Source/Editor/Levels/Private/SLevelsView.h
@@ -192,7 +192,7 @@ protected:
 	 *
 	 * @return  True if this widget can take keyboard focus
 	 */
-	virtual bool SupportsKeyboardFocus() const OVERRIDE
+	virtual bool SupportsKeyboardFocus() const override
 	{
 		return true;
 	}
@@ -205,7 +205,7 @@ protected:
 	 *
 	 * @return  Returns whether the event was handled, along with other possible actions
 	 */
-	virtual FReply OnKeyDown( const FGeometry& MyGeometry, const FKeyboardEvent& InKeyboardEvent ) OVERRIDE
+	virtual FReply OnKeyDown( const FGeometry& MyGeometry, const FKeyboardEvent& InKeyboardEvent ) override
 	{
 		return ViewModel->GetCommandList()->ProcessCommandBindings( InKeyboardEvent ) ? FReply::Handled() : FReply::Unhandled();
 	}
@@ -215,7 +215,7 @@ protected:
 	 *
 	 * @param DragDropEvent   The drag and drop event.
 	 */
-	virtual void OnDragLeave( const FDragDropEvent& DragDropEvent ) OVERRIDE
+	virtual void OnDragLeave( const FDragDropEvent& DragDropEvent ) override
 	{
 		TSharedPtr< FActorDragDropGraphEdOp > DragActorOp = DragDropEvent.GetOperationAs< FActorDragDropGraphEdOp >();
 		if (DragActorOp.IsValid())
@@ -233,7 +233,7 @@ protected:
 	 *
 	 * @return A reply that indicated whether this event was handled.
 	 */
-	virtual FReply OnDragOver( const FGeometry& MyGeometry, const FDragDropEvent& DragDropEvent ) OVERRIDE
+	virtual FReply OnDragOver( const FGeometry& MyGeometry, const FDragDropEvent& DragDropEvent ) override
 	{
 
 		TSharedPtr< FActorDragDropGraphEdOp > DragActorOp = DragDropEvent.GetOperationAs< FActorDragDropGraphEdOp >();
@@ -257,7 +257,7 @@ protected:
 	 *
 	 * @return A reply that indicated whether this event was handled.
 	 */
-	virtual FReply OnDrop( const FGeometry& MyGeometry, const FDragDropEvent& DragDropEvent ) OVERRIDE
+	virtual FReply OnDrop( const FGeometry& MyGeometry, const FDragDropEvent& DragDropEvent ) override
 	{
 		const bool bIsValidDrop = DragDropEvent.GetOperationAs<FActorDragDropGraphEdOp>().IsValid();
 		return bIsValidDrop ? FReply::Handled() : FReply::Unhandled();
diff --git a/Engine/Source/Editor/Levels/Private/SLevelsViewRow.h b/Engine/Source/Editor/Levels/Private/SLevelsViewRow.h
--- a/Engine/Source/Editor/Levels/Private/SLevelsViewRow.h
+++ b/Engine/Source/Editor/Levels/Private/SLevelsViewRow.h
@@ -106,7 +106,7 @@ protected:
 	 * @return a widget to represent the contents of a cell in this row of a TableView. 
 	 */
 	BEGIN_SLATE_FUNCTION_BUILD_OPTIMIZATION
-	virtual TSharedRef< SWidget > GenerateWidgetForColumn( const FName & ColumnID ) OVERRIDE
+	virtual TSharedRef< SWidget > GenerateWidgetForColumn( const FName & ColumnID ) override
 	{
 		TSharedPtr< SWidget > TableRowContent;
 
@@ -378,7 +378,7 @@ protected:
 	 *
 	 * @param DragDropEvent   The drag and drop event.
 	 */
-	virtual void OnDragLeave( const FDragDropEvent& DragDropEvent ) OVERRIDE
+	virtual void OnDragLeave( const FDragDropEvent& DragDropEvent ) override
 	{
 		TSharedPtr< FActorDragDropGraphEdOp > DragActorOp = DragDropEvent.GetOperationAs< FActorDragDropGraphEdOp >();
 		if (DragActorOp.IsValid())
@@ -395,7 +395,7 @@ protected:
 	 *
 	 * @return A reply that indicated whether this event was handled.
 	 */
-	virtual FReply OnDragOver( const FGeometry& MyGeometry, const FDragDropEvent& DragDropEvent ) OVERRIDE
+	virtual FReply OnDragOver( const FGeometry& MyGeometry, const FDragDropEvent& DragDropEvent ) override
 	{
 		TSharedPtr< FActorDragDropGraphEdOp > DragActorOp = DragDropEvent.GetOperationAs< FActorDragDropGraphEdOp >();
 		if (!DragActorOp.IsValid())
@@ -434,7 +434,7 @@ protected:
 	 *
 	 * @return A reply that indicated whether this event was handled.
 	 */
-	virtual FReply OnDrop( const FGeometry& MyGeometry, const FDragDropEvent& DragDropEvent ) OVERRIDE
+	virtual FReply OnDrop( const FGeometry& MyGeometry, const FDragDropEvent& DragDropEvent ) override
 	{
 		TSharedPtr< FActorDragDropGraphEdOp > DragActorOp = DragDropEvent.GetOperationAs< FActorDragDropGraphEdOp >();
 		if (!DragActorOp.IsValid())
diff --git a/Engine/Source/Editor/Levels/Public/LevelEdMode.h b/Engine/Source/Editor/Levels/Public/LevelEdMode.h
--- a/Engine/Source/Editor/Levels/Public/LevelEdMode.h
+++ b/Engine/Source/Editor/Levels/Public/LevelEdMode.h
@@ -16,31 +16,31 @@ public:
 
 	// Begin FEdMode
 
-	virtual void Exit() OVERRIDE; 
+	virtual void Exit() override; 
 
-	virtual void AddReferencedObjects( FReferenceCollector& Collector ) OVERRIDE;
+	virtual void AddReferencedObjects( FReferenceCollector& Collector ) override;
 
-	virtual EAxisList::Type GetWidgetAxisToDraw( FWidget::EWidgetMode InWidgetMode ) const OVERRIDE;
+	virtual EAxisList::Type GetWidgetAxisToDraw( FWidget::EWidgetMode InWidgetMode ) const override;
 
-	virtual bool ShouldDrawWidget() const OVERRIDE;
+	virtual bool ShouldDrawWidget() const override;
 
-	virtual bool UsesTransformWidget(FWidget::EWidgetMode CheckMode) const OVERRIDE;
+	virtual bool UsesTransformWidget(FWidget::EWidgetMode CheckMode) const override;
 
-	virtual FVector GetWidgetLocation() const OVERRIDE;
+	virtual FVector GetWidgetLocation() const override;
 
-	virtual bool AllowWidgetMove() OVERRIDE { return true; }
+	virtual bool AllowWidgetMove() override { return true; }
 
-	virtual bool InputDelta( FLevelEditorViewportClient* InViewportClient, FViewport* InViewport, FVector& InDrag, FRotator& InRot, FVector& InScale ) OVERRIDE;
+	virtual bool InputDelta( FLevelEditorViewportClient* InViewportClient, FViewport* InViewport, FVector& InDrag, FRotator& InRot, FVector& InScale ) override;
 
-	virtual void Render( const FSceneView* View, FViewport* Viewport, FPrimitiveDrawInterface* PDI ) OVERRIDE;
+	virtual void Render( const FSceneView* View, FViewport* Viewport, FPrimitiveDrawInterface* PDI ) override;
 
-	virtual bool StartTracking(FLevelEditorViewportClient* InViewportClient, FViewport* InViewport) OVERRIDE;
+	virtual bool StartTracking(FLevelEditorViewportClient* InViewportClient, FViewport* InViewport) override;
 
-	virtual bool EndTracking(FLevelEditorViewportClient* InViewportClient, FViewport* InViewport) OVERRIDE;
+	virtual bool EndTracking(FLevelEditorViewportClient* InViewportClient, FViewport* InViewport) override;
 
-	virtual bool IsSnapRotationEnabled() OVERRIDE;
+	virtual bool IsSnapRotationEnabled() override;
 
-	virtual bool SnapRotatorToGridOverride(FRotator& Rotation) OVERRIDE;
+	virtual bool SnapRotatorToGridOverride(FRotator& Rotation) override;
 
 	// End FEdMode
 
diff --git a/Engine/Source/Editor/MainFrame/Private/MainFrameModule.h b/Engine/Source/Editor/MainFrame/Private/MainFrameModule.h
--- a/Engine/Source/Editor/MainFrame/Private/MainFrameModule.h
+++ b/Engine/Source/Editor/MainFrame/Private/MainFrameModule.h
@@ -13,37 +13,37 @@ public:
 
 	// IMainFrameModule interface
 
-	virtual void CreateDefaultMainFrame( const bool bStartImmersivePIE ) OVERRIDE;
-	virtual TSharedRef<SWidget> MakeMainMenu( const TSharedPtr<FTabManager>& TabManager, const TSharedRef< FExtender > Extender ) const OVERRIDE;
-	virtual TSharedRef<SWidget> MakeMainTabMenu( const TSharedPtr<FTabManager>& TabManager, const TSharedRef< FExtender > Extender ) const OVERRIDE;
-	virtual TSharedRef<SWidget> MakeDeveloperTools( ) const OVERRIDE;
+	virtual void CreateDefaultMainFrame( const bool bStartImmersivePIE ) override;
+	virtual TSharedRef<SWidget> MakeMainMenu( const TSharedPtr<FTabManager>& TabManager, const TSharedRef< FExtender > Extender ) const override;
+	virtual TSharedRef<SWidget> MakeMainTabMenu( const TSharedPtr<FTabManager>& TabManager, const TSharedRef< FExtender > Extender ) const override;
+	virtual TSharedRef<SWidget> MakeDeveloperTools( ) const override;
 
-	virtual bool IsWindowInitialized( ) const OVERRIDE
+	virtual bool IsWindowInitialized( ) const override
 	{
 		return MainFrameHandler->GetParentWindow().IsValid();
 	}
 	
-	virtual TSharedPtr<SWindow> GetParentWindow( ) const OVERRIDE
+	virtual TSharedPtr<SWindow> GetParentWindow( ) const override
 	{
 		return MainFrameHandler->GetParentWindow();
 	}
 
-	virtual void SetMainTab(const TSharedRef<SDockTab>& MainTab) OVERRIDE
+	virtual void SetMainTab(const TSharedRef<SDockTab>& MainTab) override
 	{
 		MainFrameHandler->SetMainTab(MainTab);
 	}
 
-	virtual void EnableTabClosedDelegate( ) OVERRIDE
+	virtual void EnableTabClosedDelegate( ) override
 	{
 		MainFrameHandler->EnableTabClosedDelegate();
 	}
 
-	virtual void DisableTabClosedDelegate( ) OVERRIDE
+	virtual void DisableTabClosedDelegate( ) override
 	{
 		MainFrameHandler->DisableTabClosedDelegate();
 	}
 
-	virtual void RequestCloseEditor( ) OVERRIDE
+	virtual void RequestCloseEditor( ) override
 	{
 		if ( MainFrameHandler->CanCloseEditor() )
 		{
@@ -55,45 +55,45 @@ public:
 		}
 	}
 
-	virtual void SetLevelNameForWindowTitle( const FString& InLevelFileName ) OVERRIDE;
+	virtual void SetLevelNameForWindowTitle( const FString& InLevelFileName ) override;
 	
-	virtual FString GetLoadedLevelName( ) const OVERRIDE
+	virtual FString GetLoadedLevelName( ) const override
 	{ 
 		return LoadedLevelName; 
 	}
 
-	virtual const TSharedRef<FUICommandList>& GetMainFrameCommandBindings( ) OVERRIDE
+	virtual const TSharedRef<FUICommandList>& GetMainFrameCommandBindings( ) override
 	{
 		return FMainFrameCommands::ActionList;
 	}
 
-	virtual class FMainMRUFavoritesList* GetMRUFavoritesList() const OVERRIDE
+	virtual class FMainMRUFavoritesList* GetMRUFavoritesList() const override
 	{
 		return MRUFavoritesList;
 	}
 
-	virtual const FText GetApplicationTitle( const bool bIncludeGameName ) const OVERRIDE
+	virtual const FText GetApplicationTitle( const bool bIncludeGameName ) const override
 	{
 		return StaticGetApplicationTitle( bIncludeGameName );
 	}
 
-	virtual void ShowAboutWindow( ) const OVERRIDE
+	virtual void ShowAboutWindow( ) const override
 	{
 		FMainFrameActionCallbacks::AboutUnrealEd_Execute();
 	}
 	
 	DECLARE_DERIVED_EVENT(FMainFrameModule, IMainFrameModule::FMainFrameCreationFinishedEvent, FMainFrameCreationFinishedEvent);
-	virtual FMainFrameCreationFinishedEvent& OnMainFrameCreationFinished( ) OVERRIDE
+	virtual FMainFrameCreationFinishedEvent& OnMainFrameCreationFinished( ) override
 	{
 		return MainFrameCreationFinishedEvent;
 	}
 
 	DECLARE_DERIVED_EVENT(FMainFrameModule, IMainFrameModule::FMainFrameSDKNotInstalled, FMainFrameSDKNotInstalled);
-	virtual FMainFrameSDKNotInstalled& OnMainFrameSDKNotInstalled( ) OVERRIDE
+	virtual FMainFrameSDKNotInstalled& OnMainFrameSDKNotInstalled( ) override
 	{
 		return MainFrameSDKNotInstalled;
 	}
-	void BroadcastMainFrameSDKNotInstalled(const FString& PlatformName, const FString& DocLink) OVERRIDE
+	void BroadcastMainFrameSDKNotInstalled(const FString& PlatformName, const FString& DocLink) override
 	{
 		return MainFrameSDKNotInstalled.Broadcast(PlatformName, DocLink);
 	}
@@ -102,10 +102,10 @@ public:
 
 	// IModuleInterface interface
 
-	virtual void StartupModule( ) OVERRIDE;
-	virtual void ShutdownModule( ) OVERRIDE;
+	virtual void StartupModule( ) override;
+	virtual void ShutdownModule( ) override;
 
-	virtual bool SupportsDynamicReloading( ) OVERRIDE
+	virtual bool SupportsDynamicReloading( ) override
 	{
 		return true; // @todo: Eventually, this should probably not be allowed.
 	}
diff --git a/Engine/Source/Editor/MaterialEditor/Private/MaterialEditor.h b/Engine/Source/Editor/MaterialEditor/Private/MaterialEditor.h
--- a/Engine/Source/Editor/MaterialEditor/Private/MaterialEditor.h
+++ b/Engine/Source/Editor/MaterialEditor/Private/MaterialEditor.h
@@ -57,7 +57,7 @@ public:
 	 */
 	virtual bool ShouldCache(EShaderPlatform Platform, const FShaderType* ShaderType, const FVertexFactoryType* VertexFactoryType) const;
 
-	virtual const TArray<UTexture*>& GetReferencedTextures() const OVERRIDE
+	virtual const TArray<UTexture*>& GetReferencedTextures() const override
 	{
 		return ReferencedTextures;
 	}
@@ -108,7 +108,7 @@ public:
 	/** Entry point for compiling a specific material property.  This must call SetMaterialProperty. */
 	virtual int32 CompileProperty(EMaterialProperty Property,EShaderFrequency InShaderFrequency,FMaterialCompiler* Compiler) const;
 
-	virtual int32 GetMaterialDomain() const OVERRIDE { return MD_Surface; }
+	virtual int32 GetMaterialDomain() const override { return MD_Surface; }
 	virtual FString GetMaterialUsageDescription() const { return FString::Printf(TEXT("FMatExpressionPreview %s"), Expression.IsValid() ? *Expression->GetName() : TEXT("NULL")); }
 	virtual bool IsTwoSided() const { return false; }
 	virtual bool IsLightFunction() const { return false; }
@@ -124,13 +124,13 @@ public:
 	 * Should shaders compiled for this material be saved to disk?
 	 */
 	virtual bool IsPersistent() const { return false; }
-	virtual FGuid GetMaterialId() const OVERRIDE { return Id; }
+	virtual FGuid GetMaterialId() const override { return Id; }
 	const UMaterialExpression* GetExpression() const
 	{
 		return Expression.Get();
 	}
 
-	virtual void NotifyCompilationFinished() OVERRIDE;
+	virtual void NotifyCompilationFinished() override;
 
 	friend FArchive& operator<< ( FArchive& Ar, FMatExpressionPreview& V )
 	{
@@ -180,8 +180,8 @@ struct FMaterialInfo
 class FMaterialEditor : public IMaterialEditor, public FGCObject, public FTickableGameObject, public FEditorUndoClient, public FNotifyHook
 {
 public:
-	virtual void RegisterTabSpawners(const TSharedRef<class FTabManager>& TabManager) OVERRIDE;
-	virtual void UnregisterTabSpawners(const TSharedRef<class FTabManager>& TabManager) OVERRIDE;
+	virtual void RegisterTabSpawners(const TSharedRef<class FTabManager>& TabManager) override;
+	virtual void UnregisterTabSpawners(const TSharedRef<class FTabManager>& TabManager) override;
 
 public:
 	/** Initializes the editor to use a material. Should be the first thing called. */
@@ -205,25 +205,25 @@ public:
 	virtual ~FMaterialEditor();
 	
 	/** FGCObject interface */
-	virtual void AddReferencedObjects( FReferenceCollector& Collector ) OVERRIDE;
+	virtual void AddReferencedObjects( FReferenceCollector& Collector ) override;
 
 	/** IToolkit interface */
-	virtual FName GetToolkitFName() const OVERRIDE;
-	virtual FText GetBaseToolkitName() const OVERRIDE;
-	virtual FText GetToolkitName() const OVERRIDE;
-	virtual FString GetWorldCentricTabPrefix() const OVERRIDE;
+	virtual FName GetToolkitFName() const override;
+	virtual FText GetBaseToolkitName() const override;
+	virtual FText GetToolkitName() const override;
+	virtual FString GetWorldCentricTabPrefix() const override;
 
 	/** @return the documentation location for this editor */
-	virtual FString GetDocumentationLink() const OVERRIDE
+	virtual FString GetDocumentationLink() const override
 	{
 		return FString(TEXT("Engine/Rendering/Materials"));
 	}
 
 	/** @return Returns the color and opacity to use for the color that appears behind the tab text for this toolkit's tab in world-centric mode. */
-	virtual FLinearColor GetWorldCentricTabColorScale() const OVERRIDE;
+	virtual FLinearColor GetWorldCentricTabColorScale() const override;
 
 	/** The material instance applied to the preview mesh. */
-	virtual UMaterialInterface* GetMaterialInterface() const OVERRIDE;
+	virtual UMaterialInterface* GetMaterialInterface() const override;
 	
 	/**
 	 * Draws material info strings such as instruction count and current errors onto the canvas.
@@ -285,18 +285,18 @@ public:
 	TSharedRef<class IDetailsView> GetDetailView() const {return MaterialDetailsView.ToSharedRef();}
 	
 	// FTickableGameObject interface
-	virtual void Tick(float DeltaTime) OVERRIDE;
+	virtual void Tick(float DeltaTime) override;
 
-	virtual bool IsTickable() const OVERRIDE
+	virtual bool IsTickable() const override
 	{
 		return true;
 	}
 
-	virtual bool IsTickableWhenPaused() const OVERRIDE
+	virtual bool IsTickableWhenPaused() const override
 	{
 		return true;
 	}
-	virtual TStatId GetStatId() const OVERRIDE;
+	virtual TStatId GetStatId() const override;
 
 	/** Pushes the PreviewMesh assigned the the material instance to the thumbnail info */
 	static void UpdateThumbnailInfoPreviewMesh(UMaterialInterface* MatInterface);
@@ -308,17 +308,17 @@ public:
 	void JumpToNode(const UEdGraphNode* Node);
 
 	// IMaterial Editor Interface
-	virtual UMaterialExpression* CreateNewMaterialExpression(UClass* NewExpressionClass, const FVector2D& NodePos, bool bAutoSelect, bool bAutoAssignResource) OVERRIDE;
-	virtual UMaterialExpressionComment* CreateNewMaterialExpressionComment(const FVector2D& NodePos) OVERRIDE;
-	virtual void ForceRefreshExpressionPreviews() OVERRIDE;
-	virtual void AddToSelection(UMaterialExpression* Expression) OVERRIDE;
-	virtual void DeleteSelectedNodes() OVERRIDE;
-	virtual FString GetOriginalObjectName() const OVERRIDE;
-	virtual void UpdateMaterialAfterGraphChange() OVERRIDE;
-	virtual bool CanPasteNodes() const OVERRIDE;
-	virtual void PasteNodesHere(const FVector2D& Location) OVERRIDE;
-	virtual int32 GetNumberOfSelectedNodes() const OVERRIDE;
-	virtual FMaterialRenderProxy* GetExpressionPreview(UMaterialExpression* InExpression) OVERRIDE;
+	virtual UMaterialExpression* CreateNewMaterialExpression(UClass* NewExpressionClass, const FVector2D& NodePos, bool bAutoSelect, bool bAutoAssignResource) override;
+	virtual UMaterialExpressionComment* CreateNewMaterialExpressionComment(const FVector2D& NodePos) override;
+	virtual void ForceRefreshExpressionPreviews() override;
+	virtual void AddToSelection(UMaterialExpression* Expression) override;
+	virtual void DeleteSelectedNodes() override;
+	virtual FString GetOriginalObjectName() const override;
+	virtual void UpdateMaterialAfterGraphChange() override;
+	virtual bool CanPasteNodes() const override;
+	virtual void PasteNodesHere(const FVector2D& Location) override;
+	virtual int32 GetNumberOfSelectedNodes() const override;
+	virtual FMaterialRenderProxy* GetExpressionPreview(UMaterialExpression* InExpression) override;
 
 
 	void UpdateStatsMaterials();
@@ -359,10 +359,10 @@ public:
 	
 protected:
 	/** Called when "Save" is clicked for this asset */
-	virtual void SaveAsset_Execute() OVERRIDE;
+	virtual void SaveAsset_Execute() override;
 	
 	/** Called when this toolkit would close */
-	virtual bool OnRequestClose() OVERRIDE;
+	virtual bool OnRequestClose() override;
 
 	/** Called when the selection changes in the GraphEditor */
 	void OnSelectedNodesChanged(const TSet<class UObject*>& NewSelection);
@@ -435,7 +435,7 @@ private:
 	void ExtendToolbar();
 
 	/** Allows editor to veto the setting of a preview mesh */
-	virtual bool ApproveSetPreviewMesh(UStaticMesh* InStaticMesh, USkeletalMesh* InSkeletalMesh) OVERRIDE;
+	virtual bool ApproveSetPreviewMesh(UStaticMesh* InStaticMesh, USkeletalMesh* InSkeletalMesh) override;
 
 	/** Creates all internal widgets for the tabs to point at */
 	void CreateInternalWidgets();
@@ -528,12 +528,12 @@ private:
 	void OnMaterialUsageFlagsChanged(class UMaterial* MaterialThatChanged, int32 FlagThatChanged);
 
 	// FEditorUndoClient Interface
-	virtual void PostUndo(bool bSuccess) OVERRIDE;
-	virtual void PostRedo(bool bSuccess) OVERRIDE { PostUndo(bSuccess); }
+	virtual void PostUndo(bool bSuccess) override;
+	virtual void PostRedo(bool bSuccess) override { PostUndo(bSuccess); }
 
 	// FNotifyHook interface
-	virtual void NotifyPreChange(UProperty* PropertyAboutToChange) OVERRIDE;
-	virtual void NotifyPostChange( const FPropertyChangedEvent& PropertyChangedEvent, UProperty* PropertyThatChanged) OVERRIDE;
+	virtual void NotifyPreChange(UProperty* PropertyAboutToChange) override;
+	virtual void NotifyPostChange( const FPropertyChangedEvent& PropertyChangedEvent, UProperty* PropertyThatChanged) override;
 
 	/** Flags the material as dirty */
 	void SetMaterialDirty() {bMaterialDirty = true;}
diff --git a/Engine/Source/Editor/MaterialEditor/Private/MaterialEditorDetailCustomization.h b/Engine/Source/Editor/MaterialEditor/Private/MaterialEditorDetailCustomization.h
--- a/Engine/Source/Editor/MaterialEditor/Private/MaterialEditorDetailCustomization.h
+++ b/Engine/Source/Editor/MaterialEditor/Private/MaterialEditorDetailCustomization.h
@@ -16,7 +16,7 @@ public:
 	FMaterialExpressionParameterDetails(FOnCollectParameterGroups InCollectGroupsDelegate);
 
 	/** IDetailCustomization interface */
-	virtual void CustomizeDetails( IDetailLayoutBuilder& DetailLayout ) OVERRIDE;
+	virtual void CustomizeDetails( IDetailLayoutBuilder& DetailLayout ) override;
 
 private:
 	/** Populates the group name that this parameter details uses */
@@ -60,7 +60,7 @@ public:
 	FMaterialExpressionCollectionParameterDetails();
 
 	/** IDetailCustomization interface */
-	virtual void CustomizeDetails( IDetailLayoutBuilder& DetailLayout ) OVERRIDE;
+	virtual void CustomizeDetails( IDetailLayoutBuilder& DetailLayout ) override;
 
 private:
 	/** Populates the group name that this parameter details uses */
diff --git a/Engine/Source/Editor/MaterialEditor/Private/MaterialEditorModule.cpp b/Engine/Source/Editor/MaterialEditor/Private/MaterialEditorModule.cpp
--- a/Engine/Source/Editor/MaterialEditor/Private/MaterialEditorModule.cpp
+++ b/Engine/Source/Editor/MaterialEditor/Private/MaterialEditorModule.cpp
@@ -27,7 +27,7 @@ public:
 	/**
 	 * Called right after the module DLL has been loaded and the module object has been created
 	 */
-	virtual void StartupModule() OVERRIDE
+	virtual void StartupModule() override
 	{
 		MenuExtensibilityManager = MakeShareable(new FExtensibilityManager);
 		ToolBarExtensibilityManager = MakeShareable(new FExtensibilityManager);
@@ -36,7 +36,7 @@ public:
 	/**
 	 * Called before the module is unloaded, right before the module object is destroyed.
 	 */
-	virtual void ShutdownModule() OVERRIDE
+	virtual void ShutdownModule() override
 	{
 		MenuExtensibilityManager.Reset();
 		ToolBarExtensibilityManager.Reset();
@@ -45,7 +45,7 @@ public:
 	/**
 	 * Creates a new material editor, either for a material or a material function
 	 */
-	virtual TSharedRef<IMaterialEditor> CreateMaterialEditor( const EToolkitMode::Type Mode, const TSharedPtr< IToolkitHost >& InitToolkitHost, UMaterial* Material ) OVERRIDE
+	virtual TSharedRef<IMaterialEditor> CreateMaterialEditor( const EToolkitMode::Type Mode, const TSharedPtr< IToolkitHost >& InitToolkitHost, UMaterial* Material ) override
 	{
 		TSharedRef<FMaterialEditor> NewMaterialEditor(new FMaterialEditor());
 		NewMaterialEditor->InitEditorForMaterial(Material);
@@ -53,7 +53,7 @@ public:
 		return NewMaterialEditor;
 	}
 
-	virtual TSharedRef<IMaterialEditor> CreateMaterialEditor( const EToolkitMode::Type Mode, const TSharedPtr< IToolkitHost >& InitToolkitHost, UMaterialFunction* MaterialFunction ) OVERRIDE
+	virtual TSharedRef<IMaterialEditor> CreateMaterialEditor( const EToolkitMode::Type Mode, const TSharedPtr< IToolkitHost >& InitToolkitHost, UMaterialFunction* MaterialFunction ) override
 	{
 		TSharedRef<FMaterialEditor> NewMaterialEditor(new FMaterialEditor());
 		NewMaterialEditor->InitEditorForMaterialFunction(MaterialFunction);
@@ -61,7 +61,7 @@ public:
 		return NewMaterialEditor;
 	}
 
-	virtual TSharedRef<IMaterialEditor> CreateMaterialInstanceEditor( const EToolkitMode::Type Mode, const TSharedPtr< IToolkitHost >& InitToolkitHost, UMaterialInstance* MaterialInstance ) OVERRIDE
+	virtual TSharedRef<IMaterialEditor> CreateMaterialInstanceEditor( const EToolkitMode::Type Mode, const TSharedPtr< IToolkitHost >& InitToolkitHost, UMaterialInstance* MaterialInstance ) override
 	{
 		TSharedRef<FMaterialInstanceEditor> NewMaterialInstanceEditor(new FMaterialInstanceEditor());
 		NewMaterialInstanceEditor->InitMaterialInstanceEditor(Mode, InitToolkitHost, MaterialInstance);
diff --git a/Engine/Source/Editor/MaterialEditor/Private/MaterialInstanceEditor.h b/Engine/Source/Editor/MaterialEditor/Private/MaterialInstanceEditor.h
--- a/Engine/Source/Editor/MaterialEditor/Private/MaterialInstanceEditor.h
+++ b/Engine/Source/Editor/MaterialEditor/Private/MaterialInstanceEditor.h
@@ -12,8 +12,8 @@ class FMaterialInstanceEditor : public IMaterialEditor, public FNotifyHook, publ
 {
 public:
 	
-	virtual void RegisterTabSpawners(const TSharedRef<class FTabManager>& TabManager) OVERRIDE;
-	virtual void UnregisterTabSpawners(const TSharedRef<class FTabManager>& TabManager) OVERRIDE;
+	virtual void RegisterTabSpawners(const TSharedRef<class FTabManager>& TabManager) override;
+	virtual void UnregisterTabSpawners(const TSharedRef<class FTabManager>& TabManager) override;
 
 	/**
 	 * Edits the specified material instance object
@@ -27,18 +27,18 @@ public:
 	virtual ~FMaterialInstanceEditor();
 
 	/** FGCObject interface */
-	virtual void AddReferencedObjects(FReferenceCollector& Collector) OVERRIDE;
+	virtual void AddReferencedObjects(FReferenceCollector& Collector) override;
 
 	/** IToolkit interface */
-	virtual FName GetToolkitFName() const OVERRIDE;
-	virtual FText GetBaseToolkitName() const OVERRIDE;
-	virtual FString GetWorldCentricTabPrefix() const OVERRIDE;
+	virtual FName GetToolkitFName() const override;
+	virtual FText GetBaseToolkitName() const override;
+	virtual FString GetWorldCentricTabPrefix() const override;
 
 	/** @return Returns the color and opacity to use for the color that appears behind the tab text for this toolkit's tab in world-centric mode. */
-	virtual FLinearColor GetWorldCentricTabColorScale() const OVERRIDE;
+	virtual FLinearColor GetWorldCentricTabColorScale() const override;
 
 	/** The material instance applied to the preview mesh. */
-	virtual UMaterialInterface* GetMaterialInterface() const OVERRIDE;
+	virtual UMaterialInterface* GetMaterialInterface() const override;
 
 	/** Pre edit change notify for properties. */
 	virtual void NotifyPreChange( UProperty* PropertyAboutToChange );
@@ -121,7 +121,7 @@ private:
 	void ExtendToolbar();
 
 	/** Allows editor to veto the setting of a preview mesh */
-	virtual bool ApproveSetPreviewMesh(UStaticMesh* InStaticMesh, USkeletalMesh* InSkeletalMesh) OVERRIDE;
+	virtual bool ApproveSetPreviewMesh(UStaticMesh* InStaticMesh, USkeletalMesh* InSkeletalMesh) override;
 
 	/**	Spawns the preview tab */
 	TSharedRef<SDockTab> SpawnTab_Preview( const FSpawnTabArgs& Args );
@@ -139,8 +139,8 @@ private:
 	void Refresh();
 
 	// Begin FEditorUndoClient Interface
-	virtual void PostUndo( bool bSuccess ) OVERRIDE;
-	virtual void PostRedo( bool bSuccess ) OVERRIDE;
+	virtual void PostUndo( bool bSuccess ) override;
+	virtual void PostRedo( bool bSuccess ) override;
 	// End of FEditorUndoClient
 
 private:
diff --git a/Engine/Source/Editor/MaterialEditor/Private/SMaterialEditorViewport.cpp b/Engine/Source/Editor/MaterialEditor/Private/SMaterialEditorViewport.cpp
--- a/Engine/Source/Editor/MaterialEditor/Private/SMaterialEditorViewport.cpp
+++ b/Engine/Source/Editor/MaterialEditor/Private/SMaterialEditorViewport.cpp
@@ -18,11 +18,11 @@ public:
 	FMaterialEditorViewportClient(TWeakPtr<IMaterialEditor> InMaterialEditor, FPreviewScene& InPreviewScene);
 
 	// FEditorViewportClient interface
-	virtual FLinearColor GetBackgroundColor() const OVERRIDE;
-	virtual void Tick(float DeltaSeconds) OVERRIDE;
-	virtual void Draw(FViewport* Viewport,FCanvas* Canvas) OVERRIDE;
-	virtual bool ShouldOrbitCamera() const OVERRIDE;
-	virtual FSceneView* CalcSceneView(FSceneViewFamily* ViewFamily) OVERRIDE;
+	virtual FLinearColor GetBackgroundColor() const override;
+	virtual void Tick(float DeltaSeconds) override;
+	virtual void Draw(FViewport* Viewport,FCanvas* Canvas) override;
+	virtual bool ShouldOrbitCamera() const override;
+	virtual FSceneView* CalcSceneView(FSceneViewFamily* ViewFamily) override;
 	
 	void SetShowGrid(bool bShowGrid);
 
