Fix for crash bug in TTP 339857 - EDITOR: MATERIAL: CRASH: If you redo a clean up action on some nodes that have "preview node" enabled it can lead to a crash

diff --git a/Engine/Source/Editor/MaterialEditor/Private/MaterialEditor.cpp b/Engine/Source/Editor/MaterialEditor/Private/MaterialEditor.cpp
--- a/Engine/Source/Editor/MaterialEditor/Private/MaterialEditor.cpp
+++ b/Engine/Source/Editor/MaterialEditor/Private/MaterialEditor.cpp
@@ -2455,7 +2455,7 @@ void FMaterialEditor::SetPreviewExpression(UMaterialExpression* NewPreviewExpres
 {
 	UMaterialExpressionFunctionOutput* FunctionOutput = Cast<UMaterialExpressionFunctionOutput>(NewPreviewExpression);
 
-	if( PreviewExpression == NewPreviewExpression )
+	if( PreviewExpression == NewPreviewExpression || NULL == NewPreviewExpression )
 	{
 		if (FunctionOutput)
 		{
@@ -3530,6 +3530,11 @@ void FMaterialEditor::CleanUnusedExpressions()
 
 				FBlueprintEditorUtils::RemoveNode(NULL, GraphNode, true);
 
+				if (PreviewExpression == MaterialExpression)
+				{
+					SetPreviewExpression(NULL);
+				}
+
 				MaterialExpression->Modify();
 				Material->Expressions.Remove(MaterialExpression);
 				Material->RemoveExpressionParameter(MaterialExpression);
