Minor style change.

diff --git a/Engine/Source/Editor/MaterialEditor/Private/MaterialEditor.cpp b/Engine/Source/Editor/MaterialEditor/Private/MaterialEditor.cpp
--- a/Engine/Source/Editor/MaterialEditor/Private/MaterialEditor.cpp
+++ b/Engine/Source/Editor/MaterialEditor/Private/MaterialEditor.cpp
@@ -2455,7 +2455,7 @@ void FMaterialEditor::SetPreviewExpression(UMaterialExpression* NewPreviewExpres
 {
 	UMaterialExpressionFunctionOutput* FunctionOutput = Cast<UMaterialExpressionFunctionOutput>(NewPreviewExpression);
 
-	if( PreviewExpression == NewPreviewExpression || NULL == NewPreviewExpression )
+	if (!NewPreviewExpression || PreviewExpression == NewPreviewExpression)
 	{
 		if (FunctionOutput)
 		{
