Live preview material editor button

Toggles the compilation of the preview material everytime an edit is made.

#ttp 333463 - LIVE: RENDERING: Shaders compile automatically even when Live Nodes and Live Updates are Deselected

The Live Update, Live Nodes (and now Live Preview) toggles clearly aren't universally understood. I've added the live preview toggle to do what this user wanted. It may be an idea to have some UI guys see if extra clarity can be had from tootips/docs/icons etc.

diff --git a/Engine/Source/Editor/MaterialEditor/Private/MaterialEditor.cpp b/Engine/Source/Editor/MaterialEditor/Private/MaterialEditor.cpp
--- a/Engine/Source/Editor/MaterialEditor/Private/MaterialEditor.cpp
+++ b/Engine/Source/Editor/MaterialEditor/Private/MaterialEditor.cpp
@@ -470,6 +470,7 @@ FMaterialEditor::FMaterialEditor()
 	, ScopedTransaction(NULL)
 	, bAlwaysRefreshAllPreviews(false)
 	, bHideUnusedConnectors(false)
+	, bLivePreview(false)
 	, bIsRealtime(false)
 	, bShowStats(true)
 	, bShowBuiltinStats(false)
@@ -707,7 +708,8 @@ void FMaterialEditor::ExtendToolbar()
 			{
 				ToolbarBuilder.AddToolBarButton(FMaterialEditorCommands::Get().CameraHome);
 				ToolbarBuilder.AddToolBarButton(FMaterialEditorCommands::Get().CleanUnusedExpressions);
-				ToolbarBuilder.AddToolBarButton(FMaterialEditorCommands::Get().ShowHideConnectors);
+				ToolbarBuilder.AddToolBarButton(FMaterialEditorCommands::Get().ShowHideConnectors); 
+				ToolbarBuilder.AddToolBarButton(FMaterialEditorCommands::Get().ToggleLivePreview);
 				ToolbarBuilder.AddToolBarButton(FMaterialEditorCommands::Get().ToggleRealtimeExpressions);
 				ToolbarBuilder.AddToolBarButton(FMaterialEditorCommands::Get().AlwaysRefreshAllPreviews);
 				ToolbarBuilder.AddToolBarButton(FMaterialEditorCommands::Get().ToggleMaterialStats);
@@ -993,6 +995,7 @@ void FMaterialEditor::LoadEditorSettings()
 	EditorOptions = ConstructObject<UMaterialEditorOptions>( UMaterialEditorOptions::StaticClass() );
 	
 	if (EditorOptions->bHideUnusedConnectors) {OnShowConnectors();}
+	if (EditorOptions->bLivePreview) {ToggleLivePreview();}
 	if (EditorOptions->bAlwaysRefreshAllPreviews) {OnAlwaysRefreshAllPreviews();}
 	if (EditorOptions->bRealtimeExpressionViewport) {ToggleRealTimeExpressions();}
 
@@ -1044,6 +1047,9 @@ void FMaterialEditor::SaveEditorSettings()
 		EditorOptions->bHideUnusedConnectors		= !IsOnShowConnectorsChecked();
 		EditorOptions->bAlwaysRefreshAllPreviews	= IsOnAlwaysRefreshAllPreviews();
 		EditorOptions->bRealtimeExpressionViewport	= IsToggleRealTimeExpressionsChecked();
+		EditorOptions->bLivePreview			= IsToggleLivePreviewChecked();
+		EditorOptions->bShowBuiltinStats 		= IsToggleBuiltinStatsChecked();
+		EditorOptions->bReleaseStats 			= IsToggleReleaseStatsChecked();
 		EditorOptions->SaveConfig();
 	}
 
@@ -1117,8 +1123,14 @@ void FMaterialEditor::RegenerateCodeView()
 	}
 }
 
-void FMaterialEditor::UpdatePreviewMaterial( )
+void FMaterialEditor::UpdatePreviewMaterial( bool bForce )
 {
+	if (!bLivePreview && !bForce)
+	{
+		//Don't update the preview material
+		return;
+	}
+
 	bStatsFromPreviewMaterial = true;
 
 	if( PreviewExpression && ExpressionPreviewMaterial )
@@ -1223,7 +1235,8 @@ void FMaterialEditor::UpdateOriginalMaterial()
 	FlushShaderFileCache();
 
 	//recompile and refresh the preview material so it will be updated if there was a shader change
-	UpdatePreviewMaterial();
+	//Force it even if bLivePreview is false.
+	UpdatePreviewMaterial(true);
 
 	const FScopedBusyCursor BusyCursor;
 
@@ -1671,6 +1684,12 @@ void FMaterialEditor::BindCommands()
 		FCanExecuteAction(),
 		FIsActionChecked::CreateSP(this, &FMaterialEditor::IsOnShowConnectorsChecked));
 
+	ToolkitCommands->MapAction(
+		Commands.ToggleLivePreview,
+		FExecuteAction::CreateSP(this, &FMaterialEditor::ToggleLivePreview),
+		FCanExecuteAction(),
+		FIsActionChecked::CreateSP(this, &FMaterialEditor::IsToggleLivePreviewChecked));
+
 	ToolkitCommands->MapAction(
 		Commands.ToggleRealtimeExpressions,
 		FExecuteAction::CreateSP(this, &FMaterialEditor::ToggleRealTimeExpressions),
@@ -1792,6 +1811,20 @@ bool FMaterialEditor::IsOnShowConnectorsChecked() const
 	return bHideUnusedConnectors == false;
 }
 
+void FMaterialEditor::ToggleLivePreview()
+{
+	bLivePreview = !bLivePreview;
+	if (bLivePreview)
+	{
+		UpdatePreviewMaterial();
+	}
+}
+
+bool FMaterialEditor::IsToggleLivePreviewChecked() const
+{
+	return bLivePreview;
+}
+
 void FMaterialEditor::ToggleRealTimeExpressions()
 {
 	bIsRealtime = !bIsRealtime;
diff --git a/Engine/Source/Editor/MaterialEditor/Private/MaterialEditor.h b/Engine/Source/Editor/MaterialEditor/Private/MaterialEditor.h
--- a/Engine/Source/Editor/MaterialEditor/Private/MaterialEditor.h
+++ b/Engine/Source/Editor/MaterialEditor/Private/MaterialEditor.h
@@ -263,7 +263,7 @@ public:
 	/**
 	 * Recompiles the material used in the preview window.
 	 */
-	void UpdatePreviewMaterial();
+	void UpdatePreviewMaterial(bool bForce=false);
 
 	/**
 	 * Updates the original material with the changes made in the editor
@@ -482,6 +482,9 @@ private:
 	/** Command for the show unused connectors button */
 	void OnShowConnectors();
 	bool IsOnShowConnectorsChecked() const;
+	/** Command for the Toggle Live Preview button */
+	void ToggleLivePreview();
+	bool IsToggleLivePreviewChecked() const;
 	/** Command for the Toggle Real Time button */
 	void ToggleRealTimeExpressions();
 	bool IsToggleRealTimeExpressionsChecked() const;
@@ -646,6 +649,9 @@ private:
 	/** If true, don't render connectors that are not connected to anything. */
 	bool bHideUnusedConnectors;
 
+	/** If true, the preview material is compiled on every edit of the material. If false, only on Apply. */
+	bool bLivePreview;
+
 	/** Just storing this choice for now, not sure what difference it will make to Graph Editor */
 	bool bIsRealtime;
 
diff --git a/Engine/Source/Editor/MaterialEditor/Private/MaterialEditorActions.cpp b/Engine/Source/Editor/MaterialEditor/Private/MaterialEditorActions.cpp
--- a/Engine/Source/Editor/MaterialEditor/Private/MaterialEditorActions.cpp
+++ b/Engine/Source/Editor/MaterialEditor/Private/MaterialEditorActions.cpp
@@ -24,7 +24,8 @@ void FMaterialEditorCommands::RegisterCommands()
 
 	UI_COMMAND( CameraHome, "Home", "Goes home on the canvas.", EUserInterfaceActionType::Button, FInputGesture() );
 	UI_COMMAND( CleanUnusedExpressions, "Clean Up", "Cleans up any unused Expressions.", EUserInterfaceActionType::Button, FInputGesture() );
-	UI_COMMAND( ShowHideConnectors, "Connectors", "Show or Hide Unused Connectors", EUserInterfaceActionType::ToggleButton, FInputGesture() );
+	UI_COMMAND(ShowHideConnectors, "Connectors", "Show or Hide Unused Connectors", EUserInterfaceActionType::ToggleButton, FInputGesture());
+	UI_COMMAND( ToggleLivePreview, "Live Preview", "Toggles real time update of the preview material.", EUserInterfaceActionType::ToggleButton, FInputGesture());
 	UI_COMMAND( ToggleRealtimeExpressions, "Live Nodes", "Toggles real time update of the graph canvas.", EUserInterfaceActionType::ToggleButton, FInputGesture() );
 	UI_COMMAND( AlwaysRefreshAllPreviews, "Live Update", "All nodes are previewed live.", EUserInterfaceActionType::ToggleButton, FInputGesture() );
 	UI_COMMAND( ToggleMaterialStats, "Stats", "Toggles displaying of the material's stats.", EUserInterfaceActionType::ToggleButton, FInputGesture());
diff --git a/Engine/Source/Editor/MaterialEditor/Public/MaterialEditorActions.h b/Engine/Source/Editor/MaterialEditor/Public/MaterialEditorActions.h
--- a/Engine/Source/Editor/MaterialEditor/Public/MaterialEditorActions.h
+++ b/Engine/Source/Editor/MaterialEditor/Public/MaterialEditorActions.h
@@ -76,6 +76,9 @@ public:
 	/** Shows or hides unused connectors */
 	TSharedPtr< FUICommandInfo > ShowHideConnectors;
 	
+	/** Toggles live updating of the preview material. */
+	TSharedPtr< FUICommandInfo > ToggleLivePreview;
+
 	/** Toggles real time expression nodes */
 	TSharedPtr< FUICommandInfo > ToggleRealtimeExpressions;
 	
diff --git a/Engine/Source/Editor/UnrealEd/Classes/Preferences/MaterialEditorOptions.h b/Engine/Source/Editor/UnrealEd/Classes/Preferences/MaterialEditorOptions.h
--- a/Engine/Source/Editor/UnrealEd/Classes/Preferences/MaterialEditorOptions.h
+++ b/Engine/Source/Editor/UnrealEd/Classes/Preferences/MaterialEditorOptions.h
@@ -38,6 +38,10 @@ class UNREALED_API UMaterialEditorOptions : public UObject
 	UPROPERTY(EditAnywhere, config, Category=Options)
 	uint32 bRealtimeExpressionViewport:1;
 
+	/** If true, always refresh the material preview. */
+	UPROPERTY(EditAnywhere, config, Category = Options)
+	uint32 bLivePreview : 1;
+
 	/** If true, always refresh all expression previews. */
 	UPROPERTY(EditAnywhere, config, Category=Options)
 	uint32 bAlwaysRefreshAllPreviews:1;
diff --git a/Engine/Source/Runtime/Engine/Private/Materials/MaterialExpressions.cpp b/Engine/Source/Runtime/Engine/Private/Materials/MaterialExpressions.cpp
--- a/Engine/Source/Runtime/Engine/Private/Materials/MaterialExpressions.cpp
+++ b/Engine/Source/Runtime/Engine/Private/Materials/MaterialExpressions.cpp
@@ -501,7 +501,7 @@ void UMaterialExpression::PostEditChangeProperty(FPropertyChangedEvent& Property
 	{
 		// Don't recompile the outer material if we are in the middle of a transaction or interactively changing properties
 		// as there may be many expressions in the transaction buffer and we would just be recompiling over and over again.
-		if (Material)
+		if (Material && !Material->bIsPreviewMaterial)
 		{
 			Material->PreEditChange(NULL);
 			Material->PostEditChange();
