Fixes for crash bugs in TTP340716 and TTP 340696

Made sure preview material is compiled when opening editor.
Made sure CodeView is cleared and not regenerated if live preview is off until the material has been recompiled.

diff --git a/Engine/Source/Editor/MaterialEditor/Private/MaterialEditor.cpp b/Engine/Source/Editor/MaterialEditor/Private/MaterialEditor.cpp
--- a/Engine/Source/Editor/MaterialEditor/Private/MaterialEditor.cpp
+++ b/Engine/Source/Editor/MaterialEditor/Private/MaterialEditor.cpp
@@ -453,6 +453,11 @@ void FMaterialEditor::InitMaterialEditor( const EToolkitMode::Type Mode, const T
 
 	Material->MaterialGraph->RebuildGraph();
 	RecenterEditor();
+
+	//Make sure the preview material is initialized.
+	UpdatePreviewMaterial(true);
+	RegenerateCodeView(true);
+
 	ForceRefreshExpressionPreviews();
 }
 
@@ -470,7 +475,7 @@ FMaterialEditor::FMaterialEditor()
 	, ScopedTransaction(NULL)
 	, bAlwaysRefreshAllPreviews(false)
 	, bHideUnusedConnectors(false)
-	, bLivePreview(false)
+	, bLivePreview(true)
 	, bIsRealtime(false)
 	, bShowStats(true)
 	, bShowBuiltinStats(false)
@@ -995,7 +1000,10 @@ void FMaterialEditor::LoadEditorSettings()
 	EditorOptions = ConstructObject<UMaterialEditorOptions>( UMaterialEditorOptions::StaticClass() );
 	
 	if (EditorOptions->bHideUnusedConnectors) {OnShowConnectors();}
-	if (EditorOptions->bLivePreview) {ToggleLivePreview();}
+	if (bLivePreview != EditorOptions->bLivePreview)
+	{
+		ToggleLivePreview();
+	}
 	if (EditorOptions->bAlwaysRefreshAllPreviews) {OnAlwaysRefreshAllPreviews();}
 	if (EditorOptions->bRealtimeExpressionViewport) {ToggleRealTimeExpressions();}
 
@@ -1068,12 +1076,19 @@ FReply FMaterialEditor::CopyCodeViewTextToClipboard()
 	return FReply::Handled();
 }
 
-void FMaterialEditor::RegenerateCodeView()
+void FMaterialEditor::RegenerateCodeView(bool bForce)
 {
 #define MARKTAG TEXT("/*MARK_")
 #define MARKTAGLEN 7
 
 	HLSLCode = TEXT("");
+
+	if (!bLivePreview && !bForce)
+	{
+		//When bLivePreview is false then the source can be out of date. 
+		return;
+	}
+
 	TMap<FMaterialExpressionKey,int32> ExpressionCodeMap[MP_MAX][SF_NumFrequencies];
 	for (int32 PropertyIndex = 0; PropertyIndex < MP_MAX; PropertyIndex++)
 	{
@@ -1237,6 +1252,7 @@ void FMaterialEditor::UpdateOriginalMaterial()
 	//recompile and refresh the preview material so it will be updated if there was a shader change
 	//Force it even if bLivePreview is false.
 	UpdatePreviewMaterial(true);
+	RegenerateCodeView(true);
 
 	const FScopedBusyCursor BusyCursor;
 
@@ -1817,6 +1833,7 @@ void FMaterialEditor::ToggleLivePreview()
 	if (bLivePreview)
 	{
 		UpdatePreviewMaterial();
+		RegenerateCodeView();
 	}
 }
 
diff --git a/Engine/Source/Editor/MaterialEditor/Private/MaterialEditor.h b/Engine/Source/Editor/MaterialEditor/Private/MaterialEditor.h
--- a/Engine/Source/Editor/MaterialEditor/Private/MaterialEditor.h
+++ b/Engine/Source/Editor/MaterialEditor/Private/MaterialEditor.h
@@ -258,7 +258,7 @@ public:
 	void RefreshPreviewViewport();
 	
 	/** Regenerates the code view widget with new text */
-	void RegenerateCodeView();
+	void RegenerateCodeView(bool bForce=false);
 	
 	/**
 	 * Recompiles the material used in the preview window.
