Making sure the LivePreview toggle is on by default

diff --git a/Engine/Config/BaseEditorUserSettings.ini b/Engine/Config/BaseEditorUserSettings.ini
--- a/Engine/Config/BaseEditorUserSettings.ini
+++ b/Engine/Config/BaseEditorUserSettings.ini
@@ -382,6 +382,7 @@ bHideUnusedConnectors=False
 bRealtimeMaterialViewport=True
 bRealtimeExpressionViewport=False
 bAlwaysRefreshAllPreviews=False
+bLivePreviewUpdate=True
 
 
 [UnEdViewport]
diff --git a/Engine/Source/Editor/MaterialEditor/Private/MaterialEditor.cpp b/Engine/Source/Editor/MaterialEditor/Private/MaterialEditor.cpp
--- a/Engine/Source/Editor/MaterialEditor/Private/MaterialEditor.cpp
+++ b/Engine/Source/Editor/MaterialEditor/Private/MaterialEditor.cpp
@@ -1011,7 +1011,7 @@ void FMaterialEditor::LoadEditorSettings()
 	EditorOptions = ConstructObject<UMaterialEditorOptions>( UMaterialEditorOptions::StaticClass() );
 	
 	if (EditorOptions->bHideUnusedConnectors) {OnShowConnectors();}
-	if (bLivePreview != EditorOptions->bLivePreview)
+	if (bLivePreview != EditorOptions->bLivePreviewUpdate)
 	{
 		ToggleLivePreview();
 	}
@@ -1066,7 +1066,7 @@ void FMaterialEditor::SaveEditorSettings()
 		EditorOptions->bHideUnusedConnectors		= !IsOnShowConnectorsChecked();
 		EditorOptions->bAlwaysRefreshAllPreviews	= IsOnAlwaysRefreshAllPreviews();
 		EditorOptions->bRealtimeExpressionViewport	= IsToggleRealTimeExpressionsChecked();
-		EditorOptions->bLivePreview			= IsToggleLivePreviewChecked();
+		EditorOptions->bLivePreviewUpdate		= IsToggleLivePreviewChecked();
 		EditorOptions->bShowBuiltinStats 		= IsToggleBuiltinStatsChecked();
 		EditorOptions->bReleaseStats 			= IsToggleReleaseStatsChecked();
 		EditorOptions->SaveConfig();
diff --git a/Engine/Source/Editor/UnrealEd/Classes/Preferences/MaterialEditorOptions.h b/Engine/Source/Editor/UnrealEd/Classes/Preferences/MaterialEditorOptions.h
--- a/Engine/Source/Editor/UnrealEd/Classes/Preferences/MaterialEditorOptions.h
+++ b/Engine/Source/Editor/UnrealEd/Classes/Preferences/MaterialEditorOptions.h
@@ -40,7 +40,7 @@ class UNREALED_API UMaterialEditorOptions : public UObject
 
 	/** If true, always refresh the material preview. */
 	UPROPERTY(EditAnywhere, config, Category = Options)
-	uint32 bLivePreview : 1;
+	uint32 bLivePreviewUpdate : 1;
 
 	/** If true, always refresh all expression previews. */
 	UPROPERTY(EditAnywhere, config, Category=Options)
