Fixing crash when opening material functions.

#ttp 342489: CRITICAL: EDITOR: CRASH: FArchive does not support FAssetPtr serialization message when opening 3ColorBlend material function
#ttp 342497: CRITICAL: Regression: EDITOR:  CRASH: Opening Material Function

diff --git a/Engine/Source/Editor/MaterialEditor/Private/MaterialEditor.cpp b/Engine/Source/Editor/MaterialEditor/Private/MaterialEditor.cpp
--- a/Engine/Source/Editor/MaterialEditor/Private/MaterialEditor.cpp
+++ b/Engine/Source/Editor/MaterialEditor/Private/MaterialEditor.cpp
@@ -245,7 +245,7 @@ void FMaterialEditor::InitEditorForMaterialFunction(UMaterialFunction* InMateria
 	// Create a temporary material to preview the material function
 	Material = (UMaterial*)StaticConstructObject(UMaterial::StaticClass()); 
 	{
-		FArchive DummyArchive;
+		FArchiveUObject DummyArchive;
 		// Hack: serialize the new material with an archive that does nothing so that its material resources are created
 		Material->Serialize(DummyArchive);
 	}
