Fix missing Font sample parameter in material instances.

#TTP 340962 - LIVE: RENDERING: Font Sample Parameter Behavior in MIC
Created sub struct for UDEditorFontParameterValue to contain the font parameter fields.
Removed font specific widget construction as it can now use the generic CreateParameterValueWidget function to build the UI.

diff --git a/Engine/Source/Editor/MaterialEditor/Private/MaterialEditorInstanceDetailCustomization.cpp b/Engine/Source/Editor/MaterialEditor/Private/MaterialEditorInstanceDetailCustomization.cpp
--- a/Engine/Source/Editor/MaterialEditor/Private/MaterialEditorInstanceDetailCustomization.cpp
+++ b/Engine/Source/Editor/MaterialEditor/Private/MaterialEditorInstanceDetailCustomization.cpp
@@ -105,7 +105,7 @@ void FMaterialInstanceParameterDetails::CreateSingleGroupWidget(FEditorParameter
 		UDEditorTextureParameterValue* TextureParam = Cast<UDEditorTextureParameterValue>(Parameter);
 		UDEditorVectorParameterValue* VectorParam = Cast<UDEditorVectorParameterValue>(Parameter);
 
-		if (ScalarParam || SwitchParam || TextureParam || VectorParam)
+		if (ScalarParam || SwitchParam || TextureParam || VectorParam || FontParam)
 		{
 			CreateParameterValueWidget(Parameter, ParameterProperty, DetailGroup );
 		}
@@ -113,10 +113,6 @@ void FMaterialInstanceParameterDetails::CreateSingleGroupWidget(FEditorParameter
 		{
 			CreateMaskParameterValueWidget(Parameter, ParameterProperty, DetailGroup );
 		}
-		else if (FontParam)
-		{
-			CreateFontParameterValueWidget(Parameter, ParameterProperty, DetailGroup);
-		}
 		else
 		{
 			// Unsupported parameter type
@@ -237,50 +233,6 @@ void FMaterialInstanceParameterDetails::CreateMaskParameterValueWidget(UDEditorP
 			]
 		];
 	}
-
-	
-}
-
-void FMaterialInstanceParameterDetails::CreateFontParameterValueWidget(UDEditorParameterValue* Parameter, TSharedPtr<IPropertyHandle> ParameterProperty, IDetailGroup& DetailGroup )
-{
-	TSharedPtr<IPropertyHandle> FontPageProperty = ParameterProperty->GetChildHandle("FontPage");
-	TSharedPtr<IPropertyHandle> ParameterValueProperty = ParameterProperty->GetChildHandle("FontValue");
-
-	TAttribute<bool> IsParamEnabled = TAttribute<bool>::Create( TAttribute<bool>::FGetter::CreateSP( this, &FMaterialInstanceParameterDetails::IsOverriddenExpression, Parameter ) ) ;
-
-	IDetailPropertyRow& PropertyRow = DetailGroup.AddPropertyRow( ParameterProperty.ToSharedRef() );
-	PropertyRow.IsEnabled( IsParamEnabled );
-	// Handle reset to default manually
-	PropertyRow.OverrideResetToDefault( true, FSimpleDelegate::CreateSP( this, &FMaterialInstanceParameterDetails::ResetToDefault, Parameter ) );
-	PropertyRow.Visibility( TAttribute<EVisibility>::Create(TAttribute<EVisibility>::FGetter::CreateSP(this, &FMaterialInstanceParameterDetails::ShouldShowExpression, Parameter)) );
-
-	FDetailWidgetRow& CustomWidget = PropertyRow.CustomWidget();
-
-	CustomWidget
-	.FilterString( Parameter->ParameterName.ToString() )
-	.NameContent()
-	[
-		SNew(STextBlock)
-		.Text( NSLOCTEXT("FMaterialInstanceParameterDetails", "FontPageLabel", "Page") )
-		.ToolTipText( GetParameterExpressionDescription(Parameter) )
-		.Font(FEditorStyle::GetFontStyle(TEXT("PropertyWindow.NormalFont")))
-	]
-	.ValueContent()
-	[
-		SNew(SHorizontalBox)
-		+SHorizontalBox::Slot()
-		.AutoWidth()
-		.Padding(0.0f, 0.0f, 2.0f, 0.0f)
-		[
-			FontPageProperty->CreatePropertyValueWidget()
-		]
-		+SHorizontalBox::Slot()
-		.FillWidth(1.0f)
-		[
-			ParameterValueProperty->CreatePropertyNameWidget()
-		]
-	];
-
 }
 
 bool FMaterialInstanceParameterDetails::IsVisibleExpression(UDEditorParameterValue* Parameter)
@@ -369,8 +321,8 @@ void FMaterialInstanceParameterDetails::ResetToDefault( class UDEditorParameterV
 		int32 OutFontPage;
 		if (MaterialEditorInstance->Parent->GetFontParameterValue(ParameterName, OutFontValue,OutFontPage))
 		{
-			FontParam->FontValue = OutFontValue;
-			FontParam->FontPage = OutFontPage;
+			FontParam->ParameterValue.FontValue = OutFontValue;
+			FontParam->ParameterValue.FontPage = OutFontPage;
 			MaterialEditorInstance->CopyToSourceInstance();
 		}
 	}
diff --git a/Engine/Source/Editor/MaterialEditor/Private/MaterialEditorInstanceDetailCustomization.h b/Engine/Source/Editor/MaterialEditor/Private/MaterialEditorInstanceDetailCustomization.h
--- a/Engine/Source/Editor/MaterialEditor/Private/MaterialEditorInstanceDetailCustomization.h
+++ b/Engine/Source/Editor/MaterialEditor/Private/MaterialEditorInstanceDetailCustomization.h
@@ -34,7 +34,6 @@ private:
 	/** These methods generate the custom widgets for the various parameter types */
 	void CreateParameterValueWidget(class UDEditorParameterValue* Parameter, TSharedPtr<IPropertyHandle> ParameterProperty, IDetailGroup& DetailGroup );
 	void CreateMaskParameterValueWidget(class UDEditorParameterValue* Parameter, TSharedPtr<IPropertyHandle> ParameterProperty, IDetailGroup& DetailGroup );
-	void CreateFontParameterValueWidget(class UDEditorParameterValue* Parameter, TSharedPtr<IPropertyHandle> ParameterProperty, IDetailGroup& DetailGroup );
 
 	/** Returns true if the parameter is in the visible expressions list */
 	bool IsVisibleExpression(class UDEditorParameterValue* Parameter);
diff --git a/Engine/Source/Editor/UnrealEd/Classes/MaterialEditor/DEditorFontParameterValue.h b/Engine/Source/Editor/UnrealEd/Classes/MaterialEditor/DEditorFontParameterValue.h
--- a/Engine/Source/Editor/UnrealEd/Classes/MaterialEditor/DEditorFontParameterValue.h
+++ b/Engine/Source/Editor/UnrealEd/Classes/MaterialEditor/DEditorFontParameterValue.h
@@ -6,16 +6,25 @@
 
 #include "DEditorFontParameterValue.generated.h"
 
-UCLASS(hidecategories=Object, collapsecategories)
-class UNREALED_API UDEditorFontParameterValue : public UDEditorParameterValue
+
+USTRUCT()
+struct UNREALED_API FDFontParameters
 {
-	GENERATED_UCLASS_BODY()
+	GENERATED_USTRUCT_BODY()
 
-	UPROPERTY(EditAnywhere, Category=DEditorFontParameterValue)
+	UPROPERTY(EditAnywhere, Category = DFontParameter)
 	class UFont* FontValue;
 
-	UPROPERTY(EditAnywhere, Category=DEditorFontParameterValue)
+	UPROPERTY(EditAnywhere, Category = DFontParameter)
 	int32 FontPage;
+};
+
+UCLASS(hidecategories = Object, collapsecategories)
+class UNREALED_API UDEditorFontParameterValue : public UDEditorParameterValue
+{
+	GENERATED_UCLASS_BODY()
 
+	UPROPERTY(EditAnywhere, Category = DEditorFontParameterValue)
+	struct FDFontParameters ParameterValue;
 };
 
diff --git a/Engine/Source/Editor/UnrealEd/Private/PreviewMaterial.cpp b/Engine/Source/Editor/UnrealEd/Private/PreviewMaterial.cpp
--- a/Engine/Source/Editor/UnrealEd/Private/PreviewMaterial.cpp
+++ b/Engine/Source/Editor/UnrealEd/Private/PreviewMaterial.cpp
@@ -389,8 +389,8 @@ void UMaterialEditorInstanceConstant::RegenerateArrays()
 
 			if(SourceInstance->GetFontParameterValue(ParameterName, FontValue,FontPage))
 			{
-				ParameterValue.FontValue = FontValue;
-				ParameterValue.FontPage = FontPage;
+				ParameterValue.ParameterValue.FontValue = FontValue;
+				ParameterValue.ParameterValue.FontPage = FontPage;
 			}
 
 
@@ -402,8 +402,8 @@ void UMaterialEditorInstanceConstant::RegenerateArrays()
 				if(ParameterName==SourceParam.ParameterName)
 				{
 					ParameterValue.bOverride = true;
-					ParameterValue.FontValue = SourceParam.FontValue;
-					ParameterValue.FontPage = SourceParam.FontPage;
+					ParameterValue.ParameterValue.FontValue = SourceParam.FontValue;
+					ParameterValue.ParameterValue.FontPage = SourceParam.FontPage;
 				}
 			}
 			AssignParameterToGroup(ParentMaterial, Cast<UDEditorParameterValue>(&ParameterValue));
@@ -547,7 +547,7 @@ void UMaterialEditorInstanceConstant::CopyToSourceInstance()
 				{
 					if(FontParameterValue->bOverride)
 					{
-						SourceInstance->SetFontParameterValueEditorOnly(FontParameterValue->ParameterName,FontParameterValue->FontValue,FontParameterValue->FontPage);
+						SourceInstance->SetFontParameterValueEditorOnly(FontParameterValue->ParameterName, FontParameterValue->ParameterValue.FontValue, FontParameterValue->ParameterValue.FontPage);
 						continue;
 					}
 				}
