- Removed UI for Dev/Release shaders in material editor. Dropped to a CVar and using release shaders by default.
- Removed UI for Built-In Stats. Martin didn't like the feature but I've left the code in just in case people miss it.
- Removed useless bar at the bottom of the stats window.
- Hackily canceled some unnessessary compilaton when opening the editor.

diff --git a/Engine/Source/Developer/MessageLog/Private/Presentation/MessageLogListingViewModel.cpp b/Engine/Source/Developer/MessageLog/Private/Presentation/MessageLogListingViewModel.cpp
--- a/Engine/Source/Developer/MessageLog/Private/Presentation/MessageLogListingViewModel.cpp
+++ b/Engine/Source/Developer/MessageLog/Private/Presentation/MessageLogListingViewModel.cpp
@@ -374,6 +374,16 @@ bool FMessageLogListingViewModel::GetShowPages() const
 	return bShowPages;
 }
 
+void FMessageLogListingViewModel::SetAllowClear(bool bInAllowClear)
+{
+	bAllowClear = bAllowClear;
+}
+
+bool FMessageLogListingViewModel::GetAllowClear() const
+{
+	return bAllowClear;
+}
+
 void FMessageLogListingViewModel::SetDiscardDuplicates(bool bInDiscardDuplicates)
 {
 	bool bOldDiscardDuplicates = bDiscardDuplicates;
diff --git a/Engine/Source/Developer/MessageLog/Private/Presentation/MessageLogListingViewModel.h b/Engine/Source/Developer/MessageLog/Private/Presentation/MessageLogListingViewModel.h
--- a/Engine/Source/Developer/MessageLog/Private/Presentation/MessageLogListingViewModel.h
+++ b/Engine/Source/Developer/MessageLog/Private/Presentation/MessageLogListingViewModel.h
@@ -101,6 +101,12 @@ public:
 	/** Get whether we should show pages or not */
 	bool GetShowPages() const;
 
+	/** Set whether we should show allow the user to clear the log. */
+	void SetAllowClear(bool bInAllowClear);
+
+	/** Get whether we should show allow the user to clear the log. */
+	bool GetAllowClear() const;
+
 	/** Set whether we should discard duplicates or not */
 	void SetDiscardDuplicates(bool bInDiscardDuplicates);
 
@@ -141,6 +147,7 @@ private:
 	FMessageLogListingViewModel( TSharedPtr< FMessageLogListingModel > InMessageLogListingModel, const FText& InLogLabel, const FMessageLogInitializationOptions& InitializationOptions )
 		: bShowFilters( InitializationOptions.bShowFilters )
 		, bShowPages( InitializationOptions.bShowPages )
+		, bAllowClear( InitializationOptions.bAllowClear )
 		, bDiscardDuplicates( InitializationOptions.bDiscardDuplicates )
 		, MaxPageCount( InitializationOptions.MaxPageCount )
 		, CurrentPageIndex( 0 )
@@ -169,6 +176,9 @@ private:
 	/** Whether pages should be used/shown for this listing */
 	bool bShowPages;
 
+	/** Whether we allow the user to clear the log. */
+	bool bAllowClear;
+
 	/** Whether to check for duplicate messages & discard them */
 	bool bDiscardDuplicates;
 
diff --git a/Engine/Source/Developer/MessageLog/Private/UserInterface/SMessageLogListing.cpp b/Engine/Source/Developer/MessageLog/Private/UserInterface/SMessageLogListing.cpp
--- a/Engine/Source/Developer/MessageLog/Private/UserInterface/SMessageLogListing.cpp
+++ b/Engine/Source/Developer/MessageLog/Private/UserInterface/SMessageLogListing.cpp
@@ -256,10 +256,11 @@ void SMessageLogListing::Construct( const FArguments& InArgs, const TSharedRef<
 	MessageLogListingViewModel = StaticCastSharedRef<FMessageLogListingViewModel>(InModelView);
 
 	TSharedPtr<SHorizontalBox> HorizontalBox = NULL;
+	TSharedPtr<SVerticalBox> VerticalBox = NULL;
 
 	ChildSlot
 	[
-		SNew(SVerticalBox)
+		SAssignNew(VerticalBox, SVerticalBox)
 		+SVerticalBox::Slot()
 		.FillHeight(1)
 		[
@@ -269,76 +270,84 @@ void SMessageLogListing::Construct( const FArguments& InArgs, const TSharedRef<
 				.OnSelectionChanged(this, &SMessageLogListing::OnLineSelectionChanged)
 				.ItemHeight(24)
 		]
-		+SVerticalBox::Slot()
-		.AutoHeight()
-		.Padding(0, 0, 0, 1)
-		[
-			SNew(SSeparator)
-		]
-		+SVerticalBox::Slot()
-		.AutoHeight()
-		[
-			SAssignNew(HorizontalBox, SHorizontalBox)	
-		]
 	];
 
-	if(MessageLogListingViewModel->GetShowFilters())
+	//If we have some content below the message log, add a separator and a new box.
+	if (MessageLogListingViewModel->GetShowFilters() || 
+		MessageLogListingViewModel->GetShowPages() || 
+		MessageLogListingViewModel->GetAllowClear() )
 	{
+		VerticalBox->AddSlot()
+			.AutoHeight()
+			.Padding(0, 0, 0, 1)
+			[
+				SNew(SSeparator)
+			];
+
+		VerticalBox->AddSlot()
+			.AutoHeight()
+			[
+				SAssignNew(HorizontalBox, SHorizontalBox)
+			];
+
+		if (MessageLogListingViewModel->GetShowFilters())
+		{
+			HorizontalBox->AddSlot()
+				.FillWidth(1.0f)
+				.HAlign(HAlign_Left)
+				[
+					SNew(SComboButton)
+					.ButtonStyle(FEditorStyle::Get(), "HoverHintOnly")
+					.ForegroundColor(FSlateColor::UseForeground())
+					.ContentPadding(2)
+					.OnGetMenuContent(this, &SMessageLogListing::OnGetFilterMenuContent)
+					.ButtonContent()
+					[
+						SNew(STextBlock)
+						.Text(LOCTEXT("Show", "Show"))
+						.ToolTipText(LOCTEXT("ShowToolTip", "Only show messages of the selected types"))
+					]
+				];
+		}
+
 		HorizontalBox->AddSlot()
 			.FillWidth(1.0f)
-			.HAlign(HAlign_Left)
+			.HAlign(HAlign_Right)
 			[
 				SNew(SComboButton)
-				.ButtonStyle( FEditorStyle::Get(), "HoverHintOnly" )
+				.IsEnabled(this, &SMessageLogListing::IsPageWidgetEnabled)
+				.Visibility(this, &SMessageLogListing::GetPageWidgetVisibility)
+				.ButtonStyle(FEditorStyle::Get(), "HoverHintOnly")
 				.ForegroundColor(FSlateColor::UseForeground())
 				.ContentPadding(2)
-				.OnGetMenuContent(this, &SMessageLogListing::OnGetFilterMenuContent)
+				.OnGetMenuContent(this, &SMessageLogListing::OnGetPageMenuContent)
 				.ButtonContent()
 				[
-					SNew(STextBlock) 
-					.Text(LOCTEXT("Show", "Show"))
-					.ToolTipText(LOCTEXT("ShowToolTip", "Only show messages of the selected types"))
+					SNew(STextBlock)
+					.Text(this, &SMessageLogListing::OnGetPageMenuLabel)
+					.ToolTipText(LOCTEXT("PageToolTip", "Choose the log page to view"))
 				]
 			];
-	}
-
-	HorizontalBox->AddSlot()
-		.FillWidth(1.0f)
-		.HAlign(HAlign_Right)
-		[
-			SNew(SComboButton)
-			.IsEnabled(this, &SMessageLogListing::IsPageWidgetEnabled)
-			.Visibility(this, &SMessageLogListing::GetPageWidgetVisibility)
-			.ButtonStyle( FEditorStyle::Get(), "HoverHintOnly" )
-			.ForegroundColor(FSlateColor::UseForeground())
-			.ContentPadding(2)
-			.OnGetMenuContent(this, &SMessageLogListing::OnGetPageMenuContent)
-			.ButtonContent()
-			[
-				SNew(STextBlock) 
-				.Text(this, &SMessageLogListing::OnGetPageMenuLabel)
-				.ToolTipText(LOCTEXT("PageToolTip", "Choose the log page to view"))
-			]
-		];
 
-	// if we aren't using pages, we allow the user to manually clear the log
-	HorizontalBox->AddSlot()
-		.FillWidth(1.0f)
-		.HAlign(HAlign_Right)
-		[
-			SNew(SButton)
-			.OnClicked(this, &SMessageLogListing::OnClear)
-			.IsEnabled(this, &SMessageLogListing::IsClearWidgetEnabled)
-			.Visibility(this, &SMessageLogListing::GetClearWidgetVisibility)
-			.ButtonStyle( FEditorStyle::Get(), "HoverHintOnly" )
-			.ForegroundColor(FSlateColor::UseForeground())
-			.ContentPadding(2)
+		// if we aren't using pages, we allow the user to manually clear the log
+		HorizontalBox->AddSlot()
+			.FillWidth(1.0f)
+			.HAlign(HAlign_Right)
 			[
-				SNew(STextBlock) 
-				.Text(LOCTEXT("ClearMessageLog", "Clear"))
-				.ToolTipText(LOCTEXT("ClearMessageLog_ToolTip", "Clear the messages in this log"))
-			]
-		];
+				SNew(SButton)
+				.OnClicked(this, &SMessageLogListing::OnClear)
+				.IsEnabled(this, &SMessageLogListing::IsClearWidgetEnabled)
+				.Visibility(this, &SMessageLogListing::GetClearWidgetVisibility)
+				.ButtonStyle(FEditorStyle::Get(), "HoverHintOnly")
+				.ForegroundColor(FSlateColor::UseForeground())
+				.ContentPadding(2)
+				[
+					SNew(STextBlock)
+					.Text(LOCTEXT("ClearMessageLog", "Clear"))
+					.ToolTipText(LOCTEXT("ClearMessageLog_ToolTip", "Clear the messages in this log"))
+				]
+			];
+	}
 
 	// Register with the the view object so that it will notify if any data changes
 	MessageLogListingViewModel->OnDataChanged().AddSP( this, &SMessageLogListing::OnChanged );
@@ -611,13 +620,13 @@ bool SMessageLogListing::IsClearWidgetEnabled() const
 
 EVisibility SMessageLogListing::GetClearWidgetVisibility() const
 {
-	if(MessageLogListingViewModel->GetShowPages())
+	if (MessageLogListingViewModel->GetAllowClear() && !MessageLogListingViewModel->GetShowPages())
 	{
-		return EVisibility::Collapsed;
+		return EVisibility::Visible;
 	}
 	else
 	{
-		return EVisibility::Visible;
+		return EVisibility::Collapsed;
 	}
 }
 
diff --git a/Engine/Source/Developer/MessageLog/Public/MessageLogInitializationOptions.h b/Engine/Source/Developer/MessageLog/Public/MessageLogInitializationOptions.h
--- a/Engine/Source/Developer/MessageLog/Public/MessageLogInitializationOptions.h
+++ b/Engine/Source/Developer/MessageLog/Public/MessageLogInitializationOptions.h
@@ -10,6 +10,7 @@ struct FMessageLogInitializationOptions
 	FMessageLogInitializationOptions()
 		: bShowFilters(false)
 		, bShowPages(false)
+		, bAllowClear(true)
 		, bDiscardDuplicates(false)
 		, MaxPageCount(20)
 	{}
@@ -23,6 +24,11 @@ struct FMessageLogInitializationOptions
 	 */
 	bool bShowPages;
 
+	/**
+	* Whether to allow the user to clear this log.
+	*/
+	bool bAllowClear;
+
 	/**
 	 * Whether to check for duplicate messages & discard them
 	 */
diff --git a/Engine/Source/Editor/EditorStyle/Private/SlateEditorStyle.cpp b/Engine/Source/Editor/EditorStyle/Private/SlateEditorStyle.cpp
--- a/Engine/Source/Editor/EditorStyle/Private/SlateEditorStyle.cpp
+++ b/Engine/Source/Editor/EditorStyle/Private/SlateEditorStyle.cpp
@@ -4987,8 +4987,6 @@ void FSlateEditorStyle::FStyle::SetupToolkitStyles()
 
 		Set( "MaterialEditor.ToggleMaterialStats", new IMAGE_BRUSH( "Icons/icon_MatEd_Stats_40x", Icon40x40 ) );
 		Set( "MaterialEditor.ToggleMaterialStats.Small", new IMAGE_BRUSH( "Icons/icon_MatEd_Stats_40x", Icon20x20 ) );
-		Set( "MaterialEditor.ToggleReleaseStats", new IMAGE_BRUSH( "Icons/icon_MatEd_ReleaseStats_40x", Icon40x40 ) );
-		Set( "MaterialEditor.ToggleReleaseStats.Small", new IMAGE_BRUSH( "Icons/icon_MatEd_ReleaseStats_40x", Icon20x20 ) );
 		Set( "MaterialEditor.ToggleBuiltinStats", new IMAGE_BRUSH( "Icons/icon_MatEd_BuiltInStats_40x", Icon40x40 ) );
 		Set( "MaterialEditor.ToggleBuiltinStats.Small", new IMAGE_BRUSH( "Icons/icon_MatEd_BuiltInStats_40x", Icon20x20 ) );
 		Set( "MaterialEditor.ToggleMobileStats", new IMAGE_BRUSH( "Icons/icon_MobileStats_40x", Icon40x40 ) );
diff --git a/Engine/Source/Editor/MaterialEditor/Private/MaterialEditor.cpp b/Engine/Source/Editor/MaterialEditor/Private/MaterialEditor.cpp
--- a/Engine/Source/Editor/MaterialEditor/Private/MaterialEditor.cpp
+++ b/Engine/Source/Editor/MaterialEditor/Private/MaterialEditor.cpp
@@ -66,6 +66,12 @@
 
 DEFINE_LOG_CATEGORY_STATIC(LogMaterialEditor, Log, All);
 
+static TAutoConsoleVariable<int32> CVarMaterialEdUseDevShaders(
+	TEXT("r.MaterialEditor.UseDevShaders"),
+	0,
+	TEXT("Toggles whether the material editor will use shaders that include extra overhead incurred by the editor. Material editor must be re-opened if changed at runtime."),
+	ECVF_RenderThreadSafe);
+
 const FName FMaterialEditor::PreviewTabId( TEXT( "MaterialEditor_Preview" ) );
 const FName FMaterialEditor::GraphCanvasTabId( TEXT( "MaterialEditor_GraphCanvas" ) );
 const FName FMaterialEditor::PropertiesTabId( TEXT( "MaterialEditor_MaterialProperties" ) );
@@ -220,6 +226,11 @@ void FMaterialEditor::InitEditorForMaterial(UMaterial* InMaterial)
 	// the material editor releases the reference.
 	Material = (UMaterial*)StaticDuplicateObject(OriginalMaterial, GetTransientPackage(), NULL, ~RF_Standalone, UPreviewMaterial::StaticClass()); 
 	
+	Material->CancelOutstandingCompilation();	//The material is compiled later on anyway so no need to do it in Duplication/PostLoad. 
+												//I'm hackily canceling the jobs here but we should really not add the jobs in the first place. <<--- TODO
+												
+	Material->bAllowDevelopmentShaderCompile = CVarMaterialEdUseDevShaders.GetValueOnGameThread();
+
 	// Remove NULL entries, so the rest of the material editor can assume all entries of Material->Expressions are valid
 	// This can happen if an expression class was removed
 	for (int32 ExpressionIndex = Material->Expressions.Num() - 1; ExpressionIndex >= 0; ExpressionIndex--)
@@ -591,7 +602,9 @@ void FMaterialEditor::CreateInternalWidgets()
 	FMessageLogModule& MessageLogModule = FModuleManager::LoadModuleChecked<FMessageLogModule>("MessageLog");
 	FMessageLogInitializationOptions LogOptions;
 	// Show Pages so that user is never allowed to clear log messages
-	LogOptions.bShowPages = true;
+	LogOptions.bShowPages = false;
+	LogOptions.bShowFilters = false; //TODO - Provide custom filters? E.g. "Critical Errors" vs "Errors" needed for materials?
+	LogOptions.bAllowClear = false;
 	LogOptions.MaxPageCount = 1;
 	StatsListing = MessageLogModule.CreateLogListing( "MaterialEditorStats", LogOptions );
 
@@ -734,8 +747,6 @@ void FMaterialEditor::ExtendToolbar()
 				ToolbarBuilder.AddToolBarButton(FMaterialEditorCommands::Get().ToggleRealtimeExpressions);
 				ToolbarBuilder.AddToolBarButton(FMaterialEditorCommands::Get().AlwaysRefreshAllPreviews);
 				ToolbarBuilder.AddToolBarButton(FMaterialEditorCommands::Get().ToggleMaterialStats);
-				ToolbarBuilder.AddToolBarButton(FMaterialEditorCommands::Get().ToggleReleaseStats);
-				ToolbarBuilder.AddToolBarButton(FMaterialEditorCommands::Get().ToggleBuiltinStats);
 				ToolbarBuilder.AddToolBarButton(FMaterialEditorCommands::Get().ToggleMobileStats);
 			}
 			ToolbarBuilder.EndSection();
@@ -1038,15 +1049,6 @@ void FMaterialEditor::LoadEditorSettings()
 		ToggleMobileStats();
 	}
 
-	if (EditorOptions->bReleaseStats)
-	{
-		ToggleReleaseStats();
-	}
-
-	if (EditorOptions->bShowBuiltinStats)
-	{
-		ToggleBuiltinStats();
-	}
 
 	// Primitive type
 	int32 PrimType;
@@ -1072,8 +1074,6 @@ void FMaterialEditor::SaveEditorSettings()
 		EditorOptions->bAlwaysRefreshAllPreviews	= IsOnAlwaysRefreshAllPreviews();
 		EditorOptions->bRealtimeExpressionViewport	= IsToggleRealTimeExpressionsChecked();
 		EditorOptions->bLivePreviewUpdate		= IsToggleLivePreviewChecked();
-		EditorOptions->bShowBuiltinStats 		= IsToggleBuiltinStatsChecked();
-		EditorOptions->bReleaseStats 			= IsToggleReleaseStatsChecked();
 		EditorOptions->SaveConfig();
 	}
 
@@ -1527,6 +1527,7 @@ void FMaterialEditor::UpdateMaterialInfoList(bool bForceDisplay)
 
 				MaterialResource->GetRepresentativeInstructionCounts(Descriptions, InstructionCounts);
 
+				//Built in stats is no longer exposed to the UI but may still be useful so they're still in the code.
 				bool bBuiltinStats = false;
 				const FMaterialResource* EmptyMaterialResource = EmptyMaterial ? EmptyMaterial->GetMaterialResource(FeatureLevel) : NULL;
 				if (bShowBuiltinStats && bStatsFromPreviewMaterial && EmptyMaterialResource && InstructionCounts.Num() > 0)
@@ -1740,18 +1741,6 @@ void FMaterialEditor::BindCommands()
 		FCanExecuteAction(),
 		FIsActionChecked::CreateSP(this, &FMaterialEditor::IsToggleStatsChecked));
 
-	ToolkitCommands->MapAction(
-		Commands.ToggleReleaseStats,
-		FExecuteAction::CreateSP(this, &FMaterialEditor::ToggleReleaseStats),
-		FCanExecuteAction(),
-		FIsActionChecked::CreateSP(this, &FMaterialEditor::IsToggleReleaseStatsChecked));
-
-	ToolkitCommands->MapAction(
-		Commands.ToggleBuiltinStats,
-		FExecuteAction::CreateSP(this, &FMaterialEditor::ToggleBuiltinStats),
-		FCanExecuteAction(),
-		FIsActionChecked::CreateSP(this, &FMaterialEditor::IsToggleBuiltinStatsChecked));
-
 	ToolkitCommands->MapAction(
 		Commands.ToggleMobileStats,
 		FExecuteAction::CreateSP(this, &FMaterialEditor::ToggleMobileStats),
@@ -1894,35 +1883,6 @@ bool FMaterialEditor::IsToggleStatsChecked() const
 	return bShowStats == true;
 }
 
-void FMaterialEditor::ToggleReleaseStats()
-{
-	Material->bAllowDevelopmentShaderCompile = !Material->bAllowDevelopmentShaderCompile;
-	UpdatePreviewMaterial();
-}
-
-bool FMaterialEditor::IsToggleReleaseStatsChecked() const
-{
-	return !Material->bAllowDevelopmentShaderCompile;
-}
-
-void FMaterialEditor::ToggleBuiltinStats()
-{
-	bShowBuiltinStats = !bShowBuiltinStats;
-
-	if (bShowBuiltinStats && !bStatsFromPreviewMaterial)
-	{
-		//Have to be start using the preview material for stats.
-		UpdatePreviewMaterial();
-	}
-
-	UpdateStatsMaterials();
-}
-
-bool FMaterialEditor::IsToggleBuiltinStatsChecked() const
-{
-	return bShowBuiltinStats;
-}
-
 void FMaterialEditor::ToggleMobileStats()
 {
 	// Toggle the showing of material stats each time the user presses the show stats button
diff --git a/Engine/Source/Editor/MaterialEditor/Private/MaterialEditorActions.cpp b/Engine/Source/Editor/MaterialEditor/Private/MaterialEditorActions.cpp
--- a/Engine/Source/Editor/MaterialEditor/Private/MaterialEditorActions.cpp
+++ b/Engine/Source/Editor/MaterialEditor/Private/MaterialEditorActions.cpp
@@ -29,8 +29,6 @@ void FMaterialEditorCommands::RegisterCommands()
 	UI_COMMAND( ToggleRealtimeExpressions, "Live Nodes", "Toggles real time update of the graph canvas.", EUserInterfaceActionType::ToggleButton, FInputGesture() );
 	UI_COMMAND( AlwaysRefreshAllPreviews, "Live Update", "All nodes are previewed live.", EUserInterfaceActionType::ToggleButton, FInputGesture() );
 	UI_COMMAND( ToggleMaterialStats, "Stats", "Toggles displaying of the material's stats.", EUserInterfaceActionType::ToggleButton, FInputGesture());
-	UI_COMMAND( ToggleReleaseStats, "Release Stats", "Toggles use of release or development shaders in stats reporting.", EUserInterfaceActionType::ToggleButton, FInputGesture());
-	UI_COMMAND( ToggleBuiltinStats, "Built-In Stats", "Toggles displaying of stats for an empty material. Helps identify cost of changes in the graph.", EUserInterfaceActionType::ToggleButton, FInputGesture());
 	UI_COMMAND( ToggleMobileStats, "Mobile Stats", "Toggles material stats and compilation errors for mobile.", EUserInterfaceActionType::ToggleButton, FInputGesture() );
 	UI_COMMAND( NewComment, "New Comment", "Creates a new comment node.", EUserInterfaceActionType::Button, FInputGesture() );
 	UI_COMMAND( MatertialPasteHere, "Paste Here", "Pastes copied items at this location.", EUserInterfaceActionType::Button, FInputGesture() );
diff --git a/Engine/Source/Editor/MaterialEditor/Public/MaterialEditorActions.h b/Engine/Source/Editor/MaterialEditor/Public/MaterialEditorActions.h
--- a/Engine/Source/Editor/MaterialEditor/Public/MaterialEditorActions.h
+++ b/Engine/Source/Editor/MaterialEditor/Public/MaterialEditorActions.h
@@ -88,12 +88,6 @@ public:
 	/** Toggles the material stats on the canvas pane */
 	TSharedPtr< FUICommandInfo > ToggleMaterialStats;
 
-	/** Toggles whether the reported stats are for release shaders or development shaders. */
-	TSharedPtr< FUICommandInfo > ToggleReleaseStats;
-
-	/** Toggles the reporting of instructions for an empty material. i.e. The built in cost of that material. */
-	TSharedPtr< FUICommandInfo > ToggleBuiltinStats;
-
 	/** Shows material stats and errors when compiled for mobile. */
 	TSharedPtr< FUICommandInfo > ToggleMobileStats;
 	
diff --git a/Engine/Source/Editor/UnrealEd/Classes/Preferences/MaterialEditorOptions.h b/Engine/Source/Editor/UnrealEd/Classes/Preferences/MaterialEditorOptions.h
--- a/Engine/Source/Editor/UnrealEd/Classes/Preferences/MaterialEditorOptions.h
+++ b/Engine/Source/Editor/UnrealEd/Classes/Preferences/MaterialEditorOptions.h
@@ -58,14 +58,6 @@ class UNREALED_API UMaterialEditorOptions : public UObject
 	UPROPERTY(EditAnywhere, config, Category=Options)
 	uint32 bShowMobileStats:1;
 
-	/** If true, stats use release shaders. */
-	UPROPERTY(EditAnywhere, config, Category = Options)
-	uint32 bReleaseStats : 1;
-
-	/** If true, show stats for an blank copy (no graph) of the preview material. Helps identify cost of material graph. Incurs slight cost to build/iteration times. */
-	UPROPERTY(EditAnywhere, config, Category = Options)
-	uint32 bShowBuiltinStats : 1;
-
 	/** The users favorite material expressions. */
 	UPROPERTY(EditAnywhere, config, Category=Options)
 	TArray<FString> FavoriteExpressions;
