Add context menu option to material nodes for viewing documentation

diff --git a/Engine/Source/Editor/EditorStyle/Private/SlateEditorStyle.cpp b/Engine/Source/Editor/EditorStyle/Private/SlateEditorStyle.cpp
--- a/Engine/Source/Editor/EditorStyle/Private/SlateEditorStyle.cpp
+++ b/Engine/Source/Editor/EditorStyle/Private/SlateEditorStyle.cpp
@@ -5053,6 +5053,7 @@ void FSlateEditorStyle::FStyle::SetupToolkitStyles()
 		Set( "MaterialEditor.CameraHome.Small", new IMAGE_BRUSH( "Icons/icon_MatEd_Home_40x", Icon20x20 ) );
 		Set( "MaterialEditor.FindInMaterial", new IMAGE_BRUSH("Icons/icon_Blueprint_Find_40px", Icon40x40));
 		Set( "MaterialEditor.FindInMaterial.Small", new IMAGE_BRUSH("Icons/icon_Blueprint_Find_40px", Icon20x20));
+		Set( "MaterialEditor.GoToDocumentation", new IMAGE_BRUSH( "Common/icon_Help_Hover_16x", Icon16x16 ));
 	}
 
 	// Material Instance Editor
diff --git a/Engine/Source/Editor/MaterialEditor/Private/MaterialEditor.cpp b/Engine/Source/Editor/MaterialEditor/Private/MaterialEditor.cpp
--- a/Engine/Source/Editor/MaterialEditor/Private/MaterialEditor.cpp
+++ b/Engine/Source/Editor/MaterialEditor/Private/MaterialEditor.cpp
@@ -59,6 +59,8 @@
 #include "SMaterialPalette.h"
 #include "FindInMaterial.h"
 #include "SColorPicker.h"
+#include "EditorClassUtils.h"
+#include "IDocumentation.h"
 
 #include "Developer/MessageLog/Public/MessageLogModule.h"
 #include "Particles/ParticleSystemComponent.h"
@@ -2337,6 +2339,41 @@ void FMaterialEditor::OnFindInMaterial()
 	FindResults->FocusForUse();
 }
 
+FString FMaterialEditor::GetDocLinkForSelectedNode()
+{
+	FString DocumentationLink;
+
+	TArray<UObject*> SelectedNodes = GraphEditor->GetSelectedNodes().Array();
+	if (SelectedNodes.Num() == 1)
+	{
+		UMaterialGraphNode* SelectedGraphNode = Cast<UMaterialGraphNode>(SelectedNodes[0]);
+		if (SelectedGraphNode != NULL)
+		{
+			FString DocLink = SelectedGraphNode->GetDocumentationLink();
+			FString DocExcerpt = SelectedGraphNode->GetDocumentationExcerptName();
+
+			DocumentationLink = FEditorClassUtils::GetDocumentationLinkFromExcerpt(DocLink, DocExcerpt);
+		}
+	}
+
+	return DocumentationLink;
+}
+
+void FMaterialEditor::OnGoToDocumentation()
+{
+	FString DocumentationLink = GetDocLinkForSelectedNode();
+	if (!DocumentationLink.IsEmpty())
+	{
+		IDocumentation::Get()->Open(DocumentationLink);
+	}
+}
+
+bool FMaterialEditor::CanGoToDocumentation()
+{
+	FString DocumentationLink = GetDocLinkForSelectedNode();
+	return !DocumentationLink.IsEmpty();
+}
+
 void FMaterialEditor::RenameAssetFromRegistry(const FAssetData& InAddedAssetData, const FString& InNewName)
 {
 	// Grab the asset class, it will be checked for being a material function.
@@ -3514,6 +3551,12 @@ TSharedRef<SGraphEditor> FMaterialEditor::CreateGraphEditorWidget()
 		GraphEditorCommands->MapAction( FMaterialEditorCommands::Get().CreateComponentMaskNode,
 			FExecuteAction::CreateSP(this, &FMaterialEditor::OnCreateComponentMaskNode)
 			);
+
+		GraphEditorCommands->MapAction(FMaterialEditorCommands::Get().GoToDocumentation,
+			FExecuteAction::CreateSP(this, &FMaterialEditor::OnGoToDocumentation),
+			FCanExecuteAction::CreateSP(this, &FMaterialEditor::CanGoToDocumentation)
+			);
+
 	}
 
 	FGraphAppearanceInfo AppearanceInfo;
diff --git a/Engine/Source/Editor/MaterialEditor/Private/MaterialEditor.h b/Engine/Source/Editor/MaterialEditor/Private/MaterialEditor.h
--- a/Engine/Source/Editor/MaterialEditor/Private/MaterialEditor.h
+++ b/Engine/Source/Editor/MaterialEditor/Private/MaterialEditor.h
@@ -525,6 +525,14 @@ private:
 	/** Bring up the search tab */
 	void OnFindInMaterial();
 
+	/** Open documentation for the selected node class */
+	void OnGoToDocumentation();
+	/** Can we open documentation for the selected node */
+	bool CanGoToDocumentation();
+
+	/** Util to try and get doc link for the currently selected node */
+	FString GetDocLinkForSelectedNode();
+
 	/** Callback from the Asset Registry when an asset is renamed. */
 	void RenameAssetFromRegistry(const FAssetData& InAddedAssetData, const FString& InNewName);
 
diff --git a/Engine/Source/Editor/MaterialEditor/Private/MaterialEditorActions.cpp b/Engine/Source/Editor/MaterialEditor/Private/MaterialEditorActions.cpp
--- a/Engine/Source/Editor/MaterialEditor/Private/MaterialEditorActions.cpp
+++ b/Engine/Source/Editor/MaterialEditor/Private/MaterialEditorActions.cpp
@@ -49,8 +49,8 @@ void FMaterialEditorCommands::RegisterCommands()
 	UI_COMMAND( AddToFavorites, "Add To Favorites", "Adds this expression to your favorites.", EUserInterfaceActionType::Button, FInputGesture() );
 	UI_COMMAND( BreakLink, "Break Link", "Deletes this link.", EUserInterfaceActionType::Button, FInputGesture() );
 	UI_COMMAND( ForceRefreshPreviews, "Force Refresh Previews", "Forces a refresh of all previews", EUserInterfaceActionType::Button, FInputGesture(EKeys::SpaceBar) );
-	UI_COMMAND( CreateComponentMaskNode, "Create ComponentMask Node", "Creates a ComponentMask node at the current cursor position.", EUserInterfaceActionType::Button, FInputGesture( EModifierKey::Shift, EKeys::C ) );
-
+	UI_COMMAND( CreateComponentMaskNode, "Create ComponentMask Node", "Creates a ComponentMask node at the current cursor position.", EUserInterfaceActionType::Button, FInputGesture(EModifierKey::Shift, EKeys::C));
+	UI_COMMAND( GoToDocumentation, "View Documentation", "View documentation for this node.", EUserInterfaceActionType::Button, FInputGesture());
 	UI_COMMAND( FindInMaterial, "Search", "Finds expressions and comments in the current Material", EUserInterfaceActionType::Button, FInputGesture(EModifierKey::Control, EKeys::F));
 }
 
diff --git a/Engine/Source/Editor/MaterialEditor/Public/MaterialEditorActions.h b/Engine/Source/Editor/MaterialEditor/Public/MaterialEditorActions.h
--- a/Engine/Source/Editor/MaterialEditor/Public/MaterialEditorActions.h
+++ b/Engine/Source/Editor/MaterialEditor/Public/MaterialEditorActions.h
@@ -152,6 +152,10 @@ public:
 	/** Create component mask node */
 	TSharedPtr< FUICommandInfo > CreateComponentMaskNode;
 
+	/** Go to node documentation */
+	TSharedPtr< FUICommandInfo > GoToDocumentation;
+
+
 	/**
 	 * Initialize commands
 	 */
diff --git a/Engine/Source/Editor/UnrealEd/Private/EditorClassUtils.cpp b/Engine/Source/Editor/UnrealEd/Private/EditorClassUtils.cpp
--- a/Engine/Source/Editor/UnrealEd/Private/EditorClassUtils.cpp
+++ b/Engine/Source/Editor/UnrealEd/Private/EditorClassUtils.cpp
@@ -26,22 +26,18 @@ TSharedRef<SToolTip> FEditorClassUtils::GetTooltip(const UClass* Class, const TA
 	return (Class ? IDocumentation::Get()->CreateToolTip(OverrideText, nullptr, GetDocumentationPage(Class), GetDocumentationExcerpt(Class)) : SNew(SToolTip));
 }
 
-FString FEditorClassUtils::GetDocumentationLink(const UClass* Class, const FString& OverrideExcerpt)
+FString FEditorClassUtils::GetDocumentationLinkFromExcerpt(const FString& DocLink, const FString DocExcerpt)
 {
 	FString DocumentationLink;
-	const FString ClassDocsPage = GetDocumentationPage(Class);
-
 	TSharedRef<IDocumentation> Documentation = IDocumentation::Get();
-	if (Documentation->PageExists(ClassDocsPage))
+	if (Documentation->PageExists(DocLink))
 	{
-		TSharedRef<IDocumentationPage> ClassDocs = Documentation->GetPage(ClassDocsPage, NULL);
-
-		const FString ExcerptSection = (OverrideExcerpt.IsEmpty() ? GetDocumentationExcerpt(Class) : OverrideExcerpt);
+		TSharedRef<IDocumentationPage> ClassDocs = Documentation->GetPage(DocLink, NULL);
 
 		FExcerpt Excerpt;
-		if (ClassDocs->GetExcerpt(ExcerptSection, Excerpt))
+		if (ClassDocs->GetExcerpt(DocExcerpt, Excerpt))
 		{
-			FString* FullDocumentationLink = Excerpt.Variables.Find( TEXT("ToolTipFullLink") );
+			FString* FullDocumentationLink = Excerpt.Variables.Find(TEXT("ToolTipFullLink"));
 			if (FullDocumentationLink)
 			{
 				DocumentationLink = *FullDocumentationLink;
@@ -52,6 +48,16 @@ FString FEditorClassUtils::GetDocumentationLink(const UClass* Class, const FStri
 	return DocumentationLink;
 }
 
+
+FString FEditorClassUtils::GetDocumentationLink(const UClass* Class, const FString& OverrideExcerpt)
+{
+	const FString ClassDocsPage = GetDocumentationPage(Class);
+	const FString ExcerptSection = (OverrideExcerpt.IsEmpty() ? GetDocumentationExcerpt(Class) : OverrideExcerpt);
+
+	return GetDocumentationLinkFromExcerpt(ClassDocsPage, ExcerptSection);
+}
+
+
 TSharedRef<SWidget> FEditorClassUtils::GetDocumentationLinkWidget(const UClass* Class)
 {
 	TSharedRef<SWidget> DocLinkWidget = SNullWidget::NullWidget;
diff --git a/Engine/Source/Editor/UnrealEd/Private/MaterialGraphNode.cpp b/Engine/Source/Editor/UnrealEd/Private/MaterialGraphNode.cpp
--- a/Engine/Source/Editor/UnrealEd/Private/MaterialGraphNode.cpp
+++ b/Engine/Source/Editor/UnrealEd/Private/MaterialGraphNode.cpp
@@ -33,6 +33,7 @@
 #include "MaterialEditorActions.h"
 #include "GraphEditorActions.h"
 #include "GraphEditorSettings.h"
+#include "EditorClassUtils.h"
 
 #define LOCTEXT_NAMESPACE "MaterialGraphNode"
 
@@ -383,6 +384,12 @@ void UMaterialGraphNode::GetContextMenuActions(const FGraphNodeContextMenuBuilde
 		}
 		Context.MenuBuilder->EndSection();
 
+		Context.MenuBuilder->BeginSection("MaterialEditorMenuDocumentation");
+		{
+			Context.MenuBuilder->AddMenuEntry(FMaterialEditorCommands::Get().GoToDocumentation);
+		}
+		Context.MenuBuilder->EndSection();
+
 		// Handle the favorites options
 		if (MaterialExpression)
 		{
diff --git a/Engine/Source/Editor/UnrealEd/Public/EditorClassUtils.h b/Engine/Source/Editor/UnrealEd/Public/EditorClassUtils.h
--- a/Engine/Source/Editor/UnrealEd/Public/EditorClassUtils.h
+++ b/Engine/Source/Editor/UnrealEd/Public/EditorClassUtils.h
@@ -50,6 +50,12 @@ namespace FEditorClassUtils
 	 */
 	UNREALED_API FString GetDocumentationLink(const UClass* Class, const FString& OverrideExcerpt = FString());
 
+
+	/**
+	 * Return link path from a specified excerpt
+	 */
+	UNREALED_API FString GetDocumentationLinkFromExcerpt(const FString& DocLink, const FString DocExcerpt);
+
 	/**
 	 * Creates a link widget to the documentation for a given class
 	 *
