Moved Material Instance parameters category to the top of the Details tab in the Material Editor.  Groups within this category automatically expanded at start.
#ttp 345473 - NUX:Material Instance Editing: The parameters I exposed are collapsed at the bottom of Details, they should be right at top!
#branch UE4
#proj Editor.MaterialEditor

diff --git a/Engine/Source/Editor/MaterialEditor/Private/MaterialEditorInstanceDetailCustomization.cpp b/Engine/Source/Editor/MaterialEditor/Private/MaterialEditorInstanceDetailCustomization.cpp
--- a/Engine/Source/Editor/MaterialEditor/Private/MaterialEditorInstanceDetailCustomization.cpp
+++ b/Engine/Source/Editor/MaterialEditor/Private/MaterialEditorInstanceDetailCustomization.cpp
@@ -44,6 +44,14 @@ void FMaterialInstanceParameterDetails::OnValueCommitted(float NewValue, ETextCo
 
 void FMaterialInstanceParameterDetails::CustomizeDetails( IDetailLayoutBuilder& DetailLayout )
 {
+	// Create a new category for a custom layout for the MIC parameters at the very top
+	FName GroupsCategoryName = TEXT("ParameterGroups");
+	IDetailCategoryBuilder& GroupsCategory = DetailLayout.EditCategory(GroupsCategoryName, NSLOCTEXT("MaterialEditor", "MICParamGroupsTitle", "Parameter Groups").ToString());
+	TSharedRef<IPropertyHandle> ParameterGroupsProperty = DetailLayout.GetProperty("ParameterGroups");
+
+	CreateGroupsWidget(ParameterGroupsProperty, GroupsCategory);
+
+	// Create default category for class properties
 	FName DefaultCategoryName = NAME_None;
 	IDetailCategoryBuilder& DefaultCategory = DetailLayout.EditCategory(DefaultCategoryName);
 
@@ -70,13 +78,6 @@ void FMaterialInstanceParameterDetails::CustomizeDetails( IDetailLayoutBuilder&
 		PropertyRow.Visibility(TAttribute<EVisibility>::Create(TAttribute<EVisibility>::FGetter::CreateSP(this, &FMaterialInstanceParameterDetails::ShouldShowSubsurfaceProfile)));
 	}
 
-	// Create a new category for a custom layout for the MIC parameters
-	FName GroupsCategoryName = TEXT("ParameterGroups");
-	IDetailCategoryBuilder& GroupsCategory = DetailLayout.EditCategory(GroupsCategoryName, NSLOCTEXT("MaterialEditor", "MICParamGroupsTitle", "Parameter Groups").ToString());
-	TSharedRef<IPropertyHandle> ParameterGroupsProperty = DetailLayout.GetProperty("ParameterGroups");
-
-	CreateGroupsWidget(ParameterGroupsProperty, GroupsCategory);
-
 	//////////////////////////////////////////////////////////////////////////
 	DetailLayout.HideProperty("BasePropertyOverrides");
 	IDetailCategoryBuilder& MaterialCategory = DetailLayout.EditCategory(TEXT("MaterialOverrides"), NSLOCTEXT("MaterialEditor", "MICMaterialOverridesTitle", "Material Overrides").ToString());
@@ -92,7 +93,7 @@ void FMaterialInstanceParameterDetails::CreateGroupsWidget(TSharedRef<IPropertyH
 	{
 		FEditorParameterGroup& ParameterGroup = MaterialEditorInstance->ParameterGroups[GroupIdx];
 
-		IDetailGroup& DetailGroup = GroupsCategory.AddGroup( ParameterGroup.GroupName, ParameterGroup.GroupName.ToString() );
+		IDetailGroup& DetailGroup = GroupsCategory.AddGroup( ParameterGroup.GroupName, ParameterGroup.GroupName.ToString(), false, true );
 
 		CreateSingleGroupWidget( ParameterGroup, ParameterGroupsProperty->GetChildHandle(GroupIdx), DetailGroup );
 	}
diff --git a/Engine/Source/Editor/PropertyEditor/Private/DetailCategoryBuilderImpl.cpp b/Engine/Source/Editor/PropertyEditor/Private/DetailCategoryBuilderImpl.cpp
--- a/Engine/Source/Editor/PropertyEditor/Private/DetailCategoryBuilderImpl.cpp
+++ b/Engine/Source/Editor/PropertyEditor/Private/DetailCategoryBuilderImpl.cpp
@@ -121,10 +121,10 @@ void FDetailCategoryImpl::AddCustomBuilder( TSharedRef<IDetailCustomNodeBuilder>
 	AddCustomLayout( NewCustomization, bForAdvanced );
 }
 
-IDetailGroup& FDetailCategoryImpl::AddGroup( FName GroupName, const FString& LocalizedDisplayName, bool bForAdvanced )
+IDetailGroup& FDetailCategoryImpl::AddGroup( FName GroupName, const FString& LocalizedDisplayName, bool bForAdvanced, bool bStartExpanded )
 {
 	FDetailLayoutCustomization NewCustomization;
-	NewCustomization.DetailGroup = MakeShareable( new FDetailGroup( GroupName, AsShared(), LocalizedDisplayName ) );
+	NewCustomization.DetailGroup = MakeShareable( new FDetailGroup( GroupName, AsShared(), LocalizedDisplayName, bStartExpanded ) );
 
 	AddCustomLayout( NewCustomization, bForAdvanced );
 
diff --git a/Engine/Source/Editor/PropertyEditor/Private/DetailCategoryBuilderImpl.h b/Engine/Source/Editor/PropertyEditor/Private/DetailCategoryBuilderImpl.h
--- a/Engine/Source/Editor/PropertyEditor/Private/DetailCategoryBuilderImpl.h
+++ b/Engine/Source/Editor/PropertyEditor/Private/DetailCategoryBuilderImpl.h
@@ -140,7 +140,7 @@ public:
 	virtual IDetailLayoutBuilder& GetParentLayout() const override { return *DetailLayoutBuilder.Pin(); }
 	virtual FDetailWidgetRow& AddCustomRow( const FString& FilterString, bool bForAdvanced = false ) override;
 	virtual void AddCustomBuilder( TSharedRef<IDetailCustomNodeBuilder> InCustomBuilder, bool bForAdvanced = false ) override;
-	virtual IDetailGroup& AddGroup( FName GroupName, const FString& LocalizedDisplayName, bool bForAdvanced = false ) override;
+	virtual IDetailGroup& AddGroup( FName GroupName, const FString& LocalizedDisplayName, bool bForAdvanced = false, bool bStartExpanded = false ) override;
 	virtual void GetDefaultProperties( TArray<TSharedRef<IPropertyHandle> >& OutAllProperties, bool bSimpleProperties = true, bool bAdvancedProperties = true ) override;
 	virtual const FString& GetDisplayName() const override { return DisplayName; }
 
diff --git a/Engine/Source/Editor/PropertyEditor/Private/DetailGroup.cpp b/Engine/Source/Editor/PropertyEditor/Private/DetailGroup.cpp
--- a/Engine/Source/Editor/PropertyEditor/Private/DetailGroup.cpp
+++ b/Engine/Source/Editor/PropertyEditor/Private/DetailGroup.cpp
@@ -6,10 +6,11 @@
 #include "DetailPropertyRow.h"
 #include "DetailItemNode.h"
 
-FDetailGroup::FDetailGroup( const FName InGroupName, TSharedRef<FDetailCategoryImpl> InParentCategory, const FString& InLocalizedDisplayName )
+FDetailGroup::FDetailGroup( const FName InGroupName, TSharedRef<FDetailCategoryImpl> InParentCategory, const FString& InLocalizedDisplayName, const bool bInStartExpanded )
 	: ParentCategory( InParentCategory )
 	, LocalizedDisplayName( InLocalizedDisplayName )
 	, GroupName( InGroupName )
+	, bStartExpanded( bInStartExpanded )
 {
 
 }
diff --git a/Engine/Source/Editor/PropertyEditor/Private/DetailGroup.h b/Engine/Source/Editor/PropertyEditor/Private/DetailGroup.h
--- a/Engine/Source/Editor/PropertyEditor/Private/DetailGroup.h
+++ b/Engine/Source/Editor/PropertyEditor/Private/DetailGroup.h
@@ -5,7 +5,7 @@
 class FDetailGroup : public IDetailGroup, public TSharedFromThis<FDetailGroup>
 {
 public:
-	FDetailGroup( const FName InGroupName, TSharedRef<FDetailCategoryImpl> InParentCategory, const FString& InLocalizedDisplayName );
+	FDetailGroup( const FName InGroupName, TSharedRef<FDetailCategoryImpl> InParentCategory, const FString& InLocalizedDisplayName, const bool bStartExpanded = false );
 
 	/** IDetailGroup interface */     
 	virtual FDetailWidgetRow& HeaderRow() override;
@@ -24,6 +24,9 @@ public:
 	/** @return true if this row should be ticked */
 	bool RequiresTick() const;
 
+	/** @return true is this row should start expanded */
+	bool ShouldStartExpanded() const { return bStartExpanded; }
+
 	/** 
 	 * @return The visibility of this group
 	 */
@@ -72,4 +75,6 @@ private:
 	FString LocalizedDisplayName;
 	/** Name identifier of this group */
 	FName GroupName;
+	/** Whether the detail group should start expanded or not */
+	bool bStartExpanded;
 };
\ No newline at end of file
diff --git a/Engine/Source/Editor/PropertyEditor/Private/DetailItemNode.cpp b/Engine/Source/Editor/PropertyEditor/Private/DetailItemNode.cpp
--- a/Engine/Source/Editor/PropertyEditor/Private/DetailItemNode.cpp
+++ b/Engine/Source/Editor/PropertyEditor/Private/DetailItemNode.cpp
@@ -109,13 +109,19 @@ void FDetailItemNode::InitGroup()
 {
 	Customization.DetailGroup->OnItemNodeInitialized( AsShared(), ParentCategory.Pin().ToSharedRef(), IsParentEnabled );
 
-	// Restore saved expansion state
-	FName GroupName = Customization.DetailGroup->GetGroupName();
-	if( GroupName != NAME_None )
+	if (Customization.DetailGroup->ShouldStartExpanded())
 	{
-		bIsExpanded = ParentCategory.Pin()->GetSavedExpansionState( *this );
+		bIsExpanded = true;
+	}
+	else
+	{
+		// Restore saved expansion state
+		FName GroupName = Customization.DetailGroup->GetGroupName();
+		if (GroupName != NAME_None)
+		{
+			bIsExpanded = ParentCategory.Pin()->GetSavedExpansionState(*this);
+		}
 	}
-
 }
 
 bool FDetailItemNode::HasMultiColumnWidget() const
diff --git a/Engine/Source/Editor/PropertyEditor/Public/DetailCategoryBuilder.h b/Engine/Source/Editor/PropertyEditor/Public/DetailCategoryBuilder.h
--- a/Engine/Source/Editor/PropertyEditor/Public/DetailCategoryBuilder.h
+++ b/Engine/Source/Editor/PropertyEditor/Public/DetailCategoryBuilder.h
@@ -102,8 +102,9 @@ public:
 	 * @param GroupName	The name of the group 
 	 * @param LocalizedDisplayName	The display name of the group
 	 * @param true if the group should appear in the advanced section of the category
+	 * @param true if the group should start expanded
 	 */
-	virtual IDetailGroup& AddGroup( FName GroupName, const FString& LocalizedDisplayName, bool bForAdvanced = false ) = 0;
+	virtual IDetailGroup& AddGroup( FName GroupName, const FString& LocalizedDisplayName, bool bForAdvanced = false, bool bStartExpanded = false ) = 0;
 
 	/**
 	 * Gets the default propeties of this category
