Fixed crash when assigning a material instance parent to an existing child of the material instance.
#ttp 339373 - Crash: when you create a endless loop with material instance that has it self as a parent the editor will become unresponsive
#branch UE4
#proj Editor.MaterialEditor, Editor.UnrealEd, Runtime.Engine

diff --git a/Engine/Source/Editor/MaterialEditor/Private/MaterialEditorInstanceDetailCustomization.cpp b/Engine/Source/Editor/MaterialEditor/Private/MaterialEditorInstanceDetailCustomization.cpp
--- a/Engine/Source/Editor/MaterialEditor/Private/MaterialEditorInstanceDetailCustomization.cpp
+++ b/Engine/Source/Editor/MaterialEditor/Private/MaterialEditorInstanceDetailCustomization.cpp
@@ -13,6 +13,8 @@
 #include "ScopedTransaction.h"
 #include "DetailWidgetRow.h"
 
+#define LOCTEXT_NAMESPACE "MaterialInstanceEditor"
+
 TSharedRef<IDetailCustomization> FMaterialInstanceParameterDetails::MakeInstance(UMaterialEditorInstanceConstant* MaterialInstance, FGetShowHiddenParameters InShowHiddenDelegate)
 {
 	return MakeShareable(new FMaterialInstanceParameterDetails(MaterialInstance, InShowHiddenDelegate));
@@ -42,22 +44,51 @@ void FMaterialInstanceParameterDetails::OnValueCommitted(float NewValue, ETextCo
 	ensure(PropertyHandle->SetValue(NewValue) == FPropertyAccess::Success);
 }
 
-void FMaterialInstanceParameterDetails::CustomizeDetails( IDetailLayoutBuilder& DetailLayout )
+void FMaterialInstanceParameterDetails::CustomizeDetails(IDetailLayoutBuilder& DetailLayout)
 {
 	// Create a new category for a custom layout for the MIC parameters at the very top
 	FName GroupsCategoryName = TEXT("ParameterGroups");
-	IDetailCategoryBuilder& GroupsCategory = DetailLayout.EditCategory(GroupsCategoryName, NSLOCTEXT("MaterialEditor", "MICParamGroupsTitle", "Parameter Groups").ToString());
+	IDetailCategoryBuilder& GroupsCategory = DetailLayout.EditCategory(GroupsCategoryName, LOCTEXT("MICParamGroupsTitle", "Parameter Groups").ToString());
 	TSharedRef<IPropertyHandle> ParameterGroupsProperty = DetailLayout.GetProperty("ParameterGroups");
 
 	CreateGroupsWidget(ParameterGroupsProperty, GroupsCategory);
 
 	// Create default category for class properties
-	FName DefaultCategoryName = NAME_None;
+	const FName DefaultCategoryName = NAME_None;
 	IDetailCategoryBuilder& DefaultCategory = DetailLayout.EditCategory(DefaultCategoryName);
 
-	// Add/hide properties
+	// Add PhysMaterial property
 	DefaultCategory.AddProperty("PhysMaterial");
-	DefaultCategory.AddProperty("Parent");
+
+	// Customize Parent property so we can check for recursively set parents
+	TSharedRef<IPropertyHandle> ParentPropertyHandle = DetailLayout.GetProperty("Parent");
+	IDetailPropertyRow& ParentPropertyRow = DefaultCategory.AddProperty("Parent");
+	TSharedPtr<SWidget> NameWidget;
+	TSharedPtr<SWidget> ValueWidget;
+	FDetailWidgetRow Row;
+	ParentPropertyRow.GetDefaultWidgets(NameWidget, ValueWidget, Row);
+
+	const bool bShowChildren = true;
+	ParentPropertyRow.CustomWidget(bShowChildren)
+		.NameContent()
+		.MinDesiredWidth(Row.NameWidget.MinWidth)
+		.MaxDesiredWidth(Row.NameWidget.MaxWidth)
+		[
+			NameWidget.ToSharedRef()
+		]
+		.ValueContent()
+		.MinDesiredWidth(Row.ValueWidget.MinWidth)
+		.MaxDesiredWidth(Row.ValueWidget.MaxWidth)
+		[
+			SNew(SObjectPropertyEntryBox)
+			.PropertyHandle(ParentPropertyHandle)
+			.AllowedClass(UMaterialInterface::StaticClass())
+			.ThumbnailPool(DetailLayout.GetThumbnailPool())
+			.AllowClear(true)
+			.OnShouldSetAsset(this, &FMaterialInstanceParameterDetails::OnShouldSetAsset)
+		];
+
+	// Add/hide other properties
 	DefaultCategory.AddProperty("LightmassSettings");
 	DetailLayout.HideProperty("bUseOldStyleMICEditorGroups");
 	DetailLayout.HideProperty("ParameterGroups");
@@ -80,7 +111,7 @@ void FMaterialInstanceParameterDetails::CustomizeDetails( IDetailLayoutBuilder&
 
 	//////////////////////////////////////////////////////////////////////////
 	DetailLayout.HideProperty("BasePropertyOverrides");
-	IDetailCategoryBuilder& MaterialCategory = DetailLayout.EditCategory(TEXT("MaterialOverrides"), NSLOCTEXT("MaterialEditor", "MICMaterialOverridesTitle", "Material Overrides").ToString());
+	IDetailCategoryBuilder& MaterialCategory = DetailLayout.EditCategory(TEXT("MaterialOverrides"), LOCTEXT("MICMaterialOverridesTitle", "Material Overrides").ToString());
 	MaterialCategory.AddProperty("bOverrideBaseProperties");
 	MaterialCategory.AddProperty("BasePropertyOverrides");
 }
@@ -269,7 +300,7 @@ bool FMaterialInstanceParameterDetails::IsOverriddenExpression(UDEditorParameter
 
 void FMaterialInstanceParameterDetails::OnOverrideParameter(bool NewValue, class UDEditorParameterValue* Parameter)
 {
-	const FScopedTransaction Transaction( NSLOCTEXT( "MaterialInstanceEditor", "OverrideParameter", "Override Parameter" ) );
+	const FScopedTransaction Transaction( LOCTEXT( "OverrideParameter", "Override Parameter" ) );
 	Parameter->Modify();
 	Parameter->bOverride = NewValue;
 
@@ -279,6 +310,25 @@ void FMaterialInstanceParameterDetails::OnOverrideParameter(bool NewValue, class
 	FEditorSupportDelegates::RedrawAllViewports.Broadcast();
 }
 
+bool FMaterialInstanceParameterDetails::OnShouldSetAsset(const FAssetData& AssetData) const
+{
+	UMaterialInstance* MaterialInstance = Cast<UMaterialInstance>(AssetData.GetAsset());
+
+	if (MaterialInstance != nullptr)
+	{
+		bool bIsChild = MaterialInstance->IsChildOf(MaterialEditorInstance->SourceInstance);
+		if (bIsChild)
+		{
+			FMessageDialog::Open(
+				EAppMsgType::Ok,
+				FText::Format(LOCTEXT("CannotSetExistingChildAsParent", "Cannot set {0} as a parent as it is already a child of this material instance."), FText::FromName(AssetData.AssetName)));
+		}
+		return !bIsChild;
+	}
+
+	return true;
+}
+
 FString FMaterialInstanceParameterDetails::GetParameterExpressionDescription(UDEditorParameterValue* Parameter) const
 {
 	UMaterial* BaseMaterial = MaterialEditorInstance->SourceInstance->GetMaterial();
@@ -308,7 +358,7 @@ FString FMaterialInstanceParameterDetails::GetParameterExpressionDescription(UDE
 
 void FMaterialInstanceParameterDetails::ResetToDefault( class UDEditorParameterValue* Parameter )
 {
-	const FScopedTransaction Transaction( NSLOCTEXT( "MaterialInstanceEditor", "ResetToDefault", "Reset To Default" ) );
+	const FScopedTransaction Transaction( LOCTEXT( "ResetToDefault", "Reset To Default" ) );
 	Parameter->Modify();
 	FName ParameterName = Parameter->ParameterName;
 	
@@ -393,3 +443,5 @@ EVisibility FMaterialInstanceParameterDetails::ShouldShowSubsurfaceProfile() con
 
 	return (Model == MSM_SubsurfaceProfile) ? EVisibility::Visible : EVisibility::Collapsed;
 }
+
+#undef LOCTEXT_NAMESPACE
diff --git a/Engine/Source/Editor/MaterialEditor/Private/MaterialEditorInstanceDetailCustomization.h b/Engine/Source/Editor/MaterialEditor/Private/MaterialEditorInstanceDetailCustomization.h
--- a/Engine/Source/Editor/MaterialEditor/Private/MaterialEditorInstanceDetailCustomization.h
+++ b/Engine/Source/Editor/MaterialEditor/Private/MaterialEditorInstanceDetailCustomization.h
@@ -52,6 +52,9 @@ private:
 	 */
 	void OnOverrideParameter(bool NewValue, class UDEditorParameterValue* Parameter);
 
+	/** Called when an asset is set as a parent */
+	bool OnShouldSetAsset(const FAssetData& InAssetData) const;
+
 	/** Reset to default implementation.  Resets Parameter to default */
 	void ResetToDefault( class UDEditorParameterValue* Parameter );
 
diff --git a/Engine/Source/Editor/PropertyEditor/Private/PropertyCustomizationHelpers.cpp b/Engine/Source/Editor/PropertyEditor/Private/PropertyCustomizationHelpers.cpp
--- a/Engine/Source/Editor/PropertyEditor/Private/PropertyCustomizationHelpers.cpp
+++ b/Engine/Source/Editor/PropertyEditor/Private/PropertyCustomizationHelpers.cpp
@@ -303,6 +303,7 @@ void SObjectPropertyEntryBox::Construct( const FArguments& InArgs )
 {
 	ObjectPath = InArgs._ObjectPath;
 	OnObjectChanged = InArgs._OnObjectChanged;
+	OnShouldSetAsset = InArgs._OnShouldSetAsset;
 
 	bool bDisplayThumbnail = false;
 	FIntPoint ThumbnailSize(64, 64);
@@ -378,16 +379,18 @@ void SObjectPropertyEntryBox::OnSetObject(const FAssetData& AssetData)
 {
 	if( PropertyHandle.IsValid() && PropertyHandle->IsValidHandle() )
 	{
-		FString ObjectPathName = TEXT("None");
-		if(AssetData.IsValid())
+		if (!OnShouldSetAsset.IsBound() || OnShouldSetAsset.Execute(AssetData))
 		{
-			ObjectPathName = AssetData.ObjectPath.ToString();
-		}
+			FString ObjectPathName = TEXT("None");
+			if (AssetData.IsValid())
+			{
+				ObjectPathName = AssetData.ObjectPath.ToString();
+			}
 
-		PropertyHandle->SetValueFromFormattedString(ObjectPathName);
+			PropertyHandle->SetValueFromFormattedString(ObjectPathName);
+			OnObjectChanged.ExecuteIfBound(AssetData);
+		}
 	}
-
-	OnObjectChanged.ExecuteIfBound(AssetData);
 }
 
 void SClassPropertyEntryBox::Construct(const FArguments& InArgs)
diff --git a/Engine/Source/Editor/PropertyEditor/Public/PropertyCustomizationHelpers.h b/Engine/Source/Editor/PropertyEditor/Public/PropertyCustomizationHelpers.h
--- a/Engine/Source/Editor/PropertyEditor/Public/PropertyCustomizationHelpers.h
+++ b/Engine/Source/Editor/PropertyEditor/Public/PropertyCustomizationHelpers.h
@@ -14,6 +14,7 @@
 namespace SceneOutliner { struct FOutlinerFilters; }
 
 DECLARE_DELEGATE_OneParam(FOnAssetSelected, const class FAssetData& /*AssetData*/);
+DECLARE_DELEGATE_RetVal_OneParam(bool, FOnShouldSetAsset, const class FAssetData& /*AssetData*/);
 DECLARE_DELEGATE_RetVal_OneParam(bool, FOnShouldFilterAsset, const class FAssetData& /*AssetData*/);
 DECLARE_DELEGATE_OneParam( FOnGetActorFilters, TSharedPtr<SceneOutliner::FOutlinerFilters>& );
 
@@ -66,6 +67,8 @@ public:
 		SLATE_ARGUMENT( TSharedPtr<FAssetThumbnailPool>, ThumbnailPool )
 		/** Classes that are allowed in the asset picker */
 		SLATE_ARGUMENT( UClass*, AllowedClass )
+		/** Called to check if an asset should be set */
+		SLATE_EVENT(FOnShouldSetAsset, OnShouldSetAsset)
 		/** Called when the object value changes */
 		SLATE_EVENT(FOnSetObject, OnObjectChanged)
 		/** Called to check if an asset is valid to use */
@@ -85,6 +88,8 @@ private:
 	/** @return the object path for the object we are viewing */
 	FString OnGetObjectPath() const;
 private:
+	/** Delegate to call to determine whether the asset should be set */
+	FOnShouldSetAsset OnShouldSetAsset;
 	/** Delegate to call when the object changes */
 	FOnSetObject OnObjectChanged;
 	/** Path to the object */
diff --git a/Engine/Source/Editor/UnrealEd/Classes/MaterialEditor/MaterialEditorInstanceConstant.h b/Engine/Source/Editor/UnrealEd/Classes/MaterialEditor/MaterialEditorInstanceConstant.h
--- a/Engine/Source/Editor/UnrealEd/Classes/MaterialEditor/MaterialEditorInstanceConstant.h
+++ b/Engine/Source/Editor/UnrealEd/Classes/MaterialEditor/MaterialEditorInstanceConstant.h
@@ -209,7 +209,7 @@ class UNREALED_API UMaterialEditorInstanceConstant : public UObject
 	class UPhysicalMaterial* PhysMaterial;
 
 	// since the Parent may point across levels and the property editor needs to import this text, it must be marked lazy so it doesn't set itself to NULL in FindImportedObject
-	UPROPERTY(EditAnywhere, Category=MaterialEditorInstanceConstant)
+	UPROPERTY(EditAnywhere, Category=MaterialEditorInstanceConstant, meta=(DisplayThumbnail="true"))
 	class UMaterialInterface* Parent;
 
 	UPROPERTY(EditAnywhere, editfixedsize, Category=MaterialEditorInstanceConstant)
diff --git a/Engine/Source/Runtime/Engine/Classes/Materials/MaterialInstance.h b/Engine/Source/Runtime/Engine/Classes/Materials/MaterialInstance.h
--- a/Engine/Source/Runtime/Engine/Classes/Materials/MaterialInstance.h
+++ b/Engine/Source/Runtime/Engine/Classes/Materials/MaterialInstance.h
@@ -327,6 +327,14 @@ public:
 	// For all materials instances, UMaterialInstance::CacheResourceShadersForRendering
 	static void AllMaterialsCacheResourceShadersForRendering();
 
+	/** 
+	 * Determine whether this Material Instance is a child of another Material
+	 *
+	 * @param	Material	Material to check against
+	 * @return	true if this Material Instance is a child of the other Material.
+	 */
+	ENGINE_API bool IsChildOf(const UMaterialInterface* Material) const;
+
 protected:
 	/**
 	 * Updates parameter names on the material instance, returns true if parameters have changed.
diff --git a/Engine/Source/Runtime/Engine/Private/Materials/MaterialInstance.cpp b/Engine/Source/Runtime/Engine/Private/Materials/MaterialInstance.cpp
--- a/Engine/Source/Runtime/Engine/Private/Materials/MaterialInstance.cpp
+++ b/Engine/Source/Runtime/Engine/Private/Materials/MaterialInstance.cpp
@@ -1827,9 +1827,18 @@ void UMaterialInstance::SetParentInternal(UMaterialInterface* NewParent)
 {
 	if (Parent == NULL || Parent != NewParent)
 	{
-		if (NewParent &&
-			!NewParent->IsA(UMaterial::StaticClass()) &&
-			!NewParent->IsA(UMaterialInstanceConstant::StaticClass()))
+		// Check if the new parent is already an existing child
+		UMaterialInstance* ParentAsMaterialInstance = Cast<UMaterialInstance>(NewParent);
+
+		if (ParentAsMaterialInstance != nullptr && ParentAsMaterialInstance->IsChildOf(this))
+		{
+			UE_LOG(LogMaterial, Warning, TEXT("%s is not a valid parent for %s as it is already a child of this material instance."),
+				   *NewParent->GetFullName(),
+				   *GetFullName());
+		}
+		else if (NewParent &&
+				 !NewParent->IsA(UMaterial::StaticClass()) &&
+				 !NewParent->IsA(UMaterialInstanceConstant::StaticClass()))
 		{
 			UE_LOG(LogMaterial, Warning, TEXT("%s is not a valid parent for %s. Only Materials and MaterialInstanceConstants are valid parents for a material instance."),
 				*NewParent->GetFullName(),
@@ -2387,6 +2396,21 @@ void UMaterialInstance::AllMaterialsCacheResourceShadersForRendering()
 	}
 }
 
+
+bool UMaterialInstance::IsChildOf(const UMaterialInterface* Parent) const
+{
+	const UMaterialInterface* Material = this;
+
+	while (Material != Parent && Material != nullptr)
+	{
+		const UMaterialInstance* MaterialInstance = Cast<const UMaterialInstance>(Material);
+		Material = (MaterialInstance != nullptr) ? MaterialInstance->Parent : nullptr;
+	}
+
+	return (Material != nullptr);
+}
+
+
 /**
 	Properties of the base material. Can now be overridden by instances.
 */
