Modify CreatePropertyNameWidget so that the tooltip (like with the name) is also overrideable. (and fixedup usage).

diff --git a/Engine/Source/Editor/DetailCustomizations/Private/AttenuationSettingsCustomizations.cpp b/Engine/Source/Editor/DetailCustomizations/Private/AttenuationSettingsCustomizations.cpp
--- a/Engine/Source/Editor/DetailCustomizations/Private/AttenuationSettingsCustomizations.cpp
+++ b/Engine/Source/Editor/DetailCustomizations/Private/AttenuationSettingsCustomizations.cpp
@@ -13,11 +13,12 @@ void FAttenuationSettingsCustomization::CustomizeHeader( TSharedRef<class IPrope
 	// We'll set up reset to default ourselves
 	const bool bDisplayResetToDefault = false;
 	const FString DisplayNameOverride = TEXT("");
+	const FString DisplayToolTipOverride = TEXT("");
 
 	HeaderRow
 		.NameContent()
 		[
-			StructPropertyHandle->CreatePropertyNameWidget( DisplayNameOverride, bDisplayResetToDefault )
+			StructPropertyHandle->CreatePropertyNameWidget( DisplayNameOverride, DisplayToolTipOverride, bDisplayResetToDefault )
 		];
 }
 
diff --git a/Engine/Source/Editor/DetailCustomizations/Private/CurveStructCustomization.cpp b/Engine/Source/Editor/DetailCustomizations/Private/CurveStructCustomization.cpp
--- a/Engine/Source/Editor/DetailCustomizations/Private/CurveStructCustomization.cpp
+++ b/Engine/Source/Editor/DetailCustomizations/Private/CurveStructCustomization.cpp
@@ -52,7 +52,7 @@ void FCurveStructCustomization::CustomizeHeader( TSharedRef<class IPropertyHandl
 		HeaderRow
 			.NameContent()
 			[
-				InStructPropertyHandle->CreatePropertyNameWidget( TEXT( "" ), false )
+				InStructPropertyHandle->CreatePropertyNameWidget( TEXT( "" ), TEXT( "" ), false )
 			]
 			.ValueContent()
 			.HAlign(HAlign_Fill)
@@ -87,7 +87,7 @@ void FCurveStructCustomization::CustomizeHeader( TSharedRef<class IPropertyHandl
 		HeaderRow
 			.NameContent()
 			[
-				InStructPropertyHandle->CreatePropertyNameWidget( TEXT( "" ), false )
+				InStructPropertyHandle->CreatePropertyNameWidget( TEXT( "" ), TEXT( "" ), false )
 			]
 			.ValueContent()
 			[
diff --git a/Engine/Source/Editor/DetailCustomizations/Private/CurveTableCustomization.h b/Engine/Source/Editor/DetailCustomizations/Private/CurveTableCustomization.h
--- a/Engine/Source/Editor/DetailCustomizations/Private/CurveTableCustomization.h
+++ b/Engine/Source/Editor/DetailCustomizations/Private/CurveTableCustomization.h
@@ -21,7 +21,7 @@ public:
 		HeaderRow
 			.NameContent()
 			[
-				InStructPropertyHandle->CreatePropertyNameWidget( TEXT( "" ), false )
+				InStructPropertyHandle->CreatePropertyNameWidget( TEXT( "" ), TEXT( "" ), false )
 			];
 	}
 
diff --git a/Engine/Source/Editor/DetailCustomizations/Private/DataTableCategoryCustomization.h b/Engine/Source/Editor/DetailCustomizations/Private/DataTableCategoryCustomization.h
--- a/Engine/Source/Editor/DetailCustomizations/Private/DataTableCategoryCustomization.h
+++ b/Engine/Source/Editor/DetailCustomizations/Private/DataTableCategoryCustomization.h
@@ -23,7 +23,7 @@ public:
 		HeaderRow
 			.NameContent()
 			[
-				InStructPropertyHandle->CreatePropertyNameWidget(TEXT(""), false)
+				InStructPropertyHandle->CreatePropertyNameWidget(TEXT(""), TEXT( "" ), false)
 			];
 	}
 
diff --git a/Engine/Source/Editor/DetailCustomizations/Private/DataTableCustomization.h b/Engine/Source/Editor/DetailCustomizations/Private/DataTableCustomization.h
--- a/Engine/Source/Editor/DetailCustomizations/Private/DataTableCustomization.h
+++ b/Engine/Source/Editor/DetailCustomizations/Private/DataTableCustomization.h
@@ -23,7 +23,7 @@ public:
 		HeaderRow
 			.NameContent()
 			[
-				InStructPropertyHandle->CreatePropertyNameWidget( TEXT( "" ), false )
+				InStructPropertyHandle->CreatePropertyNameWidget( TEXT( "" ), TEXT( "" ), false )
 			];
 	}
 
diff --git a/Engine/Source/Editor/DetailCustomizations/Private/MathStructCustomizations.cpp b/Engine/Source/Editor/DetailCustomizations/Private/MathStructCustomizations.cpp
--- a/Engine/Source/Editor/DetailCustomizations/Private/MathStructCustomizations.cpp
+++ b/Engine/Source/Editor/DetailCustomizations/Private/MathStructCustomizations.cpp
@@ -39,6 +39,7 @@ void FMathStructCustomization::MakeHeaderRow( TSharedRef<class IPropertyHandle>&
 	// We'll set up reset to default ourselves
 	const bool bDisplayResetToDefault = false;
 	const FString DisplayNameOverride = TEXT("");
+	const FString DisplayToolTipOverride = TEXT("");
 
 	TWeakPtr<IPropertyHandle> StructWeakHandlePtr = StructPropertyHandle;
 
@@ -46,7 +47,7 @@ void FMathStructCustomization::MakeHeaderRow( TSharedRef<class IPropertyHandle>&
 
 	Row.NameContent()
 	[
-		StructPropertyHandle->CreatePropertyNameWidget( DisplayNameOverride, bDisplayResetToDefault )
+		StructPropertyHandle->CreatePropertyNameWidget( DisplayNameOverride, DisplayToolTipOverride, bDisplayResetToDefault )
 	]
 	.ValueContent()
 	// Make enough space for each child handle
@@ -433,10 +434,11 @@ void FColorStructCustomization::MakeHeaderRow( TSharedRef<class IPropertyHandle>
 	// We'll set up reset to default ourselves
 	const bool bDisplayResetToDefault = false;
 	const FString DisplayNameOverride = TEXT("");
+	const FString DisplayToolTipOverride = TEXT("");
 
 	Row.NameContent()
 	[
-		StructPropertyHandle->CreatePropertyNameWidget( DisplayNameOverride, bDisplayResetToDefault )
+		StructPropertyHandle->CreatePropertyNameWidget( DisplayNameOverride, DisplayToolTipOverride, bDisplayResetToDefault )
 	]
 	.ValueContent()
 	.MinDesiredWidth(250.0f)
diff --git a/Engine/Source/Editor/DetailCustomizations/Private/SlateBrushCustomization.cpp b/Engine/Source/Editor/DetailCustomizations/Private/SlateBrushCustomization.cpp
--- a/Engine/Source/Editor/DetailCustomizations/Private/SlateBrushCustomization.cpp
+++ b/Engine/Source/Editor/DetailCustomizations/Private/SlateBrushCustomization.cpp
@@ -1254,7 +1254,7 @@ void FSlateBrushStructCustomization::CustomizeChildren( TSharedRef<IPropertyHand
 				.HeaderRow()
 				.NameContent()
 				[
-					StructPropertyHandle->CreatePropertyNameWidget(NSLOCTEXT("UnrealEd", "Preview", "Preview").ToString(), false)
+					StructPropertyHandle->CreatePropertyNameWidget(NSLOCTEXT("UnrealEd", "Preview", "Preview").ToString(), TEXT( "" ), false)
 				]
 			.ValueContent()
 				.MinDesiredWidth(1)
diff --git a/Engine/Source/Editor/MaterialEditor/Private/MaterialEditorInstanceDetailCustomization.cpp b/Engine/Source/Editor/MaterialEditor/Private/MaterialEditorInstanceDetailCustomization.cpp
--- a/Engine/Source/Editor/MaterialEditor/Private/MaterialEditorInstanceDetailCustomization.cpp
+++ b/Engine/Source/Editor/MaterialEditor/Private/MaterialEditorInstanceDetailCustomization.cpp
@@ -227,7 +227,7 @@ void FMaterialInstanceParameterDetails::CreateMaskParameterValueWidget(UDEditorP
 				.HAlign(HAlign_Left)
 				.AutoWidth()
 				[
-					RMaskProperty->CreatePropertyNameWidget( TEXT(""), false )
+					RMaskProperty->CreatePropertyNameWidget( TEXT(""), TEXT( "" ), false )
 				]
 				+SHorizontalBox::Slot()
 				.HAlign(HAlign_Left)
@@ -240,7 +240,7 @@ void FMaterialInstanceParameterDetails::CreateMaskParameterValueWidget(UDEditorP
 				.Padding( FMargin( 10.0f, 0.0f, 0.0f, 0.0f ) )
 				.AutoWidth()
 				[
-					GMaskProperty->CreatePropertyNameWidget( TEXT(""), false )
+					GMaskProperty->CreatePropertyNameWidget( TEXT(""), TEXT( "" ), false )
 				]
 				+SHorizontalBox::Slot()
 				.HAlign(HAlign_Left)
@@ -253,7 +253,7 @@ void FMaterialInstanceParameterDetails::CreateMaskParameterValueWidget(UDEditorP
 				.Padding( FMargin( 10.0f, 0.0f, 0.0f, 0.0f ) )
 				.AutoWidth()
 				[
-					BMaskProperty->CreatePropertyNameWidget( TEXT(""), false )
+					BMaskProperty->CreatePropertyNameWidget( TEXT(""), TEXT( "" ), false )
 				]
 				+SHorizontalBox::Slot()
 				.HAlign(HAlign_Left)
@@ -266,7 +266,7 @@ void FMaterialInstanceParameterDetails::CreateMaskParameterValueWidget(UDEditorP
 				.Padding( FMargin( 10.0f, 0.0f, 0.0f, 0.0f ) )
 				.AutoWidth()
 				[
-					AMaskProperty->CreatePropertyNameWidget( TEXT(""), false )
+					AMaskProperty->CreatePropertyNameWidget( TEXT(""), TEXT( "" ), false )
 				]
 				+SHorizontalBox::Slot()
 				.HAlign(HAlign_Left)
diff --git a/Engine/Source/Editor/PropertyEditor/Private/PropertyHandleImpl.cpp b/Engine/Source/Editor/PropertyEditor/Private/PropertyHandleImpl.cpp
--- a/Engine/Source/Editor/PropertyEditor/Private/PropertyHandleImpl.cpp
+++ b/Engine/Source/Editor/PropertyEditor/Private/PropertyHandleImpl.cpp
@@ -1529,7 +1529,7 @@ FString FPropertyHandleBase::GeneratePathToProperty() const
 
 }
 
-TSharedRef<SWidget> FPropertyHandleBase::CreatePropertyNameWidget( const FString& NameOverride, bool bDisplayResetToDefault, bool bDisplayText, bool bDisplayThumbnail ) const
+TSharedRef<SWidget> FPropertyHandleBase::CreatePropertyNameWidget( const FString& NameOverride, const FString& ToolTipOverride, bool bDisplayResetToDefault, bool bDisplayText, bool bDisplayThumbnail ) const
 {
 	if( Implementation.IsValid() && Implementation->GetPropertyNode().IsValid() )
 	{
@@ -1538,6 +1538,11 @@ TSharedRef<SWidget> FPropertyHandleBase::CreatePropertyNameWidget( const FString
 			Implementation->GetPropertyNode()->SetDisplayNameOverride( NameOverride );
 		}
 
+		if( ToolTipOverride.Len() > 0 )
+		{
+			Implementation->GetPropertyNode()->SetToolTipOverride( ToolTipOverride );
+		}
+
 		TSharedPtr<FPropertyEditor> PropertyEditor = FPropertyEditor::Create( Implementation->GetPropertyNode().ToSharedRef(), Implementation->GetPropertyUtilities().ToSharedRef() );
 
 		return SNew( SPropertyNameWidget, PropertyEditor )
diff --git a/Engine/Source/Editor/PropertyEditor/Private/PropertyHandleImpl.h b/Engine/Source/Editor/PropertyEditor/Private/PropertyHandleImpl.h
--- a/Engine/Source/Editor/PropertyEditor/Private/PropertyHandleImpl.h
+++ b/Engine/Source/Editor/PropertyEditor/Private/PropertyHandleImpl.h
@@ -374,7 +374,7 @@ public:
 	virtual void MarkHiddenByCustomization() override;
 	virtual bool IsCustomized() const override;
 	virtual FString GeneratePathToProperty() const override;
-	virtual TSharedRef<SWidget> CreatePropertyNameWidget( const FString& NameOverride = TEXT(""), bool bDisplayResetToDefault = false, bool bDisplayText = true, bool bDisplayThumbnail = true ) const override;
+	virtual TSharedRef<SWidget> CreatePropertyNameWidget( const FString& NameOverride = TEXT(""), const FString& ToolTipOverride = TEXT(""), bool bDisplayResetToDefault = false, bool bDisplayText = true, bool bDisplayThumbnail = true ) const override;
 	virtual TSharedRef<SWidget> CreatePropertyValueWidget() const override;
 	virtual bool IsEditConst() const override;
 	virtual void SetOnPropertyValueChanged( const FSimpleDelegate& InOnPropertyValueChanged ) override;
diff --git a/Engine/Source/Editor/PropertyEditor/Public/PropertyCustomizationHelpers.h b/Engine/Source/Editor/PropertyEditor/Public/PropertyCustomizationHelpers.h
--- a/Engine/Source/Editor/PropertyEditor/Public/PropertyCustomizationHelpers.h
+++ b/Engine/Source/Editor/PropertyEditor/Public/PropertyCustomizationHelpers.h
@@ -253,7 +253,7 @@ public:
 			.FilterString(DisplayName.Len() > 0 ? DisplayName : BaseProperty->GetPropertyDisplayName())
 			.NameContent()
 			[
-				BaseProperty->CreatePropertyNameWidget(DisplayName, bDisplayResetToDefaultInNameContent)
+				BaseProperty->CreatePropertyNameWidget(DisplayName, TEXT(""), bDisplayResetToDefaultInNameContent)
 			]
 			.ValueContent()
 			[
diff --git a/Engine/Source/Editor/PropertyEditor/Public/PropertyHandle.h b/Engine/Source/Editor/PropertyEditor/Public/PropertyHandle.h
--- a/Engine/Source/Editor/PropertyEditor/Public/PropertyHandle.h
+++ b/Engine/Source/Editor/PropertyEditor/Public/PropertyHandle.h
@@ -335,13 +335,14 @@ public:
 
 	/**
 	 * Creates a name widget for this property
-	 * @param NameOverride			The name override to use instead of the property name
+	 * @param NameOverride				The name override to use instead of the property name
+	 * @param ToolTipOverride			The tooltip override to use instead of the property name
 	 * @param bDisplayResetToDefault	Whether or not to display the reset to default button
 	 * @param bDisplayText				Whether or not to display the text name of the property
 	 * @param bDisplayThumbnail			Whether or not to display the thumbnail for the property (if any)
 	 * @return the name widget for this property
 	 */
-	virtual TSharedRef<SWidget> CreatePropertyNameWidget( const FString& NameOverride = TEXT(""), bool bDisplayResetToDefault = false, bool bDisplayText = true, bool bDisplayThumbnail = true ) const = 0;
+	virtual TSharedRef<SWidget> CreatePropertyNameWidget( const FString& NameOverride = TEXT(""), const FString& ToolTipOverride = TEXT(""), bool bDisplayResetToDefault = false, bool bDisplayText = true, bool bDisplayThumbnail = true ) const = 0;
 
 	/**
 	 * Creates a value widget for this property
