Material parameters are now allowed in functions
Scalar and Vector Parameter default value changes are applied to materials in the scene in realtime
Code view tab no longer updates or compiles shaders on each edit while closed
Right click menu options to convert between float parameter and constant types

diff --git a/Engine/Source/Editor/MaterialEditor/Private/MaterialEditor.cpp b/Engine/Source/Editor/MaterialEditor/Private/MaterialEditor.cpp
--- a/Engine/Source/Editor/MaterialEditor/Private/MaterialEditor.cpp
+++ b/Engine/Source/Editor/MaterialEditor/Private/MaterialEditor.cpp
@@ -324,6 +324,8 @@ void FMaterialEditor::InitMaterialEditor( const EToolkitMode::Type Mode, const T
 	FMaterialEditorSpawnNodeCommands::Register();
 
 	FEditorSupportDelegates::MaterialUsageFlagsChanged.AddRaw(this, &FMaterialEditor::OnMaterialUsageFlagsChanged);
+	FEditorSupportDelegates::VectorParameterDefaultChanged.AddRaw(this, &FMaterialEditor::OnVectorParameterDefaultChanged);
+	FEditorSupportDelegates::ScalarParameterDefaultChanged.AddRaw(this, &FMaterialEditor::OnScalarParameterDefaultChanged);
 
 	FAssetRegistryModule& AssetRegistryModule = FModuleManager::GetModuleChecked<FAssetRegistryModule>(TEXT("AssetRegistry"));
 	
@@ -528,8 +530,20 @@ FMaterialEditor::FMaterialEditor()
 
 FMaterialEditor::~FMaterialEditor()
 {
+	for (int32 ParameterIndex = 0; ParameterIndex < OverriddenVectorParametersToRevert.Num(); ParameterIndex++)
+	{
+		SetVectorParameterDefaultOnDependentMaterials(OverriddenVectorParametersToRevert[ParameterIndex], FLinearColor::Black, false);
+	}
+
+	for (int32 ParameterIndex = 0; ParameterIndex < OverriddenScalarParametersToRevert.Num(); ParameterIndex++)
+	{
+		SetScalarParameterDefaultOnDependentMaterials(OverriddenScalarParametersToRevert[ParameterIndex], 0, false);
+	}
+
 	// Unregister this delegate
 	FEditorSupportDelegates::MaterialUsageFlagsChanged.RemoveAll(this);
+	FEditorSupportDelegates::VectorParameterDefaultChanged.RemoveAll(this);
+	FEditorSupportDelegates::ScalarParameterDefaultChanged.RemoveAll(this);
 
 	// Null out the expression preview material so they can be GC'ed
 	ExpressionPreviewMaterial = NULL;
@@ -1109,7 +1123,7 @@ void FMaterialEditor::RegenerateCodeView(bool bForce)
 
 	HLSLCode = TEXT("");
 
-	if (!bLivePreview && !bForce)
+	if (!CodeTab.IsValid() || (!bLivePreview && !bForce))
 	{
 		//When bLivePreview is false then the source can be out of date. 
 		return;
@@ -1773,6 +1787,10 @@ void FMaterialEditor::BindCommands()
 		Commands.ConvertToTextureSamples,
 		FExecuteAction::CreateSP(this, &FMaterialEditor::OnConvertTextures));
 
+	ToolkitCommands->MapAction(
+		Commands.ConvertToConstant,
+		FExecuteAction::CreateSP(this, &FMaterialEditor::OnConvertObjects));
+
 	ToolkitCommands->MapAction(
 		Commands.StopPreviewNode,
 		FExecuteAction::CreateSP(this, &FMaterialEditor::OnPreviewNode));
@@ -1960,7 +1978,7 @@ void FMaterialEditor::OnConvertObjects()
 	const FGraphPanelSelectionSet SelectedNodes = GraphEditor->GetSelectedNodes();
 	if (SelectedNodes.Num() > 0)
 	{
-		const FScopedTransaction Transaction( LOCTEXT("MaterialEditorConvert", "Material Editor: Convert to Parameter") );
+		const FScopedTransaction Transaction( LOCTEXT("MaterialEditorConvert", "Material Editor: Convert") );
 		Material->Modify();
 		Material->MaterialGraph->Modify();
 		TArray<class UEdGraphNode*> NodesToDelete;
@@ -1980,6 +1998,8 @@ void FMaterialEditor::OnConvertObjects()
 				UMaterialExpressionTextureSample* TextureSampleExpression = Cast<UMaterialExpressionTextureSample>(CurrentSelectedExpression);
 				UMaterialExpressionComponentMask* ComponentMaskExpression = Cast<UMaterialExpressionComponentMask>(CurrentSelectedExpression);
 				UMaterialExpressionParticleSubUV* ParticleSubUVExpression = Cast<UMaterialExpressionParticleSubUV>(CurrentSelectedExpression);
+				UMaterialExpressionScalarParameter* ScalarParameterExpression = Cast<UMaterialExpressionScalarParameter>(CurrentSelectedExpression);
+				UMaterialExpressionVectorParameter* VectorParameterExpression = Cast<UMaterialExpressionVectorParameter>(CurrentSelectedExpression);
 
 				// Setup the class to convert to
 				UClass* ClassToCreate = NULL;
@@ -2007,6 +2027,14 @@ void FMaterialEditor::OnConvertObjects()
 				{
 					ClassToCreate = UMaterialExpressionStaticComponentMaskParameter::StaticClass();
 				}
+				else if (ScalarParameterExpression)
+				{
+					ClassToCreate = UMaterialExpressionConstant::StaticClass();
+				}
+				else if (VectorParameterExpression)
+				{
+					ClassToCreate = UMaterialExpressionConstant4Vector::StaticClass();
+				}
 
 				if (ClassToCreate)
 				{
@@ -2066,6 +2094,16 @@ void FMaterialEditor::OnConvertObjects()
 							bNeedsRefresh = true;
 							CastChecked<UMaterialExpressionTextureSampleParameterSubUV>(NewExpression)->Texture = ParticleSubUVExpression->Texture;
 						}
+						else if (ScalarParameterExpression)
+						{
+							bNeedsRefresh = true;
+							CastChecked<UMaterialExpressionConstant>(NewExpression)->R = ScalarParameterExpression->DefaultValue;
+						}
+						else if (VectorParameterExpression)
+						{
+							bNeedsRefresh = true;
+							CastChecked<UMaterialExpressionConstant4Vector>(NewExpression)->Constant = VectorParameterExpression->DefaultValue;
+						}
 
 						if (bNeedsRefresh)
 						{
@@ -2403,6 +2441,158 @@ void FMaterialEditor::OnMaterialUsageFlagsChanged(UMaterial* MaterialThatChanged
 	}
 }
 
+void FMaterialEditor::SetVectorParameterDefaultOnDependentMaterials(FName ParameterName, const FLinearColor& Value, bool bOverride)
+{
+	TArray<UMaterial*> MaterialsToOverride;
+
+	if (MaterialFunction)
+	{
+		// Find all materials that reference this function
+		for (TObjectIterator<UMaterial> It; It; ++It)
+		{
+			UMaterial* CurrentMaterial = *It;
+
+			if (CurrentMaterial != Material)
+			{
+				bool bUpdate = false;
+
+				for (int32 FunctionIndex = 0; FunctionIndex < CurrentMaterial->MaterialFunctionInfos.Num(); FunctionIndex++)
+				{
+					if (CurrentMaterial->MaterialFunctionInfos[FunctionIndex].Function == MaterialFunction->ParentFunction)
+					{
+						bUpdate = true;
+						break;
+					}
+				}
+
+				if (bUpdate)
+				{
+					MaterialsToOverride.Add(CurrentMaterial);
+				}
+			}
+		}
+	}
+	else
+	{
+		MaterialsToOverride.Add(OriginalMaterial);
+	}
+
+	const ERHIFeatureLevel::Type FeatureLevel = GEditor->GetEditorWorldContext().World()->FeatureLevel;
+
+	for (int32 MaterialIndex = 0; MaterialIndex < MaterialsToOverride.Num(); MaterialIndex++)
+	{
+		UMaterial* CurrentMaterial = MaterialsToOverride[MaterialIndex];
+
+		CurrentMaterial->OverrideVectorParameterDefault(ParameterName, Value, bOverride, FeatureLevel);
+	}
+
+	// Update MI's that reference any of the materials affected
+	for (TObjectIterator<UMaterialInstance> It; It; ++It)
+	{
+		UMaterialInstance* CurrentMaterialInstance = *It;
+
+		// Only care about MI's with static parameters, because we are overriding parameter defaults, 
+		// And only MI's with static parameters contain uniform expressions, which contain parameter defaults
+		if (CurrentMaterialInstance->bHasStaticPermutationResource)
+		{
+			UMaterial* BaseMaterial = CurrentMaterialInstance->GetMaterial();
+
+			if (MaterialsToOverride.Contains(BaseMaterial))
+			{
+				CurrentMaterialInstance->OverrideVectorParameterDefault(ParameterName, Value, bOverride, FeatureLevel);
+			}
+		}
+	}
+}
+
+void FMaterialEditor::OnVectorParameterDefaultChanged(class UMaterialExpression* Expression, FName ParameterName, const FLinearColor& Value)
+{
+	check(Expression);
+
+	if (Expression->Material == Material && OriginalMaterial)
+	{
+		SetVectorParameterDefaultOnDependentMaterials(ParameterName, Value, true);
+
+		OverriddenVectorParametersToRevert.AddUnique(ParameterName);
+	}
+}
+
+void FMaterialEditor::SetScalarParameterDefaultOnDependentMaterials(FName ParameterName, float Value, bool bOverride)
+{
+	TArray<UMaterial*> MaterialsToOverride;
+
+	if (MaterialFunction)
+	{
+		// Find all materials that reference this function
+		for (TObjectIterator<UMaterial> It; It; ++It)
+		{
+			UMaterial* CurrentMaterial = *It;
+
+			if (CurrentMaterial != Material)
+			{
+				bool bUpdate = false;
+
+				for (int32 FunctionIndex = 0; FunctionIndex < CurrentMaterial->MaterialFunctionInfos.Num(); FunctionIndex++)
+				{
+					if (CurrentMaterial->MaterialFunctionInfos[FunctionIndex].Function == MaterialFunction->ParentFunction)
+					{
+						bUpdate = true;
+						break;
+					}
+				}
+
+				if (bUpdate)
+				{
+					MaterialsToOverride.Add(CurrentMaterial);
+				}
+			}
+		}
+	}
+	else
+	{
+		MaterialsToOverride.Add(OriginalMaterial);
+	}
+
+	const ERHIFeatureLevel::Type FeatureLevel = GEditor->GetEditorWorldContext().World()->FeatureLevel;
+
+	for (int32 MaterialIndex = 0; MaterialIndex < MaterialsToOverride.Num(); MaterialIndex++)
+	{
+		UMaterial* CurrentMaterial = MaterialsToOverride[MaterialIndex];
+
+		CurrentMaterial->OverrideScalarParameterDefault(ParameterName, Value, bOverride, FeatureLevel);
+	}
+
+	// Update MI's that reference any of the materials affected
+	for (TObjectIterator<UMaterialInstance> It; It; ++It)
+	{
+		UMaterialInstance* CurrentMaterialInstance = *It;
+
+		// Only care about MI's with static parameters, because we are overriding parameter defaults, 
+		// And only MI's with static parameters contain uniform expressions, which contain parameter defaults
+		if (CurrentMaterialInstance->bHasStaticPermutationResource)
+		{
+			UMaterial* BaseMaterial = CurrentMaterialInstance->GetMaterial();
+
+			if (MaterialsToOverride.Contains(BaseMaterial))
+			{
+				CurrentMaterialInstance->OverrideScalarParameterDefault(ParameterName, Value, bOverride, FeatureLevel);
+			}
+		}
+	}
+}
+
+void FMaterialEditor::OnScalarParameterDefaultChanged(class UMaterialExpression* Expression, FName ParameterName, float Value)
+{
+	check(Expression);
+
+	if (Expression->Material == Material && OriginalMaterial)
+	{
+		SetScalarParameterDefaultOnDependentMaterials(ParameterName, Value, true);
+
+		OverriddenScalarParametersToRevert.AddUnique(ParameterName);
+	}
+}
+
 TSharedRef<SDockTab> FMaterialEditor::SpawnTab_Preview(const FSpawnTabArgs& Args)
 {
 	TSharedRef<SDockTab> SpawnedTab =
@@ -2466,6 +2656,8 @@ TSharedRef<SDockTab> FMaterialEditor::SpawnTab_HLSLCode(const FSpawnTabArgs& Arg
 			]
 		];
 
+	CodeTab = SpawnedTab;
+
 	RegenerateCodeView();
 
 	return SpawnedTab;
@@ -2617,12 +2809,6 @@ UMaterialExpression* FMaterialEditor::CreateNewMaterialExpression(UClass* NewExp
 		Material->Expressions.Add( NewExpression );
 		NewExpression->Material = Material;
 
-		if (MaterialFunction != NULL)
-		{
-			// Parameters currently not supported in material functions
-			check(!NewExpression->bIsParameterExpression);
-		}
-
 		// Set the expression location.
 		NewExpression->MaterialExpressionEditorX = NodePos.X;
 		NewExpression->MaterialExpressionEditorY = NodePos.Y;
@@ -2733,7 +2919,7 @@ UMaterialExpression* FMaterialEditor::CreateNewMaterialExpression(UClass* NewExp
 			DynamicExpression->UpdateDynamicParameterNames();
 		}
 
-		Material->AddExpressionParameter(NewExpression);
+		Material->AddExpressionParameter(NewExpression, Material->EditorParameters);
 
 		if (NewExpression)
 		{
@@ -3452,7 +3638,8 @@ void FMaterialEditor::OnColorPickerCommitted(FLinearColor LinearColor)
 	if( Object )
 	{
 		Object->MarkPackageDirty();
-		Object->PostEditChange();
+		FPropertyChangedEvent Event(ColorPickerProperty.Get(false));
+		Object->PostEditChangeProperty(Event);
 	}
 
 	NotifyPostChange(NULL,NULL);
@@ -3522,6 +3709,10 @@ TSharedRef<SGraphEditor> FMaterialEditor::CreateGraphEditorWidget()
 			FExecuteAction::CreateSP(this, &FMaterialEditor::OnConvertTextures)
 			);
 
+		GraphEditorCommands->MapAction( FMaterialEditorCommands::Get().ConvertToConstant,
+			FExecuteAction::CreateSP(this, &FMaterialEditor::OnConvertObjects)
+			);
+
 		GraphEditorCommands->MapAction( FMaterialEditorCommands::Get().StopPreviewNode,
 			FExecuteAction::CreateSP(this, &FMaterialEditor::OnPreviewNode)
 			);
@@ -3778,6 +3969,9 @@ void FMaterialEditor::OnNodeDoubleClicked(class UEdGraphNode* Node)
 
 		FColorChannels ChannelEditStruct;
 
+		// Reset to default
+		ColorPickerProperty = NULL;
+
 		if( Constant3Expression )
 		{
 			ChannelEditStruct.Red = &Constant3Expression->Constant.R;
@@ -3804,6 +3998,9 @@ void FMaterialEditor::OnNodeDoubleClicked(class UEdGraphNode* Node)
 			ChannelEditStruct.Green = &VectorExpression->DefaultValue.G;
 			ChannelEditStruct.Blue = &VectorExpression->DefaultValue.B;
 			ChannelEditStruct.Alpha = &VectorExpression->DefaultValue.A;
+			static FName DefaultValueName = FName(TEXT("DefaultValue"));
+			// Store off the property the color picker will be manipulating, so we can construct a useful PostEditChangeProperty later
+			ColorPickerProperty = VectorExpression->GetClass()->FindPropertyByName(DefaultValueName);
 		}
 
 		if (ChannelEditStruct.Red || ChannelEditStruct.Green || ChannelEditStruct.Blue || ChannelEditStruct.Alpha)
diff --git a/Engine/Source/Editor/MaterialEditor/Private/MaterialEditor.h b/Engine/Source/Editor/MaterialEditor/Private/MaterialEditor.h
--- a/Engine/Source/Editor/MaterialEditor/Private/MaterialEditor.h
+++ b/Engine/Source/Editor/MaterialEditor/Private/MaterialEditor.h
@@ -539,6 +539,12 @@ private:
 	/** Callback to tell the Material Editor that a materials usage flags have been changed */
 	void OnMaterialUsageFlagsChanged(class UMaterial* MaterialThatChanged, int32 FlagThatChanged);
 
+	void OnVectorParameterDefaultChanged(class UMaterialExpression*, FName ParameterName, const FLinearColor& Value);
+	void OnScalarParameterDefaultChanged(class UMaterialExpression*, FName ParameterName, float Value);
+
+	void SetVectorParameterDefaultOnDependentMaterials(FName ParameterName, const FLinearColor& Value, bool bOverride);
+	void SetScalarParameterDefaultOnDependentMaterials(FName ParameterName, float, bool bOverride);
+
 	// FEditorUndoClient Interface
 	virtual void PostUndo(bool bSuccess) override;
 	virtual void PostRedo(bool bSuccess) override { PostUndo(bSuccess); }
@@ -574,6 +580,7 @@ private:
 
 	/** Pointer to the object that the current color picker is working on. Can be NULL and stale. */
 	TWeakObjectPtr<UObject> ColorPickerObject;
+	TWeakObjectPtr<UProperty> ColorPickerProperty;
 
 	/** Called before the color picker commits a change. */
 	void PreColorPickerCommit(FLinearColor LinearColor);
@@ -632,6 +639,9 @@ private:
 	/** Cached Code for the widget */
 	FString HLSLCode;
 
+	/** Tracks whether the code tab is open, so we don't have to update it when closed. */
+	TWeakPtr<SDockTab> CodeTab;
+
 	/** Palette of Material Expressions and functions */
 	TSharedPtr<class SMaterialPalette> Palette;
 
@@ -649,11 +659,14 @@ private:
 	bool bAlwaysRefreshAllPreviews;
 
 	/** Material expression previews. */
-	TIndirectArray<class FMatExpressionPreview>		ExpressionPreviews;
+	TIndirectArray<class FMatExpressionPreview> ExpressionPreviews;
 
 	/** Information about material to show when stats are enabled */
 	TArray<TSharedPtr<FMaterialInfo>> MaterialInfoList;
 
+	TArray<FName> OverriddenVectorParametersToRevert;
+	TArray<FName> OverriddenScalarParametersToRevert;
+
 	/** If true, don't render connectors that are not connected to anything. */
 	bool bHideUnusedConnectors;
 
diff --git a/Engine/Source/Editor/MaterialEditor/Private/MaterialEditorActions.cpp b/Engine/Source/Editor/MaterialEditor/Private/MaterialEditorActions.cpp
--- a/Engine/Source/Editor/MaterialEditor/Private/MaterialEditorActions.cpp
+++ b/Engine/Source/Editor/MaterialEditor/Private/MaterialEditorActions.cpp
@@ -34,6 +34,7 @@ void FMaterialEditorCommands::RegisterCommands()
 	UI_COMMAND( MatertialPasteHere, "Paste Here", "Pastes copied items at this location.", EUserInterfaceActionType::Button, FInputGesture() );
 	UI_COMMAND( UseCurrentTexture, "Use Current Texture", "Uses the current texture selected in the content browser.", EUserInterfaceActionType::Button, FInputGesture() );
 	UI_COMMAND( ConvertObjects, "Convert to Parameter", "Converts the objects to parameters.", EUserInterfaceActionType::Button, FInputGesture() );
+	UI_COMMAND( ConvertToConstant, "Convert to Constant", "Converts the parameters to constants.", EUserInterfaceActionType::Button, FInputGesture() );
 	UI_COMMAND( ConvertToTextureObjects, "Convert to Texture Object", "Converts the objects to texture objects.", EUserInterfaceActionType::Button, FInputGesture() );
 	UI_COMMAND( ConvertToTextureSamples, "Convert to Texture Sample", "Converts the objects to texture samples.", EUserInterfaceActionType::Button, FInputGesture() );
 	UI_COMMAND( StopPreviewNode, "Stop Previewing Node", "Stops the preview viewport from previewing this node", EUserInterfaceActionType::Button, FInputGesture() );
diff --git a/Engine/Source/Editor/MaterialEditor/Private/MaterialEditorUtilities.cpp b/Engine/Source/Editor/MaterialEditor/Private/MaterialEditorUtilities.cpp
--- a/Engine/Source/Editor/MaterialEditor/Private/MaterialEditorUtilities.cpp
+++ b/Engine/Source/Editor/MaterialEditor/Private/MaterialEditorUtilities.cpp
@@ -194,11 +194,13 @@ void FMaterialEditorUtilities::UpdateSearchResults(const class UEdGraph* Graph)
 /////////////////////////////////////////////////////
 // Static functions moved from SMaterialEditorCanvas
 
-void FMaterialEditorUtilities::GetVisibleMaterialParameters(const UMaterial *Material, UMaterialInstance *MaterialInstance, TArray<FGuid> &VisibleExpressions)
+void FMaterialEditorUtilities::GetVisibleMaterialParameters(const UMaterial* Material, UMaterialInstance* MaterialInstance, TArray<FGuid>& VisibleExpressions)
 {
 	VisibleExpressions.Empty();
 
-	TArray<UMaterialExpression*> ProcessedExpressions;
+	TScopedPointer<FGetVisibleMaterialParametersFunctionState> FunctionState(new FGetVisibleMaterialParametersFunctionState(NULL));
+	TArray<FGetVisibleMaterialParametersFunctionState*> FunctionStack;
+	FunctionStack.Push(FunctionState.GetOwnedPointer());
 
 	for(uint32 i = 0; i < MP_MAX; ++i)
 	{
@@ -206,12 +208,12 @@ void FMaterialEditorUtilities::GetVisibleMaterialParameters(const UMaterial *Mat
 
 		if(ExpressionInput)
 		{
-			GetVisibleMaterialParametersFromExpression(ExpressionInput->Expression, MaterialInstance, VisibleExpressions, ProcessedExpressions);
+			GetVisibleMaterialParametersFromExpression(FMaterialExpressionKey(ExpressionInput->Expression, ExpressionInput->OutputIndex), MaterialInstance, VisibleExpressions, FunctionStack);
 		}
 	}
 }
 
-bool FMaterialEditorUtilities::GetStaticSwitchExpressionValue(UMaterialInstance* MaterialInstance, UMaterialExpression *SwitchValueExpression, bool& OutValue, FGuid& OutExpressionID, const TArray<FFunctionExpressionInput>* FunctionInputs)
+bool FMaterialEditorUtilities::GetStaticSwitchExpressionValue(UMaterialInstance* MaterialInstance, UMaterialExpression* SwitchValueExpression, bool& OutValue, FGuid& OutExpressionID, const TArray<FFunctionExpressionInput>* FunctionInputs)
 {
 	// If switch value is a function input expression then we must recursively find the associated input expressions from the parent function/material to evaluate the value.
 	UMaterialExpressionFunctionInput* FunctionInputExpression =  Cast<UMaterialExpressionFunctionInput>(SwitchValueExpression);
@@ -229,8 +231,7 @@ bool FMaterialEditorUtilities::GetStaticSwitchExpressionValue(UMaterialInstance*
 		}
 	}
 
-	// The expression can only be a static bool parameter when the current scope is not within a function
-	if(SwitchValueExpression && SwitchValueExpression->Material != NULL)
+	if(SwitchValueExpression)
 	{
 		UMaterialExpressionStaticBoolParameter* SwitchParamValue = Cast<UMaterialExpressionStaticBoolParameter>(SwitchValueExpression);
 
@@ -338,37 +339,43 @@ void FMaterialEditorUtilities::InitExpressions(UMaterial* Material)
 ///////////
 // private
 
-void FMaterialEditorUtilities::GetVisibleMaterialParametersFromExpression(UMaterialExpression *MaterialExpression, UMaterialInstance *MaterialInstance, TArray<FGuid> &VisibleExpressions, TArray<UMaterialExpression*> &ProcessedExpressions, const TArray<FFunctionExpressionInput>* FunctionInputs, TArray<FGuid>* VisibleFunctionInputExpressions)
+void FMaterialEditorUtilities::GetVisibleMaterialParametersFromExpression(
+	FMaterialExpressionKey MaterialExpressionKey, 
+	UMaterialInstance* MaterialInstance, 
+	TArray<FGuid>& VisibleExpressions, 
+	TArray<FGetVisibleMaterialParametersFunctionState*>& FunctionStack)
 {
-	if(!MaterialExpression)
+	if (!MaterialExpressionKey.Expression)
 	{
 		return;
 	}
 
 	check(MaterialInstance);
 
-	//don't allow re-entrant expressions to continue
-	if (ProcessedExpressions.Contains(MaterialExpression))
+	// Bail if we already parsed this expression
+	if (FunctionStack.Top()->VisitedExpressions.Contains(MaterialExpressionKey))
 	{
 		return;
 	}
-	ProcessedExpressions.Push(MaterialExpression);
 
-	if(MaterialExpression->Material != NULL)
+	FunctionStack.Top()->VisitedExpressions.Add(MaterialExpressionKey);
+	FunctionStack.Top()->ExpressionStack.Push(MaterialExpressionKey);
+	const int32 FunctionDepth = FunctionStack.Num();
+
 	{
 		// if it's a material parameter it must be visible so add it to the map
-		UMaterialExpressionParameter *Param = Cast<UMaterialExpressionParameter>( MaterialExpression );
-		UMaterialExpressionTextureSampleParameter *TexParam = Cast<UMaterialExpressionTextureSampleParameter>( MaterialExpression );
-		UMaterialExpressionFontSampleParameter *FontParam = Cast<UMaterialExpressionFontSampleParameter>( MaterialExpression );
-		if( Param )
+		UMaterialExpressionParameter* Param = Cast<UMaterialExpressionParameter>( MaterialExpressionKey.Expression );
+		UMaterialExpressionTextureSampleParameter* TexParam = Cast<UMaterialExpressionTextureSampleParameter>( MaterialExpressionKey.Expression );
+		UMaterialExpressionFontSampleParameter* FontParam = Cast<UMaterialExpressionFontSampleParameter>( MaterialExpressionKey.Expression );
+		if (Param)
 		{
 			VisibleExpressions.AddUnique(Param->ExpressionGUID);
 
-			UMaterialExpressionScalarParameter *ScalarParam = Cast<UMaterialExpressionScalarParameter>( MaterialExpression );
-			UMaterialExpressionVectorParameter *VectorParam = Cast<UMaterialExpressionVectorParameter>( MaterialExpression );
+			UMaterialExpressionScalarParameter* ScalarParam = Cast<UMaterialExpressionScalarParameter>( MaterialExpressionKey.Expression );
+			UMaterialExpressionVectorParameter* VectorParam = Cast<UMaterialExpressionVectorParameter>( MaterialExpressionKey.Expression );
 			TArray<FName> Names;
 			TArray<FGuid> Ids;
-			if( ScalarParam )
+			if (ScalarParam)
 			{
 				MaterialInstance->GetMaterial()->GetAllScalarParameterNames( Names, Ids );
 				for( int32 i = 0; i < Names.Num(); i++ )
@@ -379,7 +386,7 @@ void FMaterialEditorUtilities::GetVisibleMaterialParametersFromExpression(UMater
 					}
 				}
 			}
-			else if ( VectorParam )
+			else if (VectorParam)
 			{
 				MaterialInstance->GetMaterial()->GetAllVectorParameterNames( Names, Ids );
 				for( int32 i = 0; i < Names.Num(); i++ )
@@ -391,7 +398,7 @@ void FMaterialEditorUtilities::GetVisibleMaterialParametersFromExpression(UMater
 				}
 			}
 		}
-		else if(TexParam)
+		else if (TexParam)
 		{
 			VisibleExpressions.AddUnique( TexParam->ExpressionGUID );
 			TArray<FName> Names;
@@ -405,7 +412,7 @@ void FMaterialEditorUtilities::GetVisibleMaterialParametersFromExpression(UMater
 				}
 			}
 		}
-		else if(FontParam)
+		else if (FontParam)
 		{
 			VisibleExpressions.AddUnique( FontParam->ExpressionGUID );
 			TArray<FName> Names;
@@ -421,121 +428,97 @@ void FMaterialEditorUtilities::GetVisibleMaterialParametersFromExpression(UMater
 		}
 	}
 
-	UMaterialExpressionFunctionInput* InputExpression = Cast<UMaterialExpressionFunctionInput>(MaterialExpression);
-	if(InputExpression)
-	{
-		VisibleFunctionInputExpressions->AddUnique(InputExpression->Id);
-	}
-
 	// check if it's a switch expression and branch according to its value
-	UMaterialExpressionStaticSwitchParameter* StaticSwitchParamExpression = Cast<UMaterialExpressionStaticSwitchParameter>(MaterialExpression);
-	UMaterialExpressionStaticSwitch* StaticSwitchExpression = Cast<UMaterialExpressionStaticSwitch>(MaterialExpression);
-	UMaterialExpressionMaterialFunctionCall* FunctionCallExpression = Cast<UMaterialExpressionMaterialFunctionCall>(MaterialExpression);
-	if(StaticSwitchParamExpression)
+	UMaterialExpressionStaticSwitchParameter* StaticSwitchParamExpression = Cast<UMaterialExpressionStaticSwitchParameter>(MaterialExpressionKey.Expression);
+	UMaterialExpressionStaticSwitch* StaticSwitchExpression = Cast<UMaterialExpressionStaticSwitch>(MaterialExpressionKey.Expression);
+	UMaterialExpressionMaterialFunctionCall* FunctionCallExpression = Cast<UMaterialExpressionMaterialFunctionCall>(MaterialExpressionKey.Expression);
+	UMaterialExpressionFunctionInput* FunctionInputExpression = Cast<UMaterialExpressionFunctionInput>(MaterialExpressionKey.Expression);
+
+	if (StaticSwitchParamExpression)
 	{
 		bool Value = false;
 		FGuid ExpressionID;
 		MaterialInstance->GetStaticSwitchParameterValue(StaticSwitchParamExpression->ParameterName, Value, ExpressionID);
 		VisibleExpressions.AddUnique(ExpressionID);
 
-		if(Value)
+		if (Value)
 		{
-			GetVisibleMaterialParametersFromExpression(StaticSwitchParamExpression->A.Expression, MaterialInstance, VisibleExpressions, ProcessedExpressions);
+			GetVisibleMaterialParametersFromExpression(FMaterialExpressionKey(StaticSwitchParamExpression->A.Expression, StaticSwitchParamExpression->A.OutputIndex), MaterialInstance, VisibleExpressions, FunctionStack);
 		}
 		else
 		{
-			GetVisibleMaterialParametersFromExpression(StaticSwitchParamExpression->B.Expression, MaterialInstance, VisibleExpressions, ProcessedExpressions);
+			GetVisibleMaterialParametersFromExpression(FMaterialExpressionKey(StaticSwitchParamExpression->B.Expression, StaticSwitchParamExpression->B.OutputIndex), MaterialInstance, VisibleExpressions, FunctionStack);
 		}
 	}
-	else if(StaticSwitchExpression)
+	else if (StaticSwitchExpression)
 	{
 		bool bValue = StaticSwitchExpression->DefaultValue;
 		FGuid ExpressionID;
-		if(StaticSwitchExpression->Value.Expression)
-		{
 
-			GetStaticSwitchExpressionValue(MaterialInstance, StaticSwitchExpression->Value.Expression, bValue, ExpressionID, FunctionInputs );
+		if (StaticSwitchExpression->Value.Expression)
+		{
+			const TArray<FFunctionExpressionInput>* FunctionInputs = FunctionStack.Top()->FunctionCall ? &(FunctionStack.Top()->FunctionCall->FunctionInputs) : NULL;
+			GetStaticSwitchExpressionValue(MaterialInstance, StaticSwitchExpression->Value.Expression, bValue, ExpressionID, FunctionInputs);
 
-			// Flag the switch value as a visible expression
-			UMaterialExpressionFunctionInput* FunctionInputExpression =  Cast<UMaterialExpressionFunctionInput>(StaticSwitchExpression->Value.Expression);
-			if(FunctionInputExpression)
+			if (ExpressionID.IsValid())
 			{
-				VisibleFunctionInputExpressions->AddUnique(FunctionInputExpression->Id);
-			}
-			else if (StaticSwitchExpression && StaticSwitchExpression->Material != NULL)
-			{
-				if(ExpressionID.IsValid())
-				{
-					VisibleExpressions.AddUnique(ExpressionID);
-				}
+				VisibleExpressions.AddUnique(ExpressionID);
 			}
 		}
 
 		if(bValue)
 		{
-			GetVisibleMaterialParametersFromExpression(StaticSwitchExpression->A.Expression, MaterialInstance, VisibleExpressions, ProcessedExpressions, FunctionInputs, VisibleFunctionInputExpressions);
+			GetVisibleMaterialParametersFromExpression(FMaterialExpressionKey(StaticSwitchExpression->A.Expression, StaticSwitchExpression->A.OutputIndex), MaterialInstance, VisibleExpressions, FunctionStack);
 		}
 		else
 		{
-			GetVisibleMaterialParametersFromExpression(StaticSwitchExpression->B.Expression, MaterialInstance, VisibleExpressions, ProcessedExpressions, FunctionInputs, VisibleFunctionInputExpressions);
+			GetVisibleMaterialParametersFromExpression(FMaterialExpressionKey(StaticSwitchExpression->B.Expression, StaticSwitchExpression->B.OutputIndex), MaterialInstance, VisibleExpressions, FunctionStack);
 		}
 	}
-	else if(FunctionCallExpression)
+	else if (FunctionCallExpression)
 	{
-		// Only process functions if any of the dependent functions contain static switches that could potentially cull the material parameters.
-		if(IsFunctionContainingSwitchExpressions(FunctionCallExpression->MaterialFunction))
+		if (FunctionCallExpression->MaterialFunction)
 		{
-
-			// Pass an accumalated array of FFunctionExpressionInput so that input expressions can be matched to the input from the parent.			
-			TArray<FGuid> VisibleFunctionInputs;
-			TArray<FFunctionExpressionInput> FunctionCallInputs = FunctionCallExpression->FunctionInputs;
-			if(FunctionInputs)
-			{
-				FunctionCallInputs.Append(*FunctionInputs);
-			}
-
-			for(int32 FunctionOutputIndex = 0; FunctionOutputIndex < FunctionCallExpression->FunctionOutputs.Num(); FunctionOutputIndex++)
+			for (int32 FunctionCallIndex = 0; FunctionCallIndex < FunctionStack.Num(); FunctionCallIndex++)
 			{
-				// Recurse material functions. Returns an array of input ids that are visible and have not been culled by static switches.
-				GetVisibleMaterialParametersFromExpression( FunctionCallExpression->FunctionOutputs[FunctionOutputIndex].ExpressionOutput->A.Expression, MaterialInstance, VisibleExpressions, ProcessedExpressions, &FunctionCallInputs, &VisibleFunctionInputs);
+				checkSlow(FunctionStack[FunctionCallIndex]->FunctionCall != FunctionCallExpression);
 			}
 
-			// Parse children of function call inputs that have not been culled 
-			for(int32 ExpressionInputIndex = 0; ExpressionInputIndex < FunctionCallExpression->FunctionInputs.Num(); ExpressionInputIndex++)
-			{
-				const FFunctionExpressionInput& Input = FunctionCallExpression->FunctionInputs[ExpressionInputIndex];
-				if(VisibleFunctionInputs.Contains(Input.ExpressionInputId)) 
-				{
-					// Retrieve the expression input and then start parsing its children
-					GetVisibleMaterialParametersFromExpression(Input.Input.Expression, MaterialInstance, VisibleExpressions, ProcessedExpressions, FunctionInputs, VisibleFunctionInputExpressions );
-				}
-			}
-		}
-		else
-		{
-			const TArray<FExpressionInput*>& ExpressionInputs = MaterialExpression->GetInputs();
-			for(int32 ExpressionInputIndex = 0; ExpressionInputIndex < ExpressionInputs.Num(); ExpressionInputIndex++)
-			{
-				//retrieve the expression input and then start parsing its children
-				FExpressionInput* Input = ExpressionInputs[ExpressionInputIndex];
-				GetVisibleMaterialParametersFromExpression(Input->Expression, MaterialInstance, VisibleExpressions, ProcessedExpressions, FunctionInputs, VisibleFunctionInputExpressions);
-			}
+			TScopedPointer<FGetVisibleMaterialParametersFunctionState> NewFunctionState(new FGetVisibleMaterialParametersFunctionState(FunctionCallExpression));
+			FunctionStack.Push(NewFunctionState.GetOwnedPointer());
+			
+			GetVisibleMaterialParametersFromExpression(FMaterialExpressionKey(FunctionCallExpression->FunctionOutputs[MaterialExpressionKey.OutputIndex].ExpressionOutput, 0), MaterialInstance, VisibleExpressions, FunctionStack);
+		
+			check(FunctionStack.Top()->ExpressionStack.Num() == 0);
+			FunctionStack.Pop();
 		}
 	}
+	else if (FunctionInputExpression)
+	{
+		FGetVisibleMaterialParametersFunctionState* FunctionState = FunctionStack.Pop();
+
+		const FFunctionExpressionInput* MatchingInput = FindInputById(FunctionInputExpression, FunctionState->FunctionCall->FunctionInputs);
+		check(MatchingInput);
+
+		GetVisibleMaterialParametersFromExpression(FMaterialExpressionKey(MatchingInput->Input.Expression, MatchingInput->Input.OutputIndex), MaterialInstance, VisibleExpressions, FunctionStack);
+
+		FunctionStack.Push(FunctionState);
+	}
 	else
 	{
-		const TArray<FExpressionInput*>& ExpressionInputs = MaterialExpression->GetInputs();
-		for(int32 ExpressionInputIndex = 0; ExpressionInputIndex < ExpressionInputs.Num(); ExpressionInputIndex++)
+		const TArray<FExpressionInput*>& ExpressionInputs = MaterialExpressionKey.Expression->GetInputs();
+		for (int32 ExpressionInputIndex = 0; ExpressionInputIndex < ExpressionInputs.Num(); ExpressionInputIndex++)
 		{
 			//retrieve the expression input and then start parsing its children
 			FExpressionInput* Input = ExpressionInputs[ExpressionInputIndex];
-			GetVisibleMaterialParametersFromExpression(Input->Expression, MaterialInstance, VisibleExpressions, ProcessedExpressions, FunctionInputs, VisibleFunctionInputExpressions);
+			GetVisibleMaterialParametersFromExpression(FMaterialExpressionKey(Input->Expression, Input->OutputIndex), MaterialInstance, VisibleExpressions, FunctionStack);
 		}
 	}
 
-	UMaterialExpression* TopExpression = ProcessedExpressions.Pop();
-	//ensure that the top of the stack matches what we expect (the same as MaterialExpression)
-	check(MaterialExpression == TopExpression);
+	FMaterialExpressionKey TopExpressionKey = FunctionStack.Top()->ExpressionStack.Pop();
+	check(FunctionDepth == FunctionStack.Num());
+	//ensure that the top of the stack matches what we expect (the same as MaterialExpressionKey)
+	check(MaterialExpressionKey == TopExpressionKey);
 }
 
 TSharedPtr<class IMaterialEditor> FMaterialEditorUtilities::GetIMaterialEditorForObject(const UObject* ObjectToFocusOn)
diff --git a/Engine/Source/Editor/MaterialEditor/Public/MaterialEditorActions.h b/Engine/Source/Editor/MaterialEditor/Public/MaterialEditorActions.h
--- a/Engine/Source/Editor/MaterialEditor/Public/MaterialEditorActions.h
+++ b/Engine/Source/Editor/MaterialEditor/Public/MaterialEditorActions.h
@@ -107,6 +107,9 @@ public:
 	TSharedPtr< FUICommandInfo > ConvertToTextureObjects;
 	TSharedPtr< FUICommandInfo > ConvertToTextureSamples;
 	
+	/** Converts selected objects to constants */
+	TSharedPtr< FUICommandInfo > ConvertToConstant;
+
 	/** Stops a node from being previewed in the viewport */
 	TSharedPtr< FUICommandInfo > StopPreviewNode;
 	
diff --git a/Engine/Source/Editor/MaterialEditor/Public/MaterialEditorUtilities.h b/Engine/Source/Editor/MaterialEditor/Public/MaterialEditorUtilities.h
--- a/Engine/Source/Editor/MaterialEditor/Public/MaterialEditorUtilities.h
+++ b/Engine/Source/Editor/MaterialEditor/Public/MaterialEditorUtilities.h
@@ -9,6 +9,19 @@ struct FExpressionInput;
 //////////////////////////////////////////////////////////////////////////
 // FMaterialEditorUtilities
 
+class FGetVisibleMaterialParametersFunctionState
+{
+public:
+
+	FGetVisibleMaterialParametersFunctionState(UMaterialExpressionMaterialFunctionCall* InFunctionCall) :
+		FunctionCall(InFunctionCall)
+	{}
+
+	class UMaterialExpressionMaterialFunctionCall* FunctionCall;
+	TArray<FMaterialExpressionKey> ExpressionStack;
+	TSet<FMaterialExpressionKey> VisitedExpressions;
+};
+
 class MATERIALEDITOR_API FMaterialEditorUtilities
 {
 public:
@@ -139,7 +152,7 @@ public:
 	* 
 	* @return	Returns true if a value for the switch expression is found, otherwise returns false.
 	*/
-	static bool GetStaticSwitchExpressionValue(UMaterialInstance* MaterialInstance, UMaterialExpression *SwitchValueExpression, bool& OutValue, FGuid& OutExpressionID, const TArray<FFunctionExpressionInput>* FunctionInputs = NULL);
+	static bool GetStaticSwitchExpressionValue(UMaterialInstance* MaterialInstance, UMaterialExpression *SwitchValueExpression, bool& OutValue, FGuid& OutExpressionID, const TArray<FFunctionExpressionInput>* FunctionInputs);
 
 	/**
 	 * Populates the specified material's Expressions array (eg if cooked out or old content).
@@ -157,10 +170,12 @@ private:
 	 * @param	MaterialExpression				The expression to parse.
 	 * @param	MaterialInstance				The material instance that contains all parameter overrides.
 	 * @param	VisisbleExpressions				The array that will contain the id's of the visible parameter expressions.
-	 * @param	FunctionInputs					An array of FFunctionExpressionInput when parsing material functions calls
-	 * @param	VisibleFunctionInputExpressions	The array that will contain the id's of visible function inputs when parsing function calls.
 	 */
-	static void GetVisibleMaterialParametersFromExpression(UMaterialExpression *MaterialExpression, UMaterialInstance *MaterialInstance, TArray<FGuid> &VisibleExpressions, TArray<UMaterialExpression*> &ProcessedExpressions, const TArray<FFunctionExpressionInput>* FunctionInputs = NULL, TArray<FGuid>* VisibleFunctionInputExpressions = NULL);
+	static void GetVisibleMaterialParametersFromExpression(
+		FMaterialExpressionKey MaterialExpressionKey, 
+		UMaterialInstance* MaterialInstance, 
+		TArray<FGuid>& VisibleExpressions, 
+		TArray<FGetVisibleMaterialParametersFunctionState*>& FunctionStack);
 
 	/** Get IMaterialEditor for given object, if it exists */
 	static TSharedPtr<class IMaterialEditor> GetIMaterialEditorForObject(const UObject* ObjectToFocusOn);
diff --git a/Engine/Source/Editor/UnrealEd/Private/MaterialGraphNode.cpp b/Engine/Source/Editor/UnrealEd/Private/MaterialGraphNode.cpp
--- a/Engine/Source/Editor/UnrealEd/Private/MaterialGraphNode.cpp
+++ b/Engine/Source/Editor/UnrealEd/Private/MaterialGraphNode.cpp
@@ -301,21 +301,18 @@ void UMaterialGraphNode::GetContextMenuActions(const FGraphNodeContextMenuBuilde
 				Context.MenuBuilder->AddMenuEntry(FMaterialEditorCommands::Get().UseCurrentTexture);
 
 				// Add a 'Convert To Texture' option for convertible types
-				if (MaterialGraph->MaterialFunction == NULL)
+				Context.MenuBuilder->BeginSection("MaterialEditorMenu0");
 				{
-					Context.MenuBuilder->BeginSection("MaterialEditorMenu0");
+					if ( MaterialExpression->IsA(UMaterialExpressionTextureSample::StaticClass()))
 					{
-						if ( MaterialExpression->IsA(UMaterialExpressionTextureSample::StaticClass()))
-						{
-							Context.MenuBuilder->AddMenuEntry(FMaterialEditorCommands::Get().ConvertToTextureObjects);
-						}
-						else if ( MaterialExpression->IsA(UMaterialExpressionTextureObject::StaticClass()))
-						{
-							Context.MenuBuilder->AddMenuEntry(FMaterialEditorCommands::Get().ConvertToTextureSamples);
-						}
+						Context.MenuBuilder->AddMenuEntry(FMaterialEditorCommands::Get().ConvertToTextureObjects);
+					}
+					else if ( MaterialExpression->IsA(UMaterialExpressionTextureObject::StaticClass()))
+					{
+						Context.MenuBuilder->AddMenuEntry(FMaterialEditorCommands::Get().ConvertToTextureSamples);
 					}
-					Context.MenuBuilder->EndSection();
 				}
+				Context.MenuBuilder->EndSection();
 			}
 
 			// Add a 'Convert To Parameter' option for convertible types
@@ -326,14 +323,22 @@ void UMaterialGraphNode::GetContextMenuActions(const FGraphNodeContextMenuBuilde
 				|| MaterialExpression->IsA(UMaterialExpressionTextureSample::StaticClass())
 				|| MaterialExpression->IsA(UMaterialExpressionComponentMask::StaticClass()))
 			{
-				if (MaterialGraph->MaterialFunction == NULL)
+				Context.MenuBuilder->BeginSection("MaterialEditorMenu1");
 				{
-					Context.MenuBuilder->BeginSection("MaterialEditorMenu1");
-					{
-						Context.MenuBuilder->AddMenuEntry(FMaterialEditorCommands::Get().ConvertObjects);
-					}
-					Context.MenuBuilder->EndSection();
+					Context.MenuBuilder->AddMenuEntry(FMaterialEditorCommands::Get().ConvertObjects);
+				}
+				Context.MenuBuilder->EndSection();
+			}
+
+			// Add a 'Convert To Constant' option for convertible types
+			if (MaterialExpression->IsA(UMaterialExpressionScalarParameter::StaticClass())
+				|| MaterialExpression->IsA(UMaterialExpressionVectorParameter::StaticClass()))
+			{
+				Context.MenuBuilder->BeginSection("MaterialEditorMenu1");
+				{
+					Context.MenuBuilder->AddMenuEntry(FMaterialEditorCommands::Get().ConvertToConstant);
 				}
+				Context.MenuBuilder->EndSection();
 			}
 
 			Context.MenuBuilder->BeginSection("MaterialEditorMenu2");
diff --git a/Engine/Source/Runtime/CoreUObject/Private/UObject/Property.cpp b/Engine/Source/Runtime/CoreUObject/Private/UObject/Property.cpp
--- a/Engine/Source/Runtime/CoreUObject/Private/UObject/Property.cpp
+++ b/Engine/Source/Runtime/CoreUObject/Private/UObject/Property.cpp
@@ -1279,3 +1279,16 @@ IMPLEMENT_CORE_INTRINSIC_CLASS(UNumericProperty, UProperty,
 {
 }
 );
+
+UProperty* UStruct::FindPropertyByName(FName Name) const
+{
+	for (UProperty* Property = PropertyLink; Property != NULL; Property = Property->PropertyLinkNext)
+	{
+		if (Property->GetFName() == Name)
+		{
+			return Property;
+		}
+	}
+
+	return NULL;
+}
\ No newline at end of file
diff --git a/Engine/Source/Runtime/CoreUObject/Public/UObject/Class.h b/Engine/Source/Runtime/CoreUObject/Public/UObject/Class.h
--- a/Engine/Source/Runtime/CoreUObject/Public/UObject/Class.h
+++ b/Engine/Source/Runtime/CoreUObject/Public/UObject/Class.h
@@ -267,6 +267,8 @@ public:
 	// UField interface.
 	virtual void AddCppProperty(UProperty* Property) override;
 
+	UProperty* FindPropertyByName(FName Name) const;
+
 	/**
 	 * Creates new copies of components
 	 * 
diff --git a/Engine/Source/Runtime/Engine/Classes/Materials/Material.h b/Engine/Source/Runtime/Engine/Classes/Materials/Material.h
--- a/Engine/Source/Runtime/Engine/Classes/Materials/Material.h
+++ b/Engine/Source/Runtime/Engine/Classes/Materials/Material.h
@@ -6,6 +6,8 @@
 #include "MaterialExpressionIO.h"
 #include "Materials/MaterialExpression.h"
 #include "Materials/MaterialInterface.h"
+#include "Materials/MaterialExpressionMaterialFunctionCall.h"
+#include "Materials/MaterialFunction.h"
 
 #include "Material.generated.h"
 
@@ -708,6 +710,7 @@ public:
 	 */
 	class FDefaultMaterialInstance* DefaultMaterialInstances[3];
 
+	/** Used to detect duplicate parameters.  Does not contain parameters in referenced functions! */
 	TMap<FName, TArray<UMaterialExpression*> > EditorParameters;
 
 #if WITH_EDITORONLY_DATA
@@ -760,6 +763,8 @@ public:
 	ENGINE_API virtual UPhysicalMaterial* GetPhysicalMaterial() const override;
 	ENGINE_API virtual void GetUsedTextures(TArray<UTexture*>& OutTextures, EMaterialQualityLevel::Type QualityLevel, bool bAllQualityLevels, ERHIFeatureLevel::Type FeatureLevel, bool bAllFeatureLevels) const override;
 	ENGINE_API virtual void OverrideTexture( const UTexture* InTextureToOverride, UTexture* OverrideTexture, ERHIFeatureLevel::Type InFeatureLevel ) override;
+	ENGINE_API virtual void OverrideVectorParameterDefault(FName ParameterName, const FLinearColor& Value, bool bOverride, ERHIFeatureLevel::Type FeatureLevel) override;
+	ENGINE_API virtual void OverrideScalarParameterDefault(FName ParameterName, float Value, bool bOverride, ERHIFeatureLevel::Type FeatureLevel) override;
 	ENGINE_API virtual bool CheckMaterialUsage(const EMaterialUsage Usage, const bool bSkipPrim = false) override;
 	ENGINE_API virtual bool CheckMaterialUsage_Concurrent(const EMaterialUsage Usage, const bool bSkipPrim = false) const override;
 	ENGINE_API virtual FMaterialResource* AllocateResource();
@@ -888,16 +893,27 @@ public:
 	 * @return	Returns a array of parameter names used in this material for the specified expression type.
 	 */
 	template<typename ExpressionType>
-	void GetAllParameterNames(TArray<FName> &OutParameterNames, TArray<FGuid> &OutParameterIds)
+	void GetAllParameterNames(TArray<FName>& OutParameterNames, TArray<FGuid>& OutParameterIds)
 	{
 		for (int32 ExpressionIndex = 0; ExpressionIndex < Expressions.Num(); ExpressionIndex++)
 		{
-			ExpressionType* ParameterExpression =
-				Cast<ExpressionType>(Expressions[ExpressionIndex]);
+			UMaterialExpressionMaterialFunctionCall* FunctionExpression = Cast<UMaterialExpressionMaterialFunctionCall>(Expressions[ExpressionIndex]);
 
-			if (ParameterExpression)
+			if (FunctionExpression)
+			{
+				if (FunctionExpression->MaterialFunction)
+				{
+					FunctionExpression->MaterialFunction->GetAllParameterNames<ExpressionType>(OutParameterNames, OutParameterIds);
+				}
+			}
+			else
 			{
-				ParameterExpression->GetAllParameterNames(OutParameterNames, OutParameterIds);
+				ExpressionType* ParameterExpression = Cast<ExpressionType>(Expressions[ExpressionIndex]);
+
+				if (ParameterExpression)
+				{
+					ParameterExpression->GetAllParameterNames(OutParameterNames, OutParameterIds);
+				}
 			}
 		}
 
@@ -1019,7 +1035,7 @@ public:
 	 * Add an expression node that represents a parameter to the list of material parameters.
 	 * @param	Expression	Pointer to the node that is going to be inserted if it's a parameter type.
 	 */
-	ENGINE_API virtual bool AddExpressionParameter(UMaterialExpression* Expression);
+	ENGINE_API virtual bool AddExpressionParameter(UMaterialExpression* Expression, TMap<FName, TArray<UMaterialExpression*> >& ParameterTypeMap);
 
 	/**
 	 * Removes an expression node that represents a parameter from the list of material parameters.
@@ -1095,22 +1111,6 @@ public:
 	 */
 	ENGINE_API static bool IsDynamicParameter(const UMaterialExpression* Expression);
 
-	/**
-	 * Return wheter the number of parameter groups. NOTE: The number returned can be innaccurate if you have parameters of different types with the same name.
-	 */
-	inline int32 GetNumEditorParameters() const
-	{
-		return EditorParameters.Num();
-	}
-
-	/**
-	 * Empty the editor parameters for the material.
-	 */
-	inline void EmptyEditorParameters()
-	{
-		EditorParameters.Empty();
-	}
-
 	/** Returns an array of the guids of functions used in this material, with the call hierarchy flattened. */
 	void GetReferencedFunctionIds(TArray<FGuid>& OutIds) const;
 
diff --git a/Engine/Source/Runtime/Engine/Classes/Materials/MaterialExpression.h b/Engine/Source/Runtime/Engine/Classes/Materials/MaterialExpression.h
--- a/Engine/Source/Runtime/Engine/Classes/Materials/MaterialExpression.h
+++ b/Engine/Source/Runtime/Engine/Classes/Materials/MaterialExpression.h
@@ -360,13 +360,6 @@ class ENGINE_API UMaterialExpression : public UObject
 	/** Checks whether any inputs to this expression create a loop */
 	bool ContainsInputLoop();
 
-	/**
-	 * Tells if this material expression type is a parameter expression.
-	 *
-	 * @returns True if the type of this is a parameter expression. False otherwise.
-	 */
-	bool IsParameterExpression() const;
-
 protected:
 	/**
 	 * Checks whether any inputs to this expression create a loop by recursively
diff --git a/Engine/Source/Runtime/Engine/Classes/Materials/MaterialExpressionScalarParameter.h b/Engine/Source/Runtime/Engine/Classes/Materials/MaterialExpressionScalarParameter.h
--- a/Engine/Source/Runtime/Engine/Classes/Materials/MaterialExpressionScalarParameter.h
+++ b/Engine/Source/Runtime/Engine/Classes/Materials/MaterialExpressionScalarParameter.h
@@ -18,6 +18,10 @@ class UMaterialExpressionScalarParameter : public UMaterialExpressionParameter
 	virtual int32 Compile(class FMaterialCompiler* Compiler, int32 OutputIndex, int32 MultiplexIndex) override;
 	virtual void GetCaption(TArray<FString>& OutCaptions) const override;
 	// End UMaterialExpression Interface
+
+#if WITH_EDITOR
+	virtual void PostEditChangeProperty(FPropertyChangedEvent& PropertyChangedEvent) override;
+#endif
 };
 
 
diff --git a/Engine/Source/Runtime/Engine/Classes/Materials/MaterialExpressionVectorParameter.h b/Engine/Source/Runtime/Engine/Classes/Materials/MaterialExpressionVectorParameter.h
--- a/Engine/Source/Runtime/Engine/Classes/Materials/MaterialExpressionVectorParameter.h
+++ b/Engine/Source/Runtime/Engine/Classes/Materials/MaterialExpressionVectorParameter.h
@@ -13,11 +13,14 @@ class UMaterialExpressionVectorParameter : public UMaterialExpressionParameter
 	UPROPERTY(EditAnywhere, Category=MaterialExpressionVectorParameter)
 	FLinearColor DefaultValue;
 
-
 	// Begin UMaterialExpression Interface
 	virtual int32 Compile(class FMaterialCompiler* Compiler, int32 OutputIndex, int32 MultiplexIndex) override;
 	virtual void GetCaption(TArray<FString>& OutCaptions) const override;
 	// End UMaterialExpression Interface
+
+#if WITH_EDITOR
+	virtual void PostEditChangeProperty(FPropertyChangedEvent& PropertyChangedEvent) override;
+#endif
 };
 
 
diff --git a/Engine/Source/Runtime/Engine/Classes/Materials/MaterialFunction.h b/Engine/Source/Runtime/Engine/Classes/Materials/MaterialFunction.h
--- a/Engine/Source/Runtime/Engine/Classes/Materials/MaterialFunction.h
+++ b/Engine/Source/Runtime/Engine/Classes/Materials/MaterialFunction.h
@@ -81,6 +81,34 @@ public:
 	/** Appends textures referenced by the expressions in this function. */
 	ENGINE_API void AppendReferencedTextures(TArray<UTexture*>& InOutTextures) const;
 
+	template<typename ExpressionType>
+	void GetAllParameterNames(TArray<FName> &OutParameterNames, TArray<FGuid> &OutParameterIds)
+	{
+		for (int32 ExpressionIndex = 0; ExpressionIndex < FunctionExpressions.Num(); ExpressionIndex++)
+		{
+			UMaterialExpressionMaterialFunctionCall* FunctionExpression = Cast<UMaterialExpressionMaterialFunctionCall>(FunctionExpressions[ExpressionIndex]);
+
+			if (FunctionExpression)
+			{
+				if (FunctionExpression->MaterialFunction)
+				{
+					FunctionExpression->MaterialFunction->GetAllParameterNames<ExpressionType>(OutParameterNames, OutParameterIds);
+				}
+			}
+			else
+			{
+				ExpressionType* ParameterExpression = Cast<ExpressionType>(FunctionExpressions[ExpressionIndex]);
+
+				if (ParameterExpression)
+				{
+					ParameterExpression->GetAllParameterNames(OutParameterNames, OutParameterIds);
+				}
+			}
+		}
+
+		check(OutParameterNames.Num() == OutParameterIds.Num());
+	}
+
 #if WITH_EDITOR
 	ENGINE_API UMaterial* GetPreviewMaterial();
 
diff --git a/Engine/Source/Runtime/Engine/Classes/Materials/MaterialInstance.h b/Engine/Source/Runtime/Engine/Classes/Materials/MaterialInstance.h
--- a/Engine/Source/Runtime/Engine/Classes/Materials/MaterialInstance.h
+++ b/Engine/Source/Runtime/Engine/Classes/Materials/MaterialInstance.h
@@ -255,6 +255,8 @@ public:
 	virtual ENGINE_API bool GetVectorParameterValue(FName ParameterName, FLinearColor& OutValue) const override;
 	virtual ENGINE_API void GetUsedTextures(TArray<UTexture*>& OutTextures, EMaterialQualityLevel::Type QualityLevel, bool bAllQualityLevels, ERHIFeatureLevel::Type FeatureLevel, bool bAllFeatureLevels) const override;
 	virtual ENGINE_API void OverrideTexture(const UTexture* InTextureToOverride, UTexture* OverrideTexture, ERHIFeatureLevel::Type InFeatureLevel) override;
+	virtual ENGINE_API void OverrideVectorParameterDefault(FName ParameterName, const FLinearColor& Value, bool bOverride, ERHIFeatureLevel::Type FeatureLevel) override;
+	virtual ENGINE_API void OverrideScalarParameterDefault(FName ParameterName, float Value, bool bOverride, ERHIFeatureLevel::Type FeatureLevel) override;
 	virtual ENGINE_API bool CheckMaterialUsage(const EMaterialUsage Usage, const bool bSkipPrim = false) override;
 	virtual ENGINE_API bool CheckMaterialUsage_Concurrent(const EMaterialUsage Usage, const bool bSkipPrim = false) const override;
 	virtual ENGINE_API bool GetStaticSwitchParameterValue(FName ParameterName, bool &OutValue, FGuid &OutExpressionGuid) override;
diff --git a/Engine/Source/Runtime/Engine/Classes/Materials/MaterialInterface.h b/Engine/Source/Runtime/Engine/Classes/Materials/MaterialInterface.h
--- a/Engine/Source/Runtime/Engine/Classes/Materials/MaterialInterface.h
+++ b/Engine/Source/Runtime/Engine/Classes/Materials/MaterialInterface.h
@@ -280,6 +280,14 @@ public:
 	 */
 	virtual void OverrideTexture(const UTexture* InTextureToOverride, UTexture* OverrideTexture, ERHIFeatureLevel::Type InFeatureLevel) PURE_VIRTUAL(UMaterialInterface::OverrideTexture, return;);
 
+	/** 
+	 * Overrides the default value of the given parameter (transient).  
+	 * This is used to implement realtime previewing of parameter defaults. 
+	 * Handles updating dependent MI's and cached uniform expressions.
+	 */
+	virtual void OverrideVectorParameterDefault(FName ParameterName, const FLinearColor& Value, bool bOverride, ERHIFeatureLevel::Type FeatureLevel) PURE_VIRTUAL(UMaterialInterface::OverrideTexture, return;);
+	virtual void OverrideScalarParameterDefault(FName ParameterName, float Value, bool bOverride, ERHIFeatureLevel::Type FeatureLevel) PURE_VIRTUAL(UMaterialInterface::OverrideTexture, return;);
+
 	/**
 	 * Checks if the material can be used with the given usage flag.  
 	 * If the flag isn't set in the editor, it will be set and the material will be recompiled with it.
diff --git a/Engine/Source/Runtime/Engine/Private/Engine.cpp b/Engine/Source/Runtime/Engine/Private/Engine.cpp
--- a/Engine/Source/Runtime/Engine/Private/Engine.cpp
+++ b/Engine/Source/Runtime/Engine/Private/Engine.cpp
@@ -44,6 +44,8 @@ FEditorSupportDelegates::FOnWindowsMessage FEditorSupportDelegates::PreWindowsMe
 FEditorSupportDelegates::FOnWindowsMessage FEditorSupportDelegates::PostWindowsMessage;
 /** Sent after the usages flags on a material have changed*/
 FEditorSupportDelegates::FOnMaterialUsageFlagsChanged FEditorSupportDelegates::MaterialUsageFlagsChanged;
+FEditorSupportDelegates::FOnVectorParameterDefaultChanged FEditorSupportDelegates::VectorParameterDefaultChanged;
+FEditorSupportDelegates::FOnScalarParameterDefaultChanged FEditorSupportDelegates::ScalarParameterDefaultChanged;
 
 #endif // WITH_EDITOR
 
diff --git a/Engine/Source/Runtime/Engine/Private/Materials/Material.cpp b/Engine/Source/Runtime/Engine/Private/Materials/Material.cpp
--- a/Engine/Source/Runtime/Engine/Private/Materials/Material.cpp
+++ b/Engine/Source/Runtime/Engine/Private/Materials/Material.cpp
@@ -19,6 +19,7 @@
 #include "Materials/MaterialExpressionStaticSwitchParameter.h"
 #include "Materials/MaterialExpressionTextureSampleParameter.h"
 #include "Materials/MaterialExpressionVectorParameter.h"
+#include "MaterialUniformExpressions.h"
 #include "Engine/SubsurfaceProfile.h"
 #include "EditorSupportDelegates.h"
 #include "MaterialShaderType.h"
@@ -668,6 +669,65 @@ void UMaterial::OverrideTexture(const UTexture* InTextureToOverride, UTexture* O
 #endif // #if WITH_EDITOR
 }
 
+void UMaterial::OverrideVectorParameterDefault(FName ParameterName, const FLinearColor& Value, bool bOverride, ERHIFeatureLevel::Type InFeatureLevel)
+{
+#if WITH_EDITOR
+	bool bShouldRecacheMaterialExpressions = false;
+
+	FMaterialResource* Resource = GetMaterialResource(InFeatureLevel);
+	// Iterate over both the 2D textures and cube texture expressions.
+	const TArray<TRefCountPtr<FMaterialUniformExpression> >& UniformExpressions = Resource->GetUniformVectorParameterExpressions();
+
+	// Iterate over each of the material's texture expressions.
+	for (FMaterialUniformExpression* UniformExpression : UniformExpressions)
+	{
+		check(UniformExpression->GetType() == &FMaterialUniformExpressionVectorParameter::StaticType);
+		FMaterialUniformExpressionVectorParameter* VectorExpression = static_cast<FMaterialUniformExpressionVectorParameter*>(UniformExpression);
+
+		if (VectorExpression->GetParameterName() == ParameterName)
+		{
+			VectorExpression->SetTransientOverrideDefaultValue(Value, bOverride);
+			bShouldRecacheMaterialExpressions = true;
+		}
+	}
+
+	if (bShouldRecacheMaterialExpressions)
+	{
+		RecacheUniformExpressions();
+		RecacheMaterialInstanceUniformExpressions(this);
+	}
+#endif // #if WITH_EDITOR
+}
+
+void UMaterial::OverrideScalarParameterDefault(FName ParameterName, float Value, bool bOverride, ERHIFeatureLevel::Type InFeatureLevel)
+{
+#if WITH_EDITOR
+	bool bShouldRecacheMaterialExpressions = false;
+
+	FMaterialResource* Resource = GetMaterialResource(InFeatureLevel);
+	// Iterate over both the 2D textures and cube texture expressions.
+	const TArray<TRefCountPtr<FMaterialUniformExpression> >& UniformExpressions = Resource->GetUniformScalarParameterExpressions();
+
+	// Iterate over each of the material's texture expressions.
+	for (FMaterialUniformExpression* UniformExpression : UniformExpressions)
+	{
+		check(UniformExpression->GetType() == &FMaterialUniformExpressionScalarParameter::StaticType);
+		FMaterialUniformExpressionScalarParameter* ScalarExpression = static_cast<FMaterialUniformExpressionScalarParameter*>(UniformExpression);
+
+		if (ScalarExpression->GetParameterName() == ParameterName)
+		{
+			ScalarExpression->SetTransientOverrideDefaultValue(Value, bOverride);
+			bShouldRecacheMaterialExpressions = true;
+		}
+	}
+
+	if (bShouldRecacheMaterialExpressions)
+	{
+		RecacheUniformExpressions();
+		RecacheMaterialInstanceUniformExpressions(this);
+	}
+#endif // #if WITH_EDITOR
+}
 
 void UMaterial::RecacheUniformExpressions() const
 {
@@ -905,7 +965,6 @@ bool UMaterial::NeedsSetMaterialUsage_Concurrent(bool &bOutHasUsage, EMaterialUs
 	return false;
 }
 
-
 bool UMaterial::SetMaterialUsage(bool &bNeedsRecompile, EMaterialUsage Usage, const bool bSkipPrim)
 {
 	bNeedsRecompile = false;
@@ -2328,7 +2387,7 @@ void UMaterial::PostEditChangeProperty(FPropertyChangedEvent& PropertyChangedEve
 } 
 #endif // WITH_EDITOR
 
-bool UMaterial::AddExpressionParameter(UMaterialExpression* Expression)
+bool UMaterial::AddExpressionParameter(UMaterialExpression* Expression, TMap<FName, TArray<UMaterialExpression*> >& ParameterTypeMap)
 {
 	if(!Expression)
 	{
@@ -2341,11 +2400,11 @@ bool UMaterial::AddExpressionParameter(UMaterialExpression* Expression)
 	{
 		UMaterialExpressionParameter *Param = (UMaterialExpressionParameter*)Expression;
 
-		TArray<UMaterialExpression*> *ExpressionList = EditorParameters.Find(Param->ParameterName);
+		TArray<UMaterialExpression*> *ExpressionList = ParameterTypeMap.Find(Param->ParameterName);
 
 		if(!ExpressionList)
 		{
-			ExpressionList = &EditorParameters.Add(Param->ParameterName, TArray<UMaterialExpression*>());
+			ExpressionList = &ParameterTypeMap.Add(Param->ParameterName, TArray<UMaterialExpression*>());
 		}
 
 		ExpressionList->Add(Param);
@@ -2355,11 +2414,11 @@ bool UMaterial::AddExpressionParameter(UMaterialExpression* Expression)
 	{
 		UMaterialExpressionTextureSampleParameter *Param = (UMaterialExpressionTextureSampleParameter*)Expression;
 
-		TArray<UMaterialExpression*> *ExpressionList = EditorParameters.Find(Param->ParameterName);
+		TArray<UMaterialExpression*> *ExpressionList = ParameterTypeMap.Find(Param->ParameterName);
 
 		if(!ExpressionList)
 		{
-			ExpressionList = &EditorParameters.Add(Param->ParameterName, TArray<UMaterialExpression*>());
+			ExpressionList = &ParameterTypeMap.Add(Param->ParameterName, TArray<UMaterialExpression*>());
 		}
 
 		ExpressionList->Add(Param);
@@ -2369,11 +2428,11 @@ bool UMaterial::AddExpressionParameter(UMaterialExpression* Expression)
 	{
 		UMaterialExpressionFontSampleParameter *Param = (UMaterialExpressionFontSampleParameter*)Expression;
 
-		TArray<UMaterialExpression*> *ExpressionList = EditorParameters.Find(Param->ParameterName);
+		TArray<UMaterialExpression*> *ExpressionList = ParameterTypeMap.Find(Param->ParameterName);
 
 		if(!ExpressionList)
 		{
-			ExpressionList = &EditorParameters.Add(Param->ParameterName, TArray<UMaterialExpression*>());
+			ExpressionList = &ParameterTypeMap.Add(Param->ParameterName, TArray<UMaterialExpression*>());
 		}
 
 		ExpressionList->Add(Param);
@@ -2390,7 +2449,7 @@ bool UMaterial::RemoveExpressionParameter(UMaterialExpression* Expression)
 
 	if(GetExpressionParameterName(Expression, ParmName))
 	{
-		TArray<UMaterialExpression*> *ExpressionList = EditorParameters.Find(ParmName);
+		TArray<UMaterialExpression*>* ExpressionList = EditorParameters.Find(ParmName);
 
 		if(ExpressionList)
 		{
@@ -2437,11 +2496,11 @@ bool UMaterial::IsDynamicParameter(const UMaterialExpression* Expression)
 
 void UMaterial::BuildEditorParameterList()
 {
-	EmptyEditorParameters();
+	EditorParameters.Empty();
 
 	for(int32 MaterialExpressionIndex = 0 ; MaterialExpressionIndex < Expressions.Num() ; ++MaterialExpressionIndex)
 	{
-		AddExpressionParameter(Expressions[ MaterialExpressionIndex ]);
+		AddExpressionParameter(Expressions[MaterialExpressionIndex], EditorParameters);
 	}
 }
 
@@ -2452,13 +2511,13 @@ bool UMaterial::HasDuplicateParameters(const UMaterialExpression* Expression)
 
 	if(GetExpressionParameterName(Expression, ExpressionName))
 	{
-		TArray<UMaterialExpression*> *ExpressionList = EditorParameters.Find(ExpressionName);
+		TArray<UMaterialExpression*>* ExpressionList = EditorParameters.Find(ExpressionName);
 
 		if(ExpressionList)
 		{
 			for(int32 ParmIndex = 0; ParmIndex < ExpressionList->Num(); ++ParmIndex)
 			{
-				UMaterialExpression *CurNode = (*ExpressionList)[ParmIndex];
+				UMaterialExpression* CurNode = (*ExpressionList)[ParmIndex];
 				if(CurNode != Expression && CurNode->GetClass() == Expression->GetClass())
 				{
 					return true;
@@ -2515,7 +2574,7 @@ void UMaterial::PropagateExpressionParameterChanges(UMaterialExpression* Paramet
 
 	if(bRet)
 	{
-		TArray<UMaterialExpression*> *ExpressionList = EditorParameters.Find(ParmName);
+		TArray<UMaterialExpression*>* ExpressionList = EditorParameters.Find(ParmName);
 
 		if(ExpressionList && ExpressionList->Num() > 1)
 		{
@@ -2545,7 +2604,7 @@ void UMaterial::UpdateExpressionParameterName(UMaterialExpression* Expression)
 				EditorParameters.Remove(Iter.Key());
 			}
 
-			AddExpressionParameter(Expression);
+			AddExpressionParameter(Expression, EditorParameters);
 			break;
 		}
 	}
diff --git a/Engine/Source/Runtime/Engine/Private/Materials/MaterialExpressions.cpp b/Engine/Source/Runtime/Engine/Private/Materials/MaterialExpressions.cpp
--- a/Engine/Source/Runtime/Engine/Private/Materials/MaterialExpressions.cpp
+++ b/Engine/Source/Runtime/Engine/Private/Materials/MaterialExpressions.cpp
@@ -163,14 +163,7 @@ bool IsAllowedExpressionType(UClass* Class, bool bMaterialFunction)
 
 	if (bMaterialFunction)
 	{
-		return bSharedAllowed
-			// Exclude parameter nodes from functions, as parameter functionality currently isn't setup to scan all functions as needed, 
-			// Also naming collisions would need to be communicated to the user
-			&& !Class->IsChildOf(UMaterialExpressionParameter::StaticClass())
-			&& !Class->IsChildOf(UMaterialExpressionTextureSampleParameter::StaticClass())
-			&& !Class->IsChildOf(UMaterialExpressionDynamicParameter::StaticClass())
-			&& !Class->IsChildOf(UMaterialExpressionFontSampleParameter::StaticClass())
-			&& !Class->GetDefaultObject<UMaterialExpression>()->IsParameterExpression();
+		return bSharedAllowed;
 	}
 	else
 	{
@@ -983,11 +976,6 @@ bool UMaterialExpression::ContainsInputLoop()
 	return ContainsInputLoopInternal(ExpressionStack);
 }
 
-bool UMaterialExpression::IsParameterExpression() const
-{
-	return bIsParameterExpression;
-}
-
 bool UMaterialExpression::ContainsInputLoopInternal(TArray<FMaterialExpressionKey>& ExpressionStack)
 {
 	const TArray<FExpressionInput*> Inputs = GetInputs();
@@ -3549,6 +3537,7 @@ void UMaterialExpressionParameter::SetEditableName(const FString& NewName)
 {
 	ParameterName = *NewName;
 }
+
 #endif
 
 void UMaterialExpressionParameter::GetAllParameterNames(TArray<FName> &OutParameterNames, TArray<FGuid> &OutParameterIds)
@@ -3594,6 +3583,26 @@ void UMaterialExpressionVectorParameter::GetCaption(TArray<FString>& OutCaptions
 	OutCaptions.Add(FString::Printf(TEXT("'%s'"), *ParameterName.ToString())); 
 }
 
+#if WITH_EDITOR
+
+void UMaterialExpressionVectorParameter::PostEditChangeProperty(FPropertyChangedEvent& PropertyChangedEvent)
+{
+	UProperty* PropertyThatChanged = PropertyChangedEvent.MemberProperty;
+	const FString PropertyName = PropertyThatChanged ? PropertyThatChanged->GetName() : TEXT("");
+
+	if (PropertyName == GET_MEMBER_NAME_STRING_CHECKED(UMaterialExpressionVectorParameter, DefaultValue))
+	{
+#if WITH_EDITOR
+		// Callback into the editor
+		FEditorSupportDelegates::VectorParameterDefaultChanged.Broadcast(this, ParameterName, DefaultValue);
+#endif
+	}
+
+	Super::PostEditChangeProperty(PropertyChangedEvent);
+}
+
+#endif
+
 //
 //	UMaterialExpressionScalarParameter
 //
@@ -3618,6 +3627,25 @@ void UMaterialExpressionScalarParameter::GetCaption(TArray<FString>& OutCaptions
 	 OutCaptions.Add(FString::Printf(TEXT("'%s'"), *ParameterName.ToString())); 
 }
 
+#if WITH_EDITOR
+
+void UMaterialExpressionScalarParameter::PostEditChangeProperty(FPropertyChangedEvent& PropertyChangedEvent)
+{
+	UProperty* PropertyThatChanged = PropertyChangedEvent.MemberProperty;
+	const FString PropertyName = PropertyThatChanged ? PropertyThatChanged->GetName() : TEXT("");
+
+	if (PropertyName == GET_MEMBER_NAME_STRING_CHECKED(UMaterialExpressionScalarParameter, DefaultValue))
+	{
+#if WITH_EDITOR
+		// Callback into the editor
+		FEditorSupportDelegates::ScalarParameterDefaultChanged.Broadcast(this, ParameterName, DefaultValue);
+#endif
+	}
+
+	Super::PostEditChangeProperty(PropertyChangedEvent);
+}
+
+#endif
 
 //
 //	UMaterialExpressionStaticSwitchParameter
diff --git a/Engine/Source/Runtime/Engine/Private/Materials/MaterialInstance.cpp b/Engine/Source/Runtime/Engine/Private/Materials/MaterialInstance.cpp
--- a/Engine/Source/Runtime/Engine/Private/Materials/MaterialInstance.cpp
+++ b/Engine/Source/Runtime/Engine/Private/Materials/MaterialInstance.cpp
@@ -8,6 +8,7 @@
 #include "Materials/MaterialExpressionStaticBoolParameter.h"
 #include "Materials/MaterialExpressionStaticComponentMaskParameter.h"
 #include "Materials/MaterialInstanceBasePropertyOverrides.h"
+#include "MaterialUniformExpressions.h"
 #include "MaterialInstanceSupport.h"
 #include "MaterialShaderType.h"
 #include "TargetPlatform.h"
@@ -772,6 +773,72 @@ void UMaterialInstance::OverrideTexture(const UTexture* InTextureToOverride, UTe
 #endif // #if WITH_EDITOR
 }
 
+void UMaterialInstance::OverrideVectorParameterDefault(FName ParameterName, const FLinearColor& Value, bool bOverride, ERHIFeatureLevel::Type InFeatureLevel)
+{
+#if WITH_EDITOR
+	bool bShouldRecacheMaterialExpressions = false;
+
+	if (bHasStaticPermutationResource)
+	{
+		const FMaterialResource* SourceMaterialResource = GetMaterialResource(InFeatureLevel);
+		const TArray<TRefCountPtr<FMaterialUniformExpression> >& UniformExpressions = SourceMaterialResource->GetUniformVectorParameterExpressions();
+
+		// Iterate over each of the material's texture expressions.
+		for (int32 ExpressionIndex = 0; ExpressionIndex < UniformExpressions.Num(); ExpressionIndex++)
+		{
+			FMaterialUniformExpression* UniformExpression = UniformExpressions[ExpressionIndex];
+			check(UniformExpression->GetType() == &FMaterialUniformExpressionVectorParameter::StaticType);
+			FMaterialUniformExpressionVectorParameter* VectorExpression = static_cast<FMaterialUniformExpressionVectorParameter*>(UniformExpression);
+
+			if (VectorExpression->GetParameterName() == ParameterName)
+			{
+				VectorExpression->SetTransientOverrideDefaultValue(Value, bOverride);
+				bShouldRecacheMaterialExpressions = true;
+			}
+		}
+	}
+
+	if (bShouldRecacheMaterialExpressions)
+	{
+		RecacheUniformExpressions();
+		RecacheMaterialInstanceUniformExpressions(this);
+	}
+#endif // #if WITH_EDITOR
+}
+
+void UMaterialInstance::OverrideScalarParameterDefault(FName ParameterName, float Value, bool bOverride, ERHIFeatureLevel::Type InFeatureLevel)
+{
+#if WITH_EDITOR
+	bool bShouldRecacheMaterialExpressions = false;
+
+	if (bHasStaticPermutationResource)
+	{
+		const FMaterialResource* SourceMaterialResource = GetMaterialResource(InFeatureLevel);
+		const TArray<TRefCountPtr<FMaterialUniformExpression> >& UniformExpressions = SourceMaterialResource->GetUniformScalarParameterExpressions();
+
+		// Iterate over each of the material's texture expressions.
+		for (int32 ExpressionIndex = 0; ExpressionIndex < UniformExpressions.Num(); ExpressionIndex++)
+		{
+			FMaterialUniformExpression* UniformExpression = UniformExpressions[ExpressionIndex];
+			check(UniformExpression->GetType() == &FMaterialUniformExpressionScalarParameter::StaticType);
+			FMaterialUniformExpressionScalarParameter* ScalarExpression = static_cast<FMaterialUniformExpressionScalarParameter*>(UniformExpression);
+
+			if (ScalarExpression->GetParameterName() == ParameterName)
+			{
+				ScalarExpression->SetTransientOverrideDefaultValue(Value, bOverride);
+				bShouldRecacheMaterialExpressions = true;
+			}
+		}
+	}
+
+	if (bShouldRecacheMaterialExpressions)
+	{
+		RecacheUniformExpressions();
+		RecacheMaterialInstanceUniformExpressions(this);
+	}
+#endif // #if WITH_EDITOR
+}
+
 bool UMaterialInstance::CheckMaterialUsage(const EMaterialUsage Usage, const bool bSkipPrim)
 {
 	check(IsInGameThread());
diff --git a/Engine/Source/Runtime/Engine/Private/Materials/MaterialShared.cpp b/Engine/Source/Runtime/Engine/Private/Materials/MaterialShared.cpp
--- a/Engine/Source/Runtime/Engine/Private/Materials/MaterialShared.cpp
+++ b/Engine/Source/Runtime/Engine/Private/Materials/MaterialShared.cpp
@@ -376,6 +376,56 @@ const TArray<TRefCountPtr<FMaterialUniformExpressionTexture> >& FMaterial::GetUn
 	return EmptyExpressions;
 }
 
+const TArray<TRefCountPtr<FMaterialUniformExpression> >& FMaterial::GetUniformVectorParameterExpressions() const 
+{ 
+	const FMaterialShaderMap* ShaderMapToUse = NULL;
+
+	if (IsInGameThread())
+	{
+		// If we are accessing uniform texture expressions on the game thread, use results from a shader map whose compile is in flight that matches this material
+		// This allows querying what textures a material uses even when it is being asynchronously compiled
+		ShaderMapToUse = GetGameThreadShaderMap() ? GetGameThreadShaderMap() : FMaterialShaderMap::GetShaderMapBeingCompiled(this);
+	}
+	else 
+	{
+		check(IsInRenderingThread());
+		ShaderMapToUse = GetRenderingThreadShaderMap();
+	}
+
+	if (ShaderMapToUse)
+	{
+		return ShaderMapToUse->GetUniformExpressionSet().UniformVectorExpressions; 
+	}
+
+	static const TArray<TRefCountPtr<FMaterialUniformExpression> > EmptyExpressions;
+	return EmptyExpressions;
+}
+
+const TArray<TRefCountPtr<FMaterialUniformExpression> >& FMaterial::GetUniformScalarParameterExpressions() const 
+{ 
+	const FMaterialShaderMap* ShaderMapToUse = NULL;
+
+	if (IsInGameThread())
+	{
+		// If we are accessing uniform texture expressions on the game thread, use results from a shader map whose compile is in flight that matches this material
+		// This allows querying what textures a material uses even when it is being asynchronously compiled
+		ShaderMapToUse = GetGameThreadShaderMap() ? GetGameThreadShaderMap() : FMaterialShaderMap::GetShaderMapBeingCompiled(this);
+	}
+	else 
+	{
+		check(IsInRenderingThread());
+		ShaderMapToUse = GetRenderingThreadShaderMap();
+	}
+
+	if (ShaderMapToUse)
+	{
+		return ShaderMapToUse->GetUniformExpressionSet().UniformScalarExpressions; 
+	}
+
+	static const TArray<TRefCountPtr<FMaterialUniformExpression> > EmptyExpressions;
+	return EmptyExpressions;
+}
+
 bool FMaterial::RequiresSceneColorCopy_GameThread() const
 {
 	check(IsInGameThread());
@@ -1488,11 +1538,9 @@ TSet<FMaterial*> FMaterial::EditorLoadedMaterialResources;
 FMaterialRenderContext::FMaterialRenderContext(
 	const FMaterialRenderProxy* InMaterialRenderProxy,
 	const FMaterial& InMaterial,
-	const FSceneView* InView,
-	bool bInWorkAroundDeferredMipArtifacts)
+	const FSceneView* InView)
 		: MaterialRenderProxy(InMaterialRenderProxy)
 		, Material(InMaterial)
-		, bWorkAroundDeferredMipArtifacts(bInWorkAroundDeferredMipArtifacts)
 {
 	bShowSelection = GIsEditor && InView && InView->Family->EngineShowFlags.Selection;
 }
diff --git a/Engine/Source/Runtime/Engine/Private/Materials/MaterialUniformExpressions.h b/Engine/Source/Runtime/Engine/Private/Materials/MaterialUniformExpressions.h
--- a/Engine/Source/Runtime/Engine/Private/Materials/MaterialUniformExpressions.h
+++ b/Engine/Source/Runtime/Engine/Private/Materials/MaterialUniformExpressions.h
@@ -105,10 +105,14 @@ class FMaterialUniformExpressionVectorParameter: public FMaterialUniformExpressi
 	DECLARE_MATERIALUNIFORMEXPRESSION_TYPE(FMaterialUniformExpressionVectorParameter);
 public:
 
-	FMaterialUniformExpressionVectorParameter() {}
+	FMaterialUniformExpressionVectorParameter() :
+		bUseOverriddenDefault(false)
+	{}
+
 	FMaterialUniformExpressionVectorParameter(FName InParameterName,const FLinearColor& InDefaultValue):
 		ParameterName(InParameterName),
-		DefaultValue(InDefaultValue)
+		DefaultValue(InDefaultValue),
+		bUseOverriddenDefault(false)
 	{}
 
 	// FMaterialUniformExpression interface.
@@ -116,19 +120,27 @@ public:
 	{
 		Ar << ParameterName << DefaultValue;
 	}
+
 	virtual void GetNumberValue(const FMaterialRenderContext& Context,FLinearColor& OutValue) const
 	{
 		OutValue.R = OutValue.G = OutValue.B = OutValue.A = 0;
 
 		if(!Context.MaterialRenderProxy->GetVectorValue(ParameterName, &OutValue, Context))
 		{
-			OutValue = DefaultValue;
+			OutValue = bUseOverriddenDefault ? OverriddenDefaultValue : DefaultValue;
 		}
 	}
+
 	virtual bool IsConstant() const
 	{
 		return false;
 	}
+
+	FName GetParameterName() const
+	{
+		return ParameterName;
+	}
+
 	virtual bool IsIdentical(const FMaterialUniformExpression* OtherExpression) const
 	{
 		if (GetType() != OtherExpression->GetType())
@@ -139,9 +151,17 @@ public:
 		return ParameterName == OtherParameter->ParameterName && DefaultValue == OtherParameter->DefaultValue;
 	}
 
+	void SetTransientOverrideDefaultValue(const FLinearColor& InOverrideDefaultValue, bool bInUseOverriddenDefault)
+	{
+		bUseOverriddenDefault = bInUseOverriddenDefault;
+		OverriddenDefaultValue = InOverrideDefaultValue;
+	}
+
 private:
 	FName ParameterName;
 	FLinearColor DefaultValue;
+	bool bUseOverriddenDefault;
+	FLinearColor OverriddenDefaultValue;
 };
 
 /**
@@ -151,10 +171,14 @@ class FMaterialUniformExpressionScalarParameter: public FMaterialUniformExpressi
 	DECLARE_MATERIALUNIFORMEXPRESSION_TYPE(FMaterialUniformExpressionScalarParameter);
 public:
 
-	FMaterialUniformExpressionScalarParameter() {}
+	FMaterialUniformExpressionScalarParameter() :
+		bUseOverriddenDefault(false)
+	{}
+
 	FMaterialUniformExpressionScalarParameter(FName InParameterName,float InDefaultValue):
 		ParameterName(InParameterName),
-		DefaultValue(InDefaultValue)
+		DefaultValue(InDefaultValue),
+		bUseOverriddenDefault(false)
 	{}
 
 	// FMaterialUniformExpression interface.
@@ -162,6 +186,7 @@ public:
 	{
 		Ar << ParameterName << DefaultValue;
 	}
+
 	virtual void GetNumberValue(const FMaterialRenderContext& Context,FLinearColor& OutValue) const
 	{
 		if(Context.MaterialRenderProxy->GetScalarValue(ParameterName, &OutValue.R, Context))
@@ -170,13 +195,20 @@ public:
 		}
 		else
 		{
-			OutValue.R = OutValue.G = OutValue.B = OutValue.A = DefaultValue;
+			OutValue.R = OutValue.G = OutValue.B = OutValue.A = bUseOverriddenDefault ? OverriddenDefaultValue : DefaultValue;
 		}
 	}
+
 	virtual bool IsConstant() const
 	{
 		return false;
 	}
+
+	FName GetParameterName() const
+	{
+		return ParameterName;
+	}
+
 	virtual bool IsIdentical(const FMaterialUniformExpression* OtherExpression) const
 	{
 		if (GetType() != OtherExpression->GetType())
@@ -187,9 +219,17 @@ public:
 		return ParameterName == OtherParameter->ParameterName && DefaultValue == OtherParameter->DefaultValue;
 	}
 
+	void SetTransientOverrideDefaultValue(float InOverrideDefaultValue, bool bInUseOverriddenDefault)
+	{
+		bUseOverriddenDefault = bInUseOverriddenDefault;
+		OverriddenDefaultValue = InOverrideDefaultValue;
+	}
+
 private:
 	FName ParameterName;
 	float DefaultValue;
+	bool bUseOverriddenDefault;
+	float OverriddenDefaultValue;
 };
 
 /** @return The texture that was associated with the given index when the given material had its uniform expressions/HLSL code generated. */
diff --git a/Engine/Source/Runtime/Engine/Public/EditorSupportDelegates.h b/Engine/Source/Runtime/Engine/Public/EditorSupportDelegates.h
--- a/Engine/Source/Runtime/Engine/Public/EditorSupportDelegates.h
+++ b/Engine/Source/Runtime/Engine/Public/EditorSupportDelegates.h
@@ -18,6 +18,10 @@ struct ENGINE_API FEditorSupportDelegates
 	DECLARE_MULTICAST_DELEGATE_TwoParams(FOnWindowsMessage, class FViewport*, uint32);
 	/** delegate type for material usage flags change events ( Params: UMaterial* material, int32 FlagThatChanged ) */
 	DECLARE_MULTICAST_DELEGATE_TwoParams(FOnMaterialUsageFlagsChanged, class UMaterial*, int32); 
+	/** delegate type for vector parameter default change event */
+	DECLARE_MULTICAST_DELEGATE_ThreeParams(FOnVectorParameterDefaultChanged, class UMaterialExpression*, FName, const FLinearColor&); 
+	/** delegate type for scalar parameter default change event */
+	DECLARE_MULTICAST_DELEGATE_ThreeParams(FOnScalarParameterDefaultChanged, class UMaterialExpression*, FName, float); 
 
 	/** Called when all viewports need to be redrawn */
 	static FSimpleMulticastDelegate RedrawAllViewports;
@@ -40,6 +44,10 @@ struct ENGINE_API FEditorSupportDelegates
 	static FOnWindowsMessage PostWindowsMessage;
 	/** Sent after the usages flags on a material have changed*/
 	static FOnMaterialUsageFlagsChanged MaterialUsageFlagsChanged;
+	/** Sent after vector param default changed */
+	static FOnVectorParameterDefaultChanged VectorParameterDefaultChanged;
+	/** Sent after scalar param default changed */
+	static FOnScalarParameterDefaultChanged ScalarParameterDefaultChanged;
 };
 
 #endif // WITH_EDITOR
diff --git a/Engine/Source/Runtime/Engine/Public/MaterialShared.h b/Engine/Source/Runtime/Engine/Public/MaterialShared.h
--- a/Engine/Source/Runtime/Engine/Public/MaterialShared.h
+++ b/Engine/Source/Runtime/Engine/Public/MaterialShared.h
@@ -88,8 +88,6 @@ struct ENGINE_API FMaterialRenderContext
 	const FMaterial& Material;
 	/** Whether or not selected objects should use their selection color. */
 	bool bShowSelection;
-	/** Whether to modify sampler state set with this context to work around mip map artifacts that show up in deferred passes. */
-	bool bWorkAroundDeferredMipArtifacts;
 
 	/** 
 	* Constructor
@@ -97,8 +95,7 @@ struct ENGINE_API FMaterialRenderContext
 	FMaterialRenderContext(
 		const FMaterialRenderProxy* InMaterialRenderProxy,
 		const FMaterial& InMaterial,
-		const FSceneView* InView,
-		bool bInWorkAroundDeferredMipArtifacts = false);
+		const FSceneView* InView);
 };
 
 /**
@@ -957,6 +954,8 @@ public:
 	const TArray<UMaterialExpression*>& GetErrorExpressions() const { return ErrorExpressions; }
 	ENGINE_API const TArray<TRefCountPtr<FMaterialUniformExpressionTexture> >& GetUniform2DTextureExpressions() const;
 	ENGINE_API const TArray<TRefCountPtr<FMaterialUniformExpressionTexture> >& GetUniformCubeTextureExpressions() const;
+	ENGINE_API const TArray<TRefCountPtr<FMaterialUniformExpression> >& GetUniformVectorParameterExpressions() const;
+	ENGINE_API const TArray<TRefCountPtr<FMaterialUniformExpression> >& GetUniformScalarParameterExpressions() const;
 	const FGuid& GetLegacyId() const { return Id_DEPRECATED; }
 
 	ERHIFeatureLevel::Type GetFeatureLevel() const { return FeatureLevel; }
