Convert Vector parameter turns into a Constant3, avoids compile error (vector parameter has rgb pin as primary)

diff --git a/Engine/Source/Editor/MaterialEditor/Private/MaterialEditor.cpp b/Engine/Source/Editor/MaterialEditor/Private/MaterialEditor.cpp
--- a/Engine/Source/Editor/MaterialEditor/Private/MaterialEditor.cpp
+++ b/Engine/Source/Editor/MaterialEditor/Private/MaterialEditor.cpp
@@ -2033,7 +2033,8 @@ void FMaterialEditor::OnConvertObjects()
 				}
 				else if (VectorParameterExpression)
 				{
-					ClassToCreate = UMaterialExpressionConstant4Vector::StaticClass();
+					// Technically should be a constant 4 but UMaterialExpressionVectorParameter has an rgb pin, so using Constant3 to avoid a compile error
+					ClassToCreate = UMaterialExpressionConstant3Vector::StaticClass();
 				}
 
 				if (ClassToCreate)
