Fixed convert parameter to constant crashing

diff --git a/Engine/Source/Editor/MaterialEditor/Private/MaterialEditor.cpp b/Engine/Source/Editor/MaterialEditor/Private/MaterialEditor.cpp
--- a/Engine/Source/Editor/MaterialEditor/Private/MaterialEditor.cpp
+++ b/Engine/Source/Editor/MaterialEditor/Private/MaterialEditor.cpp
@@ -2103,7 +2103,7 @@ void FMaterialEditor::OnConvertObjects()
 						else if (VectorParameterExpression)
 						{
 							bNeedsRefresh = true;
-							CastChecked<UMaterialExpressionConstant4Vector>(NewExpression)->Constant = VectorParameterExpression->DefaultValue;
+							CastChecked<UMaterialExpressionConstant3Vector>(NewExpression)->Constant = VectorParameterExpression->DefaultValue;
 						}
 
 						if (bNeedsRefresh)
