Changed Material Editor so F key results in a better fit for preview sphere.

Override Material Editor specific Static Mesh Component bounds calculation to use Aggregate Geometry Box Sphere bounds function.

Calculate tighter bounds for Aggregate Geometry that only contains a sphere.

diff --git a/Engine/Source/Editor/MaterialEditor/Private/SMaterialEditorViewport.cpp b/Engine/Source/Editor/MaterialEditor/Private/SMaterialEditorViewport.cpp
--- a/Engine/Source/Editor/MaterialEditor/Private/SMaterialEditorViewport.cpp
+++ b/Engine/Source/Editor/MaterialEditor/Private/SMaterialEditorViewport.cpp
@@ -30,6 +30,14 @@ public:
 	
 	void SetShowGrid(bool bShowGrid);
 
+	/**
+	* Focuses the viewport to the center of the bounding box/sphere ensuring that the entire bounds are in view
+	*
+	* @param Bounds   The bounds to focus
+	* @param bInstant Whether or not to focus the viewport instantly or over time
+	*/
+	void FocusViewportOnBounds(const FBoxSphereBounds Bounds, bool bInstant = false);
+
 private:
 
 	/** Pointer back to the material editor tool that owns us */
@@ -126,6 +134,48 @@ void FMaterialEditorViewportClient::SetShowGrid(bool bShowGrid)
 	DrawHelper.bDrawGrid = bShowGrid;
 }
 
+void FMaterialEditorViewportClient::FocusViewportOnBounds(const FBoxSphereBounds Bounds, bool bInstant /*= false*/)
+{
+	const FVector Position = Bounds.Origin;
+	float Radius = Bounds.SphereRadius;
+
+	float AspectToUse = AspectRatio;
+	FIntPoint ViewportSize = Viewport->GetSizeXY();
+	if (!bUseControllingActorViewInfo && ViewportSize.X > 0 && ViewportSize.Y > 0)
+	{
+		AspectToUse = Viewport->GetDesiredAspectRatio();
+	}
+
+	const bool bEnable = false;
+	ToggleOrbitCamera(bEnable);
+
+	/**
+	* We need to make sure we are fitting the sphere into the viewport completely, so if the height of the viewport is less
+	* than the width of the viewport, we scale the radius by the aspect ratio in order to compensate for the fact that we have
+	* less visible vertically than horizontally.
+	*/
+	if (AspectToUse > 1.0f)
+	{
+		Radius *= AspectToUse;
+	}
+
+	/**
+	* Now that we have a adjusted radius, we are taking half of the viewport's FOV,
+	* converting it to radians, and then figuring out the camera's distance from the center
+	* of the bounding sphere using some simple trig.  Once we have the distance, we back up
+	* along the camera's forward vector from the center of the sphere, and set our new view location.
+	*/
+	const float HalfFOVRadians = FMath::DegreesToRadians(ViewFOV / 2.0f);
+	const float DistanceFromSphere = Radius / FMath::Sin(HalfFOVRadians);
+	FVector CameraOffsetVector = ViewTransform.GetRotation().Vector() * -DistanceFromSphere;
+
+	ViewTransform.SetLookAt(Position);
+	ViewTransform.TransitionToLocation(Position + CameraOffsetVector, bInstant);
+
+	// Tell the viewport to redraw itself.
+	Invalidate();
+}
+
 void SMaterialEditorViewport::Construct(const FArguments& InArgs)
 {
 	MaterialEditorPtr = InArgs._MaterialEditor;
@@ -352,7 +402,7 @@ void SMaterialEditorViewport::OnFocusViewportToSelection()
 {
 	if( PreviewMeshComponent || PreviewSkeletalMeshComponent )
 	{
-		EditorViewportClient->FocusViewportOnBox( bUseSkeletalMeshAsPreview ? PreviewSkeletalMeshComponent->Bounds.GetBox() : PreviewMeshComponent->Bounds.GetBox() );
+		EditorViewportClient->FocusViewportOnBounds( bUseSkeletalMeshAsPreview ? PreviewSkeletalMeshComponent->Bounds : PreviewMeshComponent->Bounds );
 	}
 }
 
diff --git a/Engine/Source/Editor/UnrealEd/Classes/MaterialEditor/MaterialEditorMeshComponent.h b/Engine/Source/Editor/UnrealEd/Classes/MaterialEditor/MaterialEditorMeshComponent.h
--- a/Engine/Source/Editor/UnrealEd/Classes/MaterialEditor/MaterialEditorMeshComponent.h
+++ b/Engine/Source/Editor/UnrealEd/Classes/MaterialEditor/MaterialEditorMeshComponent.h
@@ -9,5 +9,8 @@ UCLASS()
 class UNREALED_API UMaterialEditorMeshComponent : public UStaticMeshComponent
 {
 	GENERATED_UCLASS_BODY()
+
+	// USceneComponent Interface
+	virtual FBoxSphereBounds CalcBounds(const FTransform& LocalToWorld) const override;
 };
 
diff --git a/Engine/Source/Editor/UnrealEd/Private/MaterialEditor/MaterialEditorMeshComponent.cpp b/Engine/Source/Editor/UnrealEd/Private/MaterialEditor/MaterialEditorMeshComponent.cpp
new file mode 100644
--- /dev/null
+++ b/Engine/Source/Editor/UnrealEd/Private/MaterialEditor/MaterialEditorMeshComponent.cpp
@@ -0,0 +1,33 @@
+// Copyright 1998-2012 Epic Games, Inc. All Rights Reserved.
+
+#include "UnrealEd.h"
+
+FBoxSphereBounds UMaterialEditorMeshComponent::CalcBounds(const FTransform& LocalToWorld) const
+{
+	if (StaticMesh)
+	{
+		// Graphics bounds.
+		FBoxSphereBounds NewBounds = StaticMesh->GetBounds().TransformBy(LocalToWorld);
+
+		// Add bounds of collision geometry (if present).
+		if (StaticMesh->BodySetup)
+		{
+			// Use more accurate but expensive bounds function for the material editor only
+			FBoxSphereBounds AggGeomBounds;
+			StaticMesh->BodySetup->AggGeom.CalcBoxSphereBounds(AggGeomBounds, LocalToWorld);
+			if (AggGeomBounds.SphereRadius != 0.f)
+			{
+				NewBounds = Union(NewBounds,AggGeomBounds);
+			}
+		}
+
+		NewBounds.BoxExtent *= BoundsScale;
+		NewBounds.SphereRadius *= BoundsScale;
+
+		return NewBounds;
+	}
+	else
+	{
+		return FBoxSphereBounds(LocalToWorld.GetLocation(), FVector::ZeroVector, 0.f);
+	}
+}
\ No newline at end of file
diff --git a/Engine/Source/Runtime/Engine/Private/Collision/KAggregateGeom.cpp b/Engine/Source/Runtime/Engine/Private/Collision/KAggregateGeom.cpp
--- a/Engine/Source/Runtime/Engine/Private/Collision/KAggregateGeom.cpp
+++ b/Engine/Source/Runtime/Engine/Private/Collision/KAggregateGeom.cpp
@@ -108,6 +108,13 @@ void FKAggregateGeom::CalcBoxSphereBounds(FBoxSphereBounds& Output, const FTrans
 		AABB.GetCenterAndExtents(Output.Origin, Output.BoxExtent);
 		Output.SphereRadius = FMath::Sqrt(RadiusSquared);
 	}
+	else if ((SphereElems.Num() == 1) && (SphylElems.Num() == 0) && (BoxElems.Num() == 0) && (ConvexElems.Num() == 0))
+	{
+		// For bounds that only consist of a single sphere,
+		// we can be certain the box extents are the same as its radius
+		AABB.GetCenterAndExtents(Output.Origin, Output.BoxExtent);
+		Output.SphereRadius = Output.BoxExtent.X;
+	}
 	else
 	{
 		// Just use the loose sphere bounds that totally fit the AABB
