Fix for CL 2368378.

diff --git a/Engine/Source/Editor/MaterialEditor/Private/SMaterialEditorViewport.cpp b/Engine/Source/Editor/MaterialEditor/Private/SMaterialEditorViewport.cpp
--- a/Engine/Source/Editor/MaterialEditor/Private/SMaterialEditorViewport.cpp
+++ b/Engine/Source/Editor/MaterialEditor/Private/SMaterialEditorViewport.cpp
@@ -167,6 +167,7 @@ void FMaterialEditorViewportClient::FocusViewportOnBounds(const FBoxSphereBounds
 	*/
 	const float HalfFOVRadians = FMath::DegreesToRadians(ViewFOV / 2.0f);
 	const float DistanceFromSphere = Radius / FMath::Sin(HalfFOVRadians);
+	FViewportCameraTransform& ViewTransform = GetViewTransform();
 	FVector CameraOffsetVector = ViewTransform.GetRotation().Vector() * -DistanceFromSphere;
 
 	ViewTransform.SetLookAt(Position);
