Fixed bug in determining parameter visibility with nested material functions.

diff --git a/Engine/Source/Editor/MaterialEditor/Private/MaterialEditorUtilities.cpp b/Engine/Source/Editor/MaterialEditor/Private/MaterialEditorUtilities.cpp
--- a/Engine/Source/Editor/MaterialEditor/Private/MaterialEditorUtilities.cpp
+++ b/Engine/Source/Editor/MaterialEditor/Private/MaterialEditorUtilities.cpp
@@ -214,22 +214,27 @@ void FMaterialEditorUtilities::GetVisibleMaterialParameters(const UMaterial* Mat
 	}
 }
 
-bool FMaterialEditorUtilities::GetStaticSwitchExpressionValue(UMaterialInstance* MaterialInstance, UMaterialExpression* SwitchValueExpression, bool& OutValue, FGuid& OutExpressionID, const TArray<FFunctionExpressionInput>* FunctionInputs)
+bool FMaterialEditorUtilities::GetStaticSwitchExpressionValue(UMaterialInstance* MaterialInstance, UMaterialExpression* SwitchValueExpression, bool& OutValue, FGuid& OutExpressionID, TArray<FGetVisibleMaterialParametersFunctionState*>& FunctionStack)
 {
 	// If switch value is a function input expression then we must recursively find the associated input expressions from the parent function/material to evaluate the value.
 	UMaterialExpressionFunctionInput* FunctionInputExpression =  Cast<UMaterialExpressionFunctionInput>(SwitchValueExpression);
 	if(FunctionInputExpression && FunctionInputExpression->InputType == FunctionInput_StaticBool)
 	{
+		FGetVisibleMaterialParametersFunctionState* TopmostFunctionState = FunctionStack.Pop();
+		const TArray<FFunctionExpressionInput>* FunctionInputs = TopmostFunctionState->FunctionCall ? &TopmostFunctionState->FunctionCall->FunctionInputs : NULL;
+
 		// Get the FFunctionExpressionInput which stores information about the input node from the parent that this is linked to.
 		const FFunctionExpressionInput* MatchingInput = FindInputById(FunctionInputExpression, *FunctionInputs);
 		if (MatchingInput && (MatchingInput->Input.Expression || !FunctionInputExpression->bUsePreviewValueAsDefault))
 		{
-			GetStaticSwitchExpressionValue(MaterialInstance, MatchingInput->Input.Expression, OutValue, OutExpressionID, FunctionInputs);
+			GetStaticSwitchExpressionValue(MaterialInstance, MatchingInput->Input.Expression, OutValue, OutExpressionID, FunctionStack);
 		}
 		else
 		{
-			GetStaticSwitchExpressionValue(MaterialInstance, FunctionInputExpression->Preview.Expression, OutValue, OutExpressionID, FunctionInputs);
+			GetStaticSwitchExpressionValue(MaterialInstance, FunctionInputExpression->Preview.Expression, OutValue, OutExpressionID, FunctionStack);
 		}
+
+		FunctionStack.Push(TopmostFunctionState);
 	}
 
 	if(SwitchValueExpression)
@@ -458,8 +463,7 @@ void FMaterialEditorUtilities::GetVisibleMaterialParametersFromExpression(
 
 		if (StaticSwitchExpression->Value.Expression)
 		{
-			const TArray<FFunctionExpressionInput>* FunctionInputs = FunctionStack.Top()->FunctionCall ? &(FunctionStack.Top()->FunctionCall->FunctionInputs) : NULL;
-			GetStaticSwitchExpressionValue(MaterialInstance, StaticSwitchExpression->Value.Expression, bValue, ExpressionID, FunctionInputs);
+			GetStaticSwitchExpressionValue(MaterialInstance, StaticSwitchExpression->Value.Expression, bValue, ExpressionID, FunctionStack);
 
 			if (ExpressionID.IsValid())
 			{
diff --git a/Engine/Source/Editor/MaterialEditor/Public/MaterialEditorUtilities.h b/Engine/Source/Editor/MaterialEditor/Public/MaterialEditorUtilities.h
--- a/Engine/Source/Editor/MaterialEditor/Public/MaterialEditorUtilities.h
+++ b/Engine/Source/Editor/MaterialEditor/Public/MaterialEditorUtilities.h
@@ -148,11 +148,11 @@ public:
 	* @param	SwitchValueExpression	The switch expression to find the value for.
 	* @param	OutValue				The value for the switch expression.
 	* @param	OutExpressionID			The Guid of the expression that is input as the switch value.
-	* @param	FunctionInputs			optional An array of FFunctionExpressionInputs when parsing a switch within a function.
+	* @param	FunctionStack			The current function stack frame.
 	* 
 	* @return	Returns true if a value for the switch expression is found, otherwise returns false.
 	*/
-	static bool GetStaticSwitchExpressionValue(UMaterialInstance* MaterialInstance, UMaterialExpression *SwitchValueExpression, bool& OutValue, FGuid& OutExpressionID, const TArray<FFunctionExpressionInput>* FunctionInputs);
+	static bool GetStaticSwitchExpressionValue(UMaterialInstance* MaterialInstance, UMaterialExpression *SwitchValueExpression, bool& OutValue, FGuid& OutExpressionID, TArray<FGetVisibleMaterialParametersFunctionState*>& FunctionStack);
 
 	/**
 	 * Populates the specified material's Expressions array (eg if cooked out or old content).
