Better fix for UE-3123

Placed FMaterialUpdateContext in UpdatePreviewMaterial for saftey and made PasteNodesHere use it.

diff --git a/Engine/Source/Editor/MaterialEditor/Private/MaterialEditor.cpp b/Engine/Source/Editor/MaterialEditor/Private/MaterialEditor.cpp
--- a/Engine/Source/Editor/MaterialEditor/Private/MaterialEditor.cpp
+++ b/Engine/Source/Editor/MaterialEditor/Private/MaterialEditor.cpp
@@ -1198,12 +1198,18 @@ void FMaterialEditor::UpdatePreviewMaterial( bool bForce )
 		// So that RebuildMaterialFunctionInfo will see all the nested material functions that may need to be updated
 		ExpressionPreviewMaterial->Expressions = Material->Expressions;
 
+		FMaterialUpdateContext UpdateContext;
+		UpdateContext.AddMaterial(ExpressionPreviewMaterial);
+
 		// If we are previewing an expression, update the expression preview material
 		ExpressionPreviewMaterial->PreEditChange( NULL );
 		ExpressionPreviewMaterial->PostEditChange();
 	}
 	else 
 	{
+		FMaterialUpdateContext UpdateContext;
+		UpdateContext.AddMaterial(Material);
+
 		// Update the regular preview material when not previewing an expression.
 		Material->PreEditChange( NULL );
 		Material->PostEditChange();
@@ -3318,14 +3324,10 @@ void FMaterialEditor::PasteNodesHere(const FVector2D& Location)
 		Node->CreateNewGuid();
 	}
 
-	// Force new pasted Material Expressions to have same connections as graph nodes
-	Material->MaterialGraph->LinkMaterialExpressionsFromGraph();
+	UpdateMaterialAfterGraphChange();
 
 	// Update UI
 	GraphEditor->NotifyGraphChanged();
-
-	Material->PostEditChange();
-	Material->MarkPackageDirty();
 }
 
 bool FMaterialEditor::CanPasteNodes() const
