Fixed performance issue introduced in CL2407391.
UpdatePreviewMaterial no longer triggers a global component reregister.

diff --git a/Engine/Source/Editor/MaterialEditor/Private/MaterialEditor.cpp b/Engine/Source/Editor/MaterialEditor/Private/MaterialEditor.cpp
--- a/Engine/Source/Editor/MaterialEditor/Private/MaterialEditor.cpp
+++ b/Engine/Source/Editor/MaterialEditor/Private/MaterialEditor.cpp
@@ -1189,7 +1189,7 @@ void FMaterialEditor::UpdatePreviewMaterial( bool bForce )
 		// So that RebuildMaterialFunctionInfo will see all the nested material functions that may need to be updated
 		ExpressionPreviewMaterial->Expressions = Material->Expressions;
 
-		FMaterialUpdateContext UpdateContext;
+		FMaterialUpdateContext UpdateContext(FMaterialUpdateContext::EOptions::SyncWithRenderingThread);
 		UpdateContext.AddMaterial(ExpressionPreviewMaterial);
 
 		// If we are previewing an expression, update the expression preview material
@@ -1198,7 +1198,7 @@ void FMaterialEditor::UpdatePreviewMaterial( bool bForce )
 	}
 	else 
 	{
-		FMaterialUpdateContext UpdateContext;
+		FMaterialUpdateContext UpdateContext(FMaterialUpdateContext::EOptions::SyncWithRenderingThread);
 		UpdateContext.AddMaterial(Material);
 
 		// Update the regular preview material when not previewing an expression.
