Only recompile materials on save if dirty

diff --git a/Engine/Source/Editor/MaterialEditor/Private/MaterialEditor.cpp b/Engine/Source/Editor/MaterialEditor/Private/MaterialEditor.cpp
--- a/Engine/Source/Editor/MaterialEditor/Private/MaterialEditor.cpp
+++ b/Engine/Source/Editor/MaterialEditor/Private/MaterialEditor.cpp
@@ -814,7 +814,10 @@ void FMaterialEditor::SaveAsset_Execute()
 {
 	UE_LOG(LogMaterialEditor, Log, TEXT("Saving and Compiling material %s"), *GetEditingObjects()[0]->GetName());
 	
-	UpdateOriginalMaterial();
+	if (bMaterialDirty)
+	{
+		UpdateOriginalMaterial();
+	}
 
 	UPackage* Package = OriginalMaterial->GetOutermost();
 
