Fixed issue with menu listings for Material Expressions that don't begin with 'MaterialExpression'

diff --git a/Engine/Source/Editor/MaterialEditor/Private/MaterialExpressionClasses.cpp b/Engine/Source/Editor/MaterialEditor/Private/MaterialExpressionClasses.cpp
--- a/Engine/Source/Editor/MaterialEditor/Private/MaterialExpressionClasses.cpp
+++ b/Engine/Source/Editor/MaterialEditor/Private/MaterialExpressionClasses.cpp
@@ -91,7 +91,13 @@ void MaterialExpressionClasses::InitMaterialExpressionClasses()
 						{
 							FMaterialExpression MaterialExpression;
 							// Trim the material expression name and add it to the list used for filtering.
-							MaterialExpression.Name = FString(*Class->GetName()).Mid(FCString::Strlen(TEXT("MaterialExpression")));
+							static const FString ExpressionPrefix = TEXT("MaterialExpression");
+							FString ClassName = *Class->GetName();
+							if (ClassName.StartsWith(ExpressionPrefix, ESearchCase::CaseSensitive))
+							{
+								ClassName = ClassName.Mid(ExpressionPrefix.Len());
+							}
+							MaterialExpression.Name = ClassName;
 							MaterialExpression.MaterialClass = Class;
 
 							AllExpressionClasses.Add(MaterialExpression);
