Material Editor: Add support for previewing any kind of asset with an associated UMeshComponent subclass, such as static meshes, skeletal meshes, and Paper2D sprites
[UE-10085]

diff --git a/Engine/Source/Editor/MaterialEditor/Private/MaterialEditor.cpp b/Engine/Source/Editor/MaterialEditor/Private/MaterialEditor.cpp
--- a/Engine/Source/Editor/MaterialEditor/Private/MaterialEditor.cpp
+++ b/Engine/Source/Editor/MaterialEditor/Private/MaterialEditor.cpp
@@ -430,10 +430,10 @@ void FMaterialEditor::InitMaterialEditor( const EToolkitMode::Type Mode, const T
 	LoadEditorSettings();
 	
 	// Set the preview mesh for the material.  This call must occur after the toolbar is initialized.
-	if (!SetPreviewMesh(*Material->PreviewMesh.ToString()))
+	if (!SetPreviewAssetByName(*Material->PreviewMesh.ToString()))
 	{
 		// The material preview mesh couldn't be found or isn't loaded.  Default to the one of the primitive types.
-		Viewport->SetPreviewMesh( GUnrealEd->GetThumbnailManager()->EditorSphere, NULL );
+		SetPreviewAsset( GUnrealEd->GetThumbnailManager()->EditorSphere );
 	}
 
 	// Initialize expression previews.
@@ -798,15 +798,20 @@ UMaterialInterface* FMaterialEditor::GetMaterialInterface() const
 	return Material;
 }
 
-bool FMaterialEditor::ApproveSetPreviewMesh(UStaticMesh* InStaticMesh, USkeletalMesh* InSkeletalMesh)
+bool FMaterialEditor::ApproveSetPreviewAsset(UObject* InAsset)
 {
 	bool bApproved = true;
+
 	// Only permit the use of a skeletal mesh if the material has bUsedWithSkeltalMesh.
-	if ( InSkeletalMesh && !Material->bUsedWithSkeletalMesh )
+	if (USkeletalMesh* SkeletalMesh = Cast<USkeletalMesh>(InAsset))
 	{
-		FMessageDialog::Open( EAppMsgType::Ok, NSLOCTEXT("UnrealEd", "Error_MaterialEditor_CantPreviewOnSkelMesh", "Can't preview on the specified skeletal mesh because the material has not been compiled with bUsedWithSkeletalMesh.") );
-		bApproved = false;
+		if (!Material->bUsedWithSkeletalMesh)
+		{
+			FMessageDialog::Open(EAppMsgType::Ok, NSLOCTEXT("UnrealEd", "Error_MaterialEditor_CantPreviewOnSkelMesh", "Can't preview on the specified skeletal mesh because the material has not been compiled with bUsedWithSkeletalMesh."));
+			bApproved = false;
+		}
 	}
+
 	return bApproved;
 }
 
@@ -1015,20 +1020,20 @@ void FMaterialEditor::RecenterEditor()
 	}
 }
 
-bool FMaterialEditor::SetPreviewMesh(UStaticMesh* InStaticMesh, USkeletalMesh* InSkeletalMesh)
+bool FMaterialEditor::SetPreviewAsset(UObject* InAsset)
 {
 	if (Viewport.IsValid())
 	{
-		return Viewport->SetPreviewMesh(InStaticMesh, InSkeletalMesh);
+		return Viewport->SetPreviewAsset(InAsset);
 	}
 	return false;
 }
 
-bool FMaterialEditor::SetPreviewMesh(const TCHAR* InMeshName)
+bool FMaterialEditor::SetPreviewAssetByName(const TCHAR* InAssetName)
 {
 	if (Viewport.IsValid())
 	{
-		return Viewport->SetPreviewMesh(InMeshName);
+		return Viewport->SetPreviewAssetByName(InAssetName);
 	}
 	return false;
 }
@@ -3442,14 +3447,14 @@ void FMaterialEditor::NotifyPostChange( const FPropertyChangedEvent& PropertyCha
 	if ( PropertyThatChanged )
 	{
 		const FName NameOfPropertyThatChanged( *PropertyThatChanged->GetName() );
-		if ( NameOfPropertyThatChanged == FName(TEXT("PreviewMesh")) ||
-			NameOfPropertyThatChanged == FName(TEXT("bUsedWithSkeletalMesh")) )
+		if ((NameOfPropertyThatChanged == GET_MEMBER_NAME_CHECKED(UMaterialInterface, PreviewMesh)) ||
+			(NameOfPropertyThatChanged == GET_MEMBER_NAME_CHECKED(UMaterial, bUsedWithSkeletalMesh)))
 		{
 			// SetPreviewMesh will return false if the material has bUsedWithSkeletalMesh and
 			// a skeleton was requested, in which case revert to a sphere static mesh.
-			if (!SetPreviewMesh(*Material->PreviewMesh.ToString()))
+			if (!SetPreviewAssetByName(*Material->PreviewMesh.ToString()))
 			{
-				SetPreviewMesh( GUnrealEd->GetThumbnailManager()->EditorSphere, NULL );
+				SetPreviewAsset(GUnrealEd->GetThumbnailManager()->EditorSphere);
 			}
 		}
 
diff --git a/Engine/Source/Editor/MaterialEditor/Private/MaterialEditor.h b/Engine/Source/Editor/MaterialEditor/Private/MaterialEditor.h
--- a/Engine/Source/Editor/MaterialEditor/Private/MaterialEditor.h
+++ b/Engine/Source/Editor/MaterialEditor/Private/MaterialEditor.h
@@ -249,8 +249,8 @@ public:
 	void RecenterEditor();
 
 	/** Passes instructions to the preview viewport */
-	bool SetPreviewMesh(UStaticMesh* InStaticMesh, USkeletalMesh* InSkeletalMesh);
-	bool SetPreviewMesh(const TCHAR* InMeshName);
+	bool SetPreviewAsset(UObject* InAsset);
+	bool SetPreviewAssetByName(const TCHAR* InAssetName);
 	void SetPreviewMaterial(UMaterialInterface* InMaterialInterface);
 	
 	/**
@@ -436,8 +436,8 @@ private:
 	/** Builds the toolbar widget for the material editor */
 	void ExtendToolbar();
 
-	/** Allows editor to veto the setting of a preview mesh */
-	virtual bool ApproveSetPreviewMesh(UStaticMesh* InStaticMesh, USkeletalMesh* InSkeletalMesh) override;
+	/** Allows editor to veto the setting of a preview asset */
+	virtual bool ApproveSetPreviewAsset(UObject* InAsset) override;
 
 	/** Creates all internal widgets for the tabs to point at */
 	void CreateInternalWidgets();
diff --git a/Engine/Source/Editor/MaterialEditor/Private/MaterialInstanceEditor.cpp b/Engine/Source/Editor/MaterialEditor/Private/MaterialInstanceEditor.cpp
--- a/Engine/Source/Editor/MaterialEditor/Private/MaterialInstanceEditor.cpp
+++ b/Engine/Source/Editor/MaterialEditor/Private/MaterialInstanceEditor.cpp
@@ -232,7 +232,7 @@ void FMaterialInstanceEditor::InitMaterialInstanceEditor( const EToolkitMode::Ty
 	{
 		InstanceConstant->PreviewMesh = ThumbnailInfoWithPrim->PreviewMesh;
 	}
-	SetPreviewMesh(*InstanceConstant->PreviewMesh.ToString());
+	SetPreviewAssetByName(*InstanceConstant->PreviewMesh.ToString());
 }
 
 FMaterialInstanceEditor::~FMaterialInstanceEditor()
@@ -810,20 +810,20 @@ void FMaterialInstanceEditor::DrawSamplerWarningStrings(FCanvas* Canvas, int32&
 	}
 }
 
-bool FMaterialInstanceEditor::SetPreviewMesh(UStaticMesh* InStaticMesh, USkeletalMesh* InSkeletalMesh)
+bool FMaterialInstanceEditor::SetPreviewAsset(UObject* InAsset)
 {
 	if (PreviewVC.IsValid())
 	{
-		return PreviewVC->SetPreviewMesh(InStaticMesh, InSkeletalMesh);
+		return PreviewVC->SetPreviewAsset(InAsset);
 	}
 	return false;
 }
 
-bool FMaterialInstanceEditor::SetPreviewMesh(const TCHAR* InMeshName)
+bool FMaterialInstanceEditor::SetPreviewAssetByName(const TCHAR* InAssetName)
 {
 	if (PreviewVC.IsValid())
 	{
-		return PreviewVC->SetPreviewMesh(InMeshName);
+		return PreviewVC->SetPreviewAssetByName(InAssetName);
 	}
 	return false;
 }
@@ -934,7 +934,7 @@ UObject* FMaterialInstanceEditor::GetSyncObject()
 	return NULL;
 }
 
-bool FMaterialInstanceEditor::ApproveSetPreviewMesh(UStaticMesh* InStaticMesh, USkeletalMesh* InSkeletalMesh)
+bool FMaterialInstanceEditor::ApproveSetPreviewAsset(UObject* InAsset)
 {
 	// Default impl is to always accept.
 	return true;
diff --git a/Engine/Source/Editor/MaterialEditor/Private/MaterialInstanceEditor.h b/Engine/Source/Editor/MaterialEditor/Private/MaterialInstanceEditor.h
--- a/Engine/Source/Editor/MaterialEditor/Private/MaterialInstanceEditor.h
+++ b/Engine/Source/Editor/MaterialEditor/Private/MaterialInstanceEditor.h
@@ -66,8 +66,8 @@ public:
 	void DrawSamplerWarningStrings(FCanvas* Canvas, int32& DrawPositionY);
 
 	/** Passes instructions to the preview viewport */
-	bool SetPreviewMesh(UStaticMesh* InStaticMesh, USkeletalMesh* InSkeletalMesh);
-	bool SetPreviewMesh(const TCHAR* InMeshName);
+	bool SetPreviewAsset(UObject* InAsset);
+	bool SetPreviewAssetByName(const TCHAR* InMeshName);
 	void SetPreviewMaterial(UMaterialInterface* InMaterialInterface);
 
 	/** Returns true if hidden parameters should be shown */
@@ -121,8 +121,8 @@ private:
 	/** Builds the toolbar widget for the material editor */
 	void ExtendToolbar();
 
-	/** Allows editor to veto the setting of a preview mesh */
-	virtual bool ApproveSetPreviewMesh(UStaticMesh* InStaticMesh, USkeletalMesh* InSkeletalMesh) override;
+	// IMaterialEditor interface
+	virtual bool ApproveSetPreviewAsset(UObject* InAsset) override;
 
 	/**	Spawns the preview tab */
 	TSharedRef<SDockTab> SpawnTab_Preview( const FSpawnTabArgs& Args );
diff --git a/Engine/Source/Editor/MaterialEditor/Private/SMaterialEditorViewport.cpp b/Engine/Source/Editor/MaterialEditor/Private/SMaterialEditorViewport.cpp
--- a/Engine/Source/Editor/MaterialEditor/Private/SMaterialEditorViewport.cpp
+++ b/Engine/Source/Editor/MaterialEditor/Private/SMaterialEditorViewport.cpp
@@ -15,6 +15,8 @@
 #include "Engine/StaticMesh.h"
 #include "Engine/Selection.h"
 #include "Engine/TextureCube.h"
+#include "ComponentAssetBroker.h"
+
 /** Viewport Client for the preview viewport */
 class FMaterialEditorViewportClient : public FEditorViewportClient
 {
@@ -193,23 +195,23 @@ void SMaterialEditorViewport::Construct(const FArguments& InArgs)
 
 	SEditorViewport::Construct( SEditorViewport::FArguments() );
 
+	PreviewMaterial = nullptr;
+	PreviewMeshComponent = nullptr;
 
-	PreviewMeshComponent = NewObject<UMaterialEditorMeshComponent>(GetTransientPackage(), NAME_None, RF_Transient);
-	PreviewSkeletalMeshComponent = NewObject<USkeletalMeshComponent>(GetTransientPackage(), NAME_None, RF_Transient);
-
-	UMaterialInterface* Material = MaterialEditorPtr.Pin()->GetMaterialInterface();
-	if (Material)
+	if (UMaterialInterface* Material = MaterialEditorPtr.Pin()->GetMaterialInterface())
 	{
 		SetPreviewMaterial(Material);
 	}
 
-	SetPreviewMesh( GUnrealEd->GetThumbnailManager()->EditorSphere, NULL );
+	SetPreviewAsset( GUnrealEd->GetThumbnailManager()->EditorSphere );
 }
 
 SMaterialEditorViewport::~SMaterialEditorViewport()
 {
-	PreviewMeshComponent->OverrideMaterials.Empty();
-	PreviewSkeletalMeshComponent->OverrideMaterials.Empty();
+	if (PreviewMeshComponent != nullptr)
+	{
+		PreviewMeshComponent->OverrideMaterials.Empty();
+	}
 
 	if (EditorViewportClient.IsValid())
 	{
@@ -220,55 +222,57 @@ SMaterialEditorViewport::~SMaterialEditorViewport()
 void SMaterialEditorViewport::AddReferencedObjects( FReferenceCollector& Collector )
 {
 	Collector.AddReferencedObject( PreviewMeshComponent );
-	Collector.AddReferencedObject( PreviewSkeletalMeshComponent );
+	Collector.AddReferencedObject( PreviewMaterial );
 }
 
 void SMaterialEditorViewport::RefreshViewport()
 {
-	//reregister the preview components, so if the preview material changed it will be propagated to the render thread
-	PreviewMeshComponent->MarkRenderStateDirty();
-	PreviewSkeletalMeshComponent->MarkRenderStateDirty();
+	// reregister the preview components, so if the preview material changed it will be propagated to the render thread
+	if (PreviewMeshComponent != nullptr)
+	{
+		PreviewMeshComponent->MarkRenderStateDirty();
+	}
 	SceneViewport->InvalidateDisplay();
 }
 
-bool SMaterialEditorViewport::SetPreviewMesh(UStaticMesh* InStaticMesh, USkeletalMesh* InSkeletalMesh)
+bool SMaterialEditorViewport::SetPreviewAsset(UObject* InAsset)
 {
-	// Only permit the use of a skeletal mesh if the material has bUsedWithSkeltalMesh.
-	if (!MaterialEditorPtr.Pin()->ApproveSetPreviewMesh(InStaticMesh, InSkeletalMesh))
+	if (!MaterialEditorPtr.Pin()->ApproveSetPreviewAsset(InAsset))
 	{
 		return false;
 	}
-	
-	FTransform Transform = FTransform::Identity;
-	bUseSkeletalMeshAsPreview = InSkeletalMesh != NULL;
-	if ( bUseSkeletalMeshAsPreview )
-	{
-		// Remove the static mesh from the preview scene.
-		PreviewScene.RemoveComponent( PreviewMeshComponent );
-		PreviewScene.AddComponent( PreviewSkeletalMeshComponent, Transform );
-		
-		// Update the toolbar state implicitly through PreviewPrimType.
-		PreviewPrimType = TPT_None;
 
-		// Set the new preview skeletal mesh.
-		PreviewSkeletalMeshComponent->SetSkeletalMesh( InSkeletalMesh );
+	// Unregister the current component
+	if (PreviewMeshComponent != nullptr)
+	{
+		PreviewScene.RemoveComponent(PreviewMeshComponent);
+		PreviewMeshComponent = nullptr;
 	}
-	else
+
+	FTransform Transform = FTransform::Identity;
+
+	if (UStaticMesh* StaticMesh = Cast<UStaticMesh>(InAsset))
 	{
+		// Special case handling for static meshes, to use more accurate bounds via a subclass
+		UStaticMeshComponent* NewSMComponent = NewObject<UMaterialEditorMeshComponent>(GetTransientPackage(), NAME_None, RF_Transient);
+		NewSMComponent->SetStaticMesh(StaticMesh);
+
+		PreviewMeshComponent = NewSMComponent;
+
 		// Update the toolbar state implicitly through PreviewPrimType.
-		if ( InStaticMesh == GUnrealEd->GetThumbnailManager()->EditorCylinder )
+		if (StaticMesh == GUnrealEd->GetThumbnailManager()->EditorCylinder)
 		{
 			PreviewPrimType = TPT_Cylinder;
 		}
-		else if ( InStaticMesh == GUnrealEd->GetThumbnailManager()->EditorCube )
+		else if (StaticMesh == GUnrealEd->GetThumbnailManager()->EditorCube)
 		{
 			PreviewPrimType = TPT_Cube;
 		}
-		else if ( InStaticMesh == GUnrealEd->GetThumbnailManager()->EditorSphere )
+		else if (StaticMesh == GUnrealEd->GetThumbnailManager()->EditorSphere)
 		{
 			PreviewPrimType = TPT_Sphere;
 		}
-		else if ( InStaticMesh == GUnrealEd->GetThumbnailManager()->EditorPlane )
+		else if (StaticMesh == GUnrealEd->GetThumbnailManager()->EditorPlane)
 		{
 			// Need to rotate the plane so that it faces the camera
 			Transform.SetRotation(FQuat(FRotator(0, 90, 0)));
@@ -278,36 +282,43 @@ bool SMaterialEditorViewport::SetPreviewMesh(UStaticMesh* InStaticMesh, USkeleta
 		{
 			PreviewPrimType = TPT_None;
 		}
+	}
+	else if (InAsset != nullptr)
+	{
+		// Fall back to the component asset broker
+		if (TSubclassOf<UActorComponent> ComponentClass = FComponentAssetBrokerage::GetPrimaryComponentForAsset(InAsset->GetClass()))
+		{
+			if (ComponentClass->IsChildOf(UMeshComponent::StaticClass()))
+			{
+				PreviewMeshComponent = NewObject<UMeshComponent>(GetTransientPackage(), ComponentClass, NAME_None, RF_Transient);
 
-		// Remove the skeletal mesh from the preview scene.
-		PreviewScene.RemoveComponent( PreviewSkeletalMeshComponent );
-		PreviewScene.AddComponent( PreviewMeshComponent, Transform );
+				FComponentAssetBrokerage::AssignAssetToComponent(PreviewMeshComponent, InAsset);
 
-		// Set the new preview static mesh.
-		FComponentReregisterContext ReregisterContext( PreviewMeshComponent );
-		PreviewMeshComponent->StaticMesh = InStaticMesh;
+				PreviewPrimType = TPT_None;
+			}
+		}
 	}
 
-	return true;
+	// Add the new component to the scene
+	if (PreviewMeshComponent != nullptr)
+	{
+		PreviewScene.AddComponent(PreviewMeshComponent, Transform);
+	}
+
+	// Make sure the preview material is applied to the component
+	SetPreviewMaterial(PreviewMaterial);
+
+	return (PreviewMeshComponent != nullptr);
 }
 
-bool SMaterialEditorViewport::SetPreviewMesh(const TCHAR* InMeshName)
+bool SMaterialEditorViewport::SetPreviewAssetByName(const TCHAR* InAssetName)
 {
 	bool bSuccess = false;
-	if ( InMeshName && FString(InMeshName).Len() > 0 )
+	if ((InAssetName != nullptr) && (*InAssetName != 0))
 	{
-		UStaticMesh* StaticMesh = LoadObject<UStaticMesh>( NULL, InMeshName);
-		if ( StaticMesh )
+		if (UObject* Asset = LoadObject<UObject>(nullptr, InAssetName))
 		{
-			bSuccess = SetPreviewMesh( StaticMesh, NULL );
-		}
-		else
-		{
-			USkeletalMesh* SkeletalMesh = LoadObject<USkeletalMesh>( NULL, InMeshName );
-			if ( SkeletalMesh )
-			{
-				bSuccess = SetPreviewMesh( NULL, SkeletalMesh );
-			}
+			bSuccess = SetPreviewAsset(Asset);
 		}
 	}
 	return bSuccess;
@@ -315,13 +326,13 @@ bool SMaterialEditorViewport::SetPreviewMesh(const TCHAR* InMeshName)
 
 void SMaterialEditorViewport::SetPreviewMaterial(UMaterialInterface* InMaterialInterface)
 {
-	check( PreviewMeshComponent );
-	check( PreviewSkeletalMeshComponent );
+	PreviewMaterial = InMaterialInterface;
 
-	PreviewMeshComponent->OverrideMaterials.Empty();
-	PreviewMeshComponent->OverrideMaterials.Add( InMaterialInterface );
-	PreviewSkeletalMeshComponent->OverrideMaterials.Empty();
-	PreviewSkeletalMeshComponent->OverrideMaterials.Add( InMaterialInterface );
+	if (PreviewMeshComponent != nullptr)
+	{
+		PreviewMeshComponent->OverrideMaterials.Empty();
+		PreviewMeshComponent->OverrideMaterials.Add(PreviewMaterial);
+	}
 }
 
 void SMaterialEditorViewport::ToggleRealtime()
@@ -397,9 +408,9 @@ void SMaterialEditorViewport::BindCommands()
 
 void SMaterialEditorViewport::OnFocusViewportToSelection()
 {
-	if( PreviewMeshComponent || PreviewSkeletalMeshComponent )
+	if( PreviewMeshComponent != nullptr )
 	{
-		EditorViewportClient->FocusViewportOnBounds( bUseSkeletalMeshAsPreview ? PreviewSkeletalMeshComponent->Bounds : PreviewMeshComponent->Bounds );
+		EditorViewportClient->FocusViewportOnBounds( PreviewMeshComponent->Bounds );
 	}
 }
 
@@ -407,7 +418,7 @@ void SMaterialEditorViewport::OnSetPreviewPrimitive(EThumbnailPrimType PrimType)
 {
 	if (SceneViewport.IsValid())
 	{
-		UStaticMesh* Primitive = NULL;
+		UStaticMesh* Primitive = nullptr;
 		switch (PrimType)
 		{
 		case TPT_Cylinder: Primitive = GUnrealEd->GetThumbnailManager()->EditorCylinder; break;
@@ -415,9 +426,10 @@ void SMaterialEditorViewport::OnSetPreviewPrimitive(EThumbnailPrimType PrimType)
 		case TPT_Plane: Primitive = GUnrealEd->GetThumbnailManager()->EditorPlane; break;
 		case TPT_Cube: Primitive = GUnrealEd->GetThumbnailManager()->EditorCube; break;
 		}
-		if (Primitive)
+
+		if (Primitive != nullptr)
 		{
-			SetPreviewMesh(Primitive , NULL );
+			SetPreviewAsset(Primitive);
 			RefreshViewport();
 		}
 	}
@@ -433,35 +445,37 @@ void SMaterialEditorViewport::OnSetPreviewMeshFromSelection()
 	bool bFoundPreviewMesh = false;
 	FEditorDelegates::LoadSelectedAssetsIfNeeded.Broadcast();
 
-	// Look for a selected static mesh.
-	UStaticMesh* SelectedStaticMesh = GEditor->GetSelectedObjects()->GetTop<UStaticMesh>();
 	UMaterialInterface* MaterialInterface = MaterialEditorPtr.Pin()->GetMaterialInterface();
-	if ( SelectedStaticMesh )
-	{
-		SetPreviewMesh( SelectedStaticMesh, NULL );
-		MaterialInterface->PreviewMesh = SelectedStaticMesh->GetPathName();
-		bFoundPreviewMesh = true;
-	}
-	else
+
+	// Look for a selected asset that can be converted to a mesh component
+	for (FSelectionIterator SelectionIt(*GEditor->GetSelectedObjects()); SelectionIt && !bFoundPreviewMesh; ++SelectionIt)
 	{
-		// No static meshes were selected; look for a selected skeletal mesh.
-		USkeletalMesh* SelectedSkeletalMesh = GEditor->GetSelectedObjects()->GetTop<USkeletalMesh>();
-		if ( SelectedSkeletalMesh )
+		UObject* TestAsset = *SelectionIt;
+		if (TestAsset->IsAsset())
 		{
-			//Automatically set the usage flag.
-			if( MaterialInterface->GetMaterial() )
+			if (TSubclassOf<UActorComponent> ComponentClass = FComponentAssetBrokerage::GetPrimaryComponentForAsset(TestAsset->GetClass()))
 			{
-				bool bNeedsRecompile = false;
-				MaterialInterface->GetMaterial()->SetMaterialUsage( bNeedsRecompile, MATUSAGE_SkeletalMesh );
+				if (ComponentClass->IsChildOf(UMeshComponent::StaticClass()))
+				{
+					if (USkeletalMesh* SkeletalMesh = Cast<USkeletalMesh>(TestAsset))
+					{
+						// Special case handling for skeletal meshes, sets the material to be usable with them
+						if (MaterialInterface->GetMaterial())
+						{
+							bool bNeedsRecompile = false;
+							MaterialInterface->GetMaterial()->SetMaterialUsage(bNeedsRecompile, MATUSAGE_SkeletalMesh);
+						}
+					}
+
+					SetPreviewAsset(TestAsset);
+					MaterialInterface->PreviewMesh = TestAsset->GetPathName();
+					bFoundPreviewMesh = true;
+				}
 			}
-
-			SetPreviewMesh( NULL, SelectedSkeletalMesh );
-			MaterialInterface->PreviewMesh = SelectedSkeletalMesh->GetPathName();
-			bFoundPreviewMesh = true;
 		}
 	}
 
-	if ( bFoundPreviewMesh )
+	if (bFoundPreviewMesh)
 	{
 		FMaterialEditor::UpdateThumbnailInfoPreviewMesh(MaterialInterface);
 
@@ -470,7 +484,7 @@ void SMaterialEditorViewport::OnSetPreviewMeshFromSelection()
 	}
 	else
 	{
-		FSuppressableWarningDialog::FSetupInfo Info(NSLOCTEXT("UnrealEd", "Warning_NoPreviewMeshFound_Message", "You need to select a mesh in the content browser to preview it."),
+		FSuppressableWarningDialog::FSetupInfo Info(NSLOCTEXT("UnrealEd", "Warning_NoPreviewMeshFound_Message", "You need to select a mesh-based asset in the content browser to preview it."),
 			NSLOCTEXT("UnrealEd", "Warning_NoPreviewMeshFound", "Warning: No Preview Mesh Found"), "Warning_NoPreviewMeshFound");
 		Info.ConfirmText = NSLOCTEXT("UnrealEd", "Warning_NoPreviewMeshFound_Confirm", "Continue");
 		
diff --git a/Engine/Source/Editor/MaterialEditor/Private/SMaterialEditorViewport.h b/Engine/Source/Editor/MaterialEditor/Private/SMaterialEditorViewport.h
--- a/Engine/Source/Editor/MaterialEditor/Private/SMaterialEditorViewport.h
+++ b/Engine/Source/Editor/MaterialEditor/Private/SMaterialEditorViewport.h
@@ -29,16 +29,16 @@ public:
 	 *
 	 * @return	true if a mesh was set successfully, false otherwise.
 	 */
-	bool SetPreviewMesh(UStaticMesh* InStaticMesh, USkeletalMesh* InSkeletalMesh);
+	bool SetPreviewAsset(UObject* InAsset);
 
 	/**
-	 * Sets the preview mesh to the named mesh, if it can be found.  Checks static meshes first, then skeletal meshes.
-	 * Does nothing if the named mesh is not found or if the named mesh is a skeletal mesh but the material has
+	 * Sets the preview asset to the named asset, if it can be found and is convertible into a UMeshComponent subclass.
+	 * Does nothing if the named asset is not found or if the named asset is a skeletal mesh but the material has
 	 * bUsedWithSkeletalMesh=false.
 	 *
-	 * @return	true if the named mesh was found and set successfully, false otherwise.
+	 * @return	true if the named asset was found and set successfully, false otherwise.
 	 */
-	bool SetPreviewMesh(const TCHAR* InMeshName);
+	bool SetPreviewAssetByName(const TCHAR* InAssetName);
 
 	void SetPreviewMaterial(UMaterialInterface* InMaterialInterface);
 	
@@ -47,15 +47,12 @@ public:
 	/** @return The list of commands known by the material editor */
 	TSharedRef<FUICommandList> GetMaterialEditorCommands() const;
 
-	/** If true, use PreviewSkeletalMeshComponent as the preview mesh; if false, use PreviewMeshComponent. */
-	bool bUseSkeletalMeshAsPreview;
+	/** Component for the preview mesh. */
+	UMeshComponent* PreviewMeshComponent;
 
-	/** Component for the preview static mesh. */
-	UMaterialEditorMeshComponent* PreviewMeshComponent;
+	/** Material for the preview mesh */
+	UMaterialInterface* PreviewMaterial;
 	
-	/** Component for the preview skeletal mesh. */
-	USkeletalMeshComponent*	PreviewSkeletalMeshComponent;
-
 	/** The preview primitive we are using. */
 	EThumbnailPrimType PreviewPrimType;
 	
diff --git a/Engine/Source/Editor/MaterialEditor/Public/IMaterialEditor.h b/Engine/Source/Editor/MaterialEditor/Public/IMaterialEditor.h
--- a/Engine/Source/Editor/MaterialEditor/Public/IMaterialEditor.h
+++ b/Engine/Source/Editor/MaterialEditor/Public/IMaterialEditor.h
@@ -87,7 +87,7 @@ public:
 	virtual UMaterialInterface* GetMaterialInterface() const = 0;
 
 	/** Allows editor to veto the setting of a preview mesh */
-	virtual bool ApproveSetPreviewMesh(UStaticMesh* InStaticMesh, USkeletalMesh* InSkeletalMesh) = 0;
+	virtual bool ApproveSetPreviewAsset(UObject* InAsset) = 0;
 
 	/** Draws messages on the specified viewport and canvas. */
 	virtual void DrawMessages( FViewport* Viewport, FCanvas* Canvas ) = 0;
