Made the "Set Preview Mesh" in the Material Editor Viewport highlight when selected.

diff --git a/Engine/Source/Editor/MaterialEditor/Private/MaterialEditorActions.cpp b/Engine/Source/Editor/MaterialEditor/Private/MaterialEditorActions.cpp
--- a/Engine/Source/Editor/MaterialEditor/Private/MaterialEditorActions.cpp
+++ b/Engine/Source/Editor/MaterialEditor/Private/MaterialEditorActions.cpp
@@ -18,7 +18,7 @@ void FMaterialEditorCommands::RegisterCommands()
 	UI_COMMAND( SetSpherePreview, "Sphere", "Sets the preview mesh to a sphere primitive.", EUserInterfaceActionType::ToggleButton, FInputGesture() );
 	UI_COMMAND( SetPlanePreview, "Plane", "Sets the preview mesh to a plane primitive.", EUserInterfaceActionType::ToggleButton, FInputGesture() );
 	UI_COMMAND( SetCubePreview, "Cube", "Sets the preview mesh to a cube primitive.", EUserInterfaceActionType::ToggleButton, FInputGesture() );
-	UI_COMMAND( SetPreviewMeshFromSelection, "Mesh", "Sets the preview mesh based on the current content browser selection.", EUserInterfaceActionType::Button, FInputGesture() );
+	UI_COMMAND( SetPreviewMeshFromSelection, "Mesh", "Sets the preview mesh based on the current content browser selection.", EUserInterfaceActionType::ToggleButton, FInputGesture() );
 	UI_COMMAND( TogglePreviewGrid, "Grid", "Toggles the preview pane's grid.", EUserInterfaceActionType::ToggleButton, FInputGesture() );
 	UI_COMMAND( TogglePreviewBackground, "Background", "Toggles the preview pane's background.", EUserInterfaceActionType::ToggleButton, FInputGesture() );
 
diff --git a/Engine/Source/Editor/MaterialEditor/Private/SMaterialEditorViewport.cpp b/Engine/Source/Editor/MaterialEditor/Private/SMaterialEditorViewport.cpp
--- a/Engine/Source/Editor/MaterialEditor/Private/SMaterialEditorViewport.cpp
+++ b/Engine/Source/Editor/MaterialEditor/Private/SMaterialEditorViewport.cpp
@@ -391,7 +391,8 @@ void SMaterialEditorViewport::BindCommands()
 	ToolkitCommandList->MapAction(
 		Commands.SetPreviewMeshFromSelection,
 		FExecuteAction::CreateSP( this, &SMaterialEditorViewport::OnSetPreviewMeshFromSelection ),
-		FCanExecuteAction());
+		FCanExecuteAction(),
+		FIsActionChecked::CreateSP( this, &SMaterialEditorViewport::IsPreviewMeshFromSelectionChecked ) );
 
 	ToolkitCommandList->MapAction(
 		Commands.TogglePreviewGrid,
@@ -493,6 +494,11 @@ void SMaterialEditorViewport::OnSetPreviewMeshFromSelection()
 	}
 }
 
+bool SMaterialEditorViewport::IsPreviewMeshFromSelectionChecked() const
+{
+	return (PreviewPrimType == TPT_None && PreviewMeshComponent != nullptr);
+}
+
 void SMaterialEditorViewport::TogglePreviewGrid()
 {
 	bShowGrid = !bShowGrid;
diff --git a/Engine/Source/Editor/MaterialEditor/Private/SMaterialEditorViewport.h b/Engine/Source/Editor/MaterialEditor/Private/SMaterialEditorViewport.h
--- a/Engine/Source/Editor/MaterialEditor/Private/SMaterialEditorViewport.h
+++ b/Engine/Source/Editor/MaterialEditor/Private/SMaterialEditorViewport.h
@@ -72,6 +72,7 @@ public:
 	void OnSetPreviewPrimitive(EThumbnailPrimType PrimType);
 	bool IsPreviewPrimitiveChecked(EThumbnailPrimType PrimType) const;
 	void OnSetPreviewMeshFromSelection();
+	bool IsPreviewMeshFromSelectionChecked() const;
 	void TogglePreviewGrid();
 	bool IsTogglePreviewGridChecked() const;
 	void TogglePreviewBackground();
