Material editor preview now refreshes correctly when a referenced texture is (re)imported.

diff --git a/Engine/Source/Editor/MaterialEditor/Private/MaterialEditor.cpp b/Engine/Source/Editor/MaterialEditor/Private/MaterialEditor.cpp
--- a/Engine/Source/Editor/MaterialEditor/Private/MaterialEditor.cpp
+++ b/Engine/Source/Editor/MaterialEditor/Private/MaterialEditor.cpp
@@ -326,6 +326,7 @@ void FMaterialEditor::InitMaterialEditor( const EToolkitMode::Type Mode, const T
 	FEditorSupportDelegates::MaterialUsageFlagsChanged.AddRaw(this, &FMaterialEditor::OnMaterialUsageFlagsChanged);
 	FEditorSupportDelegates::VectorParameterDefaultChanged.AddRaw(this, &FMaterialEditor::OnVectorParameterDefaultChanged);
 	FEditorSupportDelegates::ScalarParameterDefaultChanged.AddRaw(this, &FMaterialEditor::OnScalarParameterDefaultChanged);
+	FEditorDelegates::OnAssetPostImport.AddRaw(this, &FMaterialEditor::OnAssetPostImport);
 
 	FAssetRegistryModule& AssetRegistryModule = FModuleManager::GetModuleChecked<FAssetRegistryModule>(TEXT("AssetRegistry"));
 	
@@ -544,6 +545,7 @@ FMaterialEditor::~FMaterialEditor()
 	FEditorSupportDelegates::MaterialUsageFlagsChanged.RemoveAll(this);
 	FEditorSupportDelegates::VectorParameterDefaultChanged.RemoveAll(this);
 	FEditorSupportDelegates::ScalarParameterDefaultChanged.RemoveAll(this);
+	FEditorDelegates::OnAssetPostImport.RemoveAll(this);
 
 	// Null out the expression preview material so they can be GC'ed
 	ExpressionPreviewMaterial = NULL;
@@ -2435,6 +2437,64 @@ void FMaterialEditor::RenameAssetFromRegistry(const FAssetData& InAddedAssetData
 	}
 }
 
+
+void FMaterialEditor::OnAssetPostImport(UFactory* InFactory, UObject* InObject)
+{
+	UTexture* Texture = Cast<UTexture>(InObject);
+
+	if (InFactory->IsA(UTextureFactory::StaticClass()) && Texture != nullptr)
+	{
+		// When a texture which is referenced in the material is imported, update the preview material
+		for (UMaterialExpression* Expression : Material->Expressions)
+		{
+			if (UMaterialExpressionTextureBase* ExpressionTexture = Cast<UMaterialExpressionTextureBase>(Expression))
+			{
+				if (ExpressionTexture->Texture == Texture)
+				{
+					UpdatePreviewMaterial(true);
+					break;
+				}
+			}
+			else if (UMaterialExpressionMaterialFunctionCall* ExpressionFunctionCall = Cast<UMaterialExpressionMaterialFunctionCall>(Expression))
+			{
+				struct Local
+				{
+					static bool ReferencesImportedObject(UMaterialFunction* Function, UTexture* Texture)
+					{
+						for (UMaterialExpression* Expression : Function->FunctionExpressions)
+						{
+							if (UMaterialExpressionTextureBase* ExpressionTexture = Cast<UMaterialExpressionTextureBase>(Expression))
+							{
+								if (ExpressionTexture->Texture == Texture)
+								{
+									return true;
+								}
+							}
+							else if (UMaterialExpressionMaterialFunctionCall* ExpressionFunctionCall = Cast<UMaterialExpressionMaterialFunctionCall>(Expression))
+							{
+								if (ExpressionFunctionCall->MaterialFunction != nullptr &&
+									ReferencesImportedObject(ExpressionFunctionCall->MaterialFunction, Texture))
+								{
+									return true;
+								}
+							}
+						}
+
+						return false;
+					}
+				};
+
+				if (ExpressionFunctionCall->MaterialFunction != nullptr &&
+					Local::ReferencesImportedObject(ExpressionFunctionCall->MaterialFunction, Texture))
+				{
+					UpdatePreviewMaterial(true);
+					break;
+				}
+			}
+		}
+	}
+}
+
 void FMaterialEditor::OnMaterialUsageFlagsChanged(UMaterial* MaterialThatChanged, int32 FlagThatChanged)
 {
 	EMaterialUsage Flag = static_cast<EMaterialUsage>(FlagThatChanged);
diff --git a/Engine/Source/Editor/MaterialEditor/Private/MaterialEditor.h b/Engine/Source/Editor/MaterialEditor/Private/MaterialEditor.h
--- a/Engine/Source/Editor/MaterialEditor/Private/MaterialEditor.h
+++ b/Engine/Source/Editor/MaterialEditor/Private/MaterialEditor.h
@@ -540,6 +540,9 @@ private:
 	/** Callback to tell the Material Editor that a materials usage flags have been changed */
 	void OnMaterialUsageFlagsChanged(class UMaterial* MaterialThatChanged, int32 FlagThatChanged);
 
+	/** Callback when an asset is imported */
+	void OnAssetPostImport(UFactory* InFactory, UObject* InObject);
+
 	void OnVectorParameterDefaultChanged(class UMaterialExpression*, FName ParameterName, const FLinearColor& Value);
 	void OnScalarParameterDefaultChanged(class UMaterialExpression*, FName ParameterName, float Value);
 
diff --git a/Engine/Source/Editor/UnrealEd/Private/Factories/EditorFactories.cpp b/Engine/Source/Editor/UnrealEd/Private/Factories/EditorFactories.cpp
--- a/Engine/Source/Editor/UnrealEd/Private/Factories/EditorFactories.cpp
+++ b/Engine/Source/Editor/UnrealEd/Private/Factories/EditorFactories.cpp
@@ -13,6 +13,7 @@
 #include "Sound/SoundNodeAttenuation.h"
 #include "Sound/SoundNodeModulator.h"
 #include "Factories.h"
+#include "NormalMapIdentification.h"
 #include "SoundDefinitions.h"
 #include "PhysicsPublic.h"
 #include "BlueprintUtilities.h"
@@ -4490,7 +4491,10 @@ UObject* UTextureFactory::FactoryCreateBinary
 		Texture2D->bHasBeenPaintedInEditor = false;
 	}
 
-	FEditorDelegates::OnAssetPostImport.Broadcast( this, Texture );
+	// Automatically detect if the texture is a normal map and configure its properties accordingly
+	NormalMapIdentification::HandleAssetPostImport(this, Texture);
+
+	FEditorDelegates::OnAssetPostImport.Broadcast(this, Texture);
 
 	// Invalidate any materials using the newly imported texture. (occurs if you import over an existing texture)
 	Texture->PostEditChange();
diff --git a/Engine/Source/Editor/UnrealEd/Private/NormalMapIdentification.cpp b/Engine/Source/Editor/UnrealEd/Private/NormalMapIdentification.cpp
--- a/Engine/Source/Editor/UnrealEd/Private/NormalMapIdentification.cpp
+++ b/Engine/Source/Editor/UnrealEd/Private/NormalMapIdentification.cpp
@@ -448,10 +448,8 @@ public:
 	TWeakPtr<SNotificationItem> Notification;
 };
 
-void NormalMapIdentification::HandleAssetPostImport( UFactory* InFactory, UObject* InObject )
+void NormalMapIdentification::HandleAssetPostImport( UTextureFactory* TextureFactory, UTexture* Texture )
 {
-	UTextureFactory* TextureFactory = Cast<UTextureFactory>(InFactory);
-	UTexture* Texture = Cast<UTexture>(InObject);
 	if(TextureFactory != NULL && Texture != NULL)
 	{
 		// Try to automatically identify a normal map
diff --git a/Engine/Source/Editor/UnrealEd/Private/NormalMapIdentification.h b/Engine/Source/Editor/UnrealEd/Private/NormalMapIdentification.h
--- a/Engine/Source/Editor/UnrealEd/Private/NormalMapIdentification.h
+++ b/Engine/Source/Editor/UnrealEd/Private/NormalMapIdentification.h
@@ -6,8 +6,8 @@ namespace NormalMapIdentification
 {
 	/**
 	 * Handle callback when an asset is imported.
-	 * @param	InFactory	The factory being used.
-	 * @param	InObject	The object that was imported.
+	 * @param	InFactory	The texture factory being used.
+	 * @param	InTexture	The texture that was imported.
 	 */
-	void HandleAssetPostImport( UFactory* InFactory, UObject* InObject );
+	void HandleAssetPostImport( UTextureFactory* InFactory, UTexture* InTexture );
 }
diff --git a/Engine/Source/Editor/UnrealEd/Private/UnrealEdMisc.cpp b/Engine/Source/Editor/UnrealEd/Private/UnrealEdMisc.cpp
--- a/Engine/Source/Editor/UnrealEd/Private/UnrealEdMisc.cpp
+++ b/Engine/Source/Editor/UnrealEd/Private/UnrealEdMisc.cpp
@@ -29,7 +29,6 @@
 #include "PackageTools.h"
 #include "GameProjectGenerationModule.h"
 #include "MaterialEditorActions.h"
-#include "NormalMapIdentification.h"
 #include "EngineBuildSettings.h"
 #include "SlateBasics.h"
 #include "DesktopPlatformModule.h"
@@ -200,8 +199,6 @@ void FUnrealEdMisc::OnInit()
 	FCoreDelegates::PreModal.AddRaw(this, &FUnrealEdMisc::OnEditorPreModal);
 	FCoreDelegates::PostModal.AddRaw(this, &FUnrealEdMisc::OnEditorPostModal);
 
-	FEditorDelegates::OnAssetPostImport.AddStatic(&NormalMapIdentification::HandleAssetPostImport);
-
 	// Register the play world commands
 	FPlayWorldCommands::Register();
 	FPlayWorldCommands::BindGlobalPlayWorldCommands();
