Material editor viewport now properly drives Slate tick/paint when realtime is enabled

diff --git a/Engine/Source/Editor/MaterialEditor/Private/MaterialEditor.cpp b/Engine/Source/Editor/MaterialEditor/Private/MaterialEditor.cpp
--- a/Engine/Source/Editor/MaterialEditor/Private/MaterialEditor.cpp
+++ b/Engine/Source/Editor/MaterialEditor/Private/MaterialEditor.cpp
@@ -1072,7 +1072,7 @@ void FMaterialEditor::LoadEditorSettings()
 	{
 		if (EditorOptions->bShowGrid) {Viewport->TogglePreviewGrid();}
 		if (EditorOptions->bShowBackground) {Viewport->TogglePreviewBackground();}
-		if (EditorOptions->bRealtimeMaterialViewport) {Viewport->ToggleRealtime();}
+		if (EditorOptions->bRealtimeMaterialViewport) {Viewport->OnToggleRealtime();}
 
 		// Load the preview scene
 		Viewport->PreviewScene.LoadSettings(TEXT("MaterialEditor"));
@@ -1719,7 +1719,7 @@ void FMaterialEditor::BindCommands()
 
 	ToolkitCommands->MapAction(
 		FEditorViewportCommands::Get().ToggleRealTime,
-		FExecuteAction::CreateSP( Viewport.ToSharedRef(), &SMaterialEditorViewport::ToggleRealtime ),
+		FExecuteAction::CreateSP( Viewport.ToSharedRef(), &SMaterialEditorViewport::OnToggleRealtime ),
 		FCanExecuteAction(),
 		FIsActionChecked::CreateSP( Viewport.ToSharedRef(), &SMaterialEditorViewport::IsRealtime ) );
 
diff --git a/Engine/Source/Editor/MaterialEditor/Private/MaterialInstanceEditor.cpp b/Engine/Source/Editor/MaterialEditor/Private/MaterialInstanceEditor.cpp
--- a/Engine/Source/Editor/MaterialEditor/Private/MaterialInstanceEditor.cpp
+++ b/Engine/Source/Editor/MaterialEditor/Private/MaterialInstanceEditor.cpp
@@ -273,7 +273,7 @@ void FMaterialInstanceEditor::BindCommands()
 
 	ToolkitCommands->MapAction(
 		FEditorViewportCommands::Get().ToggleRealTime,
-		FExecuteAction::CreateSP( PreviewVC.ToSharedRef(), &SMaterialEditorViewport::ToggleRealtime ),
+		FExecuteAction::CreateSP( PreviewVC.ToSharedRef(), &SMaterialEditorViewport::OnToggleRealtime ),
 		FCanExecuteAction(),
 		FIsActionChecked::CreateSP( PreviewVC.ToSharedRef(), &SMaterialEditorViewport::IsRealtime ) );
 
@@ -879,7 +879,7 @@ void FMaterialInstanceEditor::LoadSettings()
 		}
 		if ( bRealtime )
 		{
-			PreviewVC->ToggleRealtime();
+			PreviewVC->OnToggleRealtime();
 		}
 		PreviewVC->OnSetPreviewPrimitive( static_cast<EThumbnailPrimType>(PrimType) );
 
diff --git a/Engine/Source/Editor/MaterialEditor/Private/SMaterialEditorViewport.cpp b/Engine/Source/Editor/MaterialEditor/Private/SMaterialEditorViewport.cpp
--- a/Engine/Source/Editor/MaterialEditor/Private/SMaterialEditorViewport.cpp
+++ b/Engine/Source/Editor/MaterialEditor/Private/SMaterialEditorViewport.cpp
@@ -335,11 +335,6 @@ void SMaterialEditorViewport::SetPreviewMaterial(UMaterialInterface* InMaterialI
 	}
 }
 
-void SMaterialEditorViewport::ToggleRealtime()
-{
-	EditorViewportClient->ToggleRealtime();
-}
-
 TSharedRef<FUICommandList> SMaterialEditorViewport::GetMaterialEditorCommands() const
 {
 	check(MaterialEditorPtr.IsValid());
diff --git a/Engine/Source/Editor/MaterialEditor/Private/SMaterialEditorViewport.h b/Engine/Source/Editor/MaterialEditor/Private/SMaterialEditorViewport.h
--- a/Engine/Source/Editor/MaterialEditor/Private/SMaterialEditorViewport.h
+++ b/Engine/Source/Editor/MaterialEditor/Private/SMaterialEditorViewport.h
@@ -42,8 +42,6 @@ public:
 
 	void SetPreviewMaterial(UMaterialInterface* InMaterialInterface);
 	
-	void ToggleRealtime();
-
 	/** @return The list of commands known by the material editor */
 	TSharedRef<FUICommandList> GetMaterialEditorCommands() const;
 
diff --git a/Engine/Source/Editor/UnrealEd/Public/SEditorViewport.h b/Engine/Source/Editor/UnrealEd/Public/SEditorViewport.h
--- a/Engine/Source/Editor/UnrealEd/Public/SEditorViewport.h
+++ b/Engine/Source/Editor/UnrealEd/Public/SEditorViewport.h
@@ -33,6 +33,9 @@ public:
 	 */
 	void Invalidate();
 
+	/** Toggles realtime on/off for the viewport. Slate tick/paint is ensured when realtime is on. */
+	void OnToggleRealtime();
+
 	/**
 	 * @return true if the specified coordinate system the active one active
 	 */
@@ -61,7 +64,6 @@ protected:
 	virtual EVisibility OnGetViewportContentVisibility() const;
 
 	/** UI Command delegate bindings */
-	void OnToggleRealtime();
 	void OnToggleStats();
 
 	/**
