Material Editor: Switched to use the full standard in-viewport toolbar
- Switched over to use SCommonEditorViewportToolbarBase, which includes the 'Camera', 'View Mode', Show', etc... menus
- Moved 'preview shape' toolbar to bottom right and removed 'real time' and 'show grid' options, which are now in the standard menu locations

diff --git a/Engine/Source/Editor/MaterialEditor/Private/SMaterialEditorViewport.cpp b/Engine/Source/Editor/MaterialEditor/Private/SMaterialEditorViewport.cpp
--- a/Engine/Source/Editor/MaterialEditor/Private/SMaterialEditorViewport.cpp
+++ b/Engine/Source/Editor/MaterialEditor/Private/SMaterialEditorViewport.cpp
@@ -335,12 +335,6 @@ void SMaterialEditorViewport::SetPreviewMaterial(UMaterialInterface* InMaterialI
 	}
 }
 
-TSharedRef<FUICommandList> SMaterialEditorViewport::GetMaterialEditorCommands() const
-{
-	check(MaterialEditorPtr.IsValid());
-	return MaterialEditorPtr.Pin()->GetToolkitCommands();
-}
-
 void SMaterialEditorViewport::OnAddedToTab( const TSharedRef<SDockTab>& OwnerTab )
 {
 	ParentTab = OwnerTab;
@@ -356,46 +350,48 @@ void SMaterialEditorViewport::BindCommands()
 	SEditorViewport::BindCommands();
 
 	const FMaterialEditorCommands& Commands = FMaterialEditorCommands::Get();
-	TSharedRef<FUICommandList> ToolkitCommandList = GetMaterialEditorCommands();
+
+	check(MaterialEditorPtr.IsValid());
+	CommandList->Append(MaterialEditorPtr.Pin()->GetToolkitCommands());
 
 	// Add the commands to the toolkit command list so that the toolbar buttons can find them
-	ToolkitCommandList->MapAction(
+	CommandList->MapAction(
 		Commands.SetCylinderPreview,
 		FExecuteAction::CreateSP( this, &SMaterialEditorViewport::OnSetPreviewPrimitive, TPT_Cylinder ),
 		FCanExecuteAction(),
 		FIsActionChecked::CreateSP( this, &SMaterialEditorViewport::IsPreviewPrimitiveChecked, TPT_Cylinder ) );
 
-	ToolkitCommandList->MapAction(
+	CommandList->MapAction(
 		Commands.SetSpherePreview,
 		FExecuteAction::CreateSP( this, &SMaterialEditorViewport::OnSetPreviewPrimitive, TPT_Sphere ),
 		FCanExecuteAction(),
 		FIsActionChecked::CreateSP( this, &SMaterialEditorViewport::IsPreviewPrimitiveChecked, TPT_Sphere ) );
 
-	ToolkitCommandList->MapAction(
+	CommandList->MapAction(
 		Commands.SetPlanePreview,
 		FExecuteAction::CreateSP( this, &SMaterialEditorViewport::OnSetPreviewPrimitive, TPT_Plane ),
 		FCanExecuteAction(),
 		FIsActionChecked::CreateSP( this, &SMaterialEditorViewport::IsPreviewPrimitiveChecked, TPT_Plane ) );
 
-	ToolkitCommandList->MapAction(
+	CommandList->MapAction(
 		Commands.SetCubePreview,
 		FExecuteAction::CreateSP( this, &SMaterialEditorViewport::OnSetPreviewPrimitive, TPT_Cube ),
 		FCanExecuteAction(),
 		FIsActionChecked::CreateSP( this, &SMaterialEditorViewport::IsPreviewPrimitiveChecked, TPT_Cube ) );
 
-	ToolkitCommandList->MapAction(
+	CommandList->MapAction(
 		Commands.SetPreviewMeshFromSelection,
 		FExecuteAction::CreateSP( this, &SMaterialEditorViewport::OnSetPreviewMeshFromSelection ),
 		FCanExecuteAction(),
 		FIsActionChecked::CreateSP( this, &SMaterialEditorViewport::IsPreviewMeshFromSelectionChecked ) );
 
-	ToolkitCommandList->MapAction(
+	CommandList->MapAction(
 		Commands.TogglePreviewGrid,
 		FExecuteAction::CreateSP( this, &SMaterialEditorViewport::TogglePreviewGrid ),
 		FCanExecuteAction(),
 		FIsActionChecked::CreateSP( this, &SMaterialEditorViewport::IsTogglePreviewGridChecked ) );
 
-	ToolkitCommandList->MapAction(
+	CommandList->MapAction(
 		Commands.TogglePreviewBackground,
 		FExecuteAction::CreateSP( this, &SMaterialEditorViewport::TogglePreviewBackground ),
 		FCanExecuteAction(),
@@ -518,6 +514,21 @@ bool SMaterialEditorViewport::IsTogglePreviewBackgroundChecked() const
 	return bShowBackground;
 }
 
+TSharedRef<class SEditorViewport> SMaterialEditorViewport::GetViewportWidget()
+{
+	return SharedThis(this);
+}
+
+TSharedPtr<FExtender> SMaterialEditorViewport::GetExtenders() const
+{
+	TSharedPtr<FExtender> Result(MakeShareable(new FExtender));
+	return Result;
+}
+
+void SMaterialEditorViewport::OnFloatingButtonClicked()
+{
+}
+
 TSharedRef<FEditorViewportClient> SMaterialEditorViewport::MakeEditorViewportClient() 
 {
 	EditorViewportClient = MakeShareable( new FMaterialEditorViewportClient(MaterialEditorPtr, PreviewScene, SharedThis(this)) );
@@ -533,10 +544,19 @@ TSharedRef<FEditorViewportClient> SMaterialEditorViewport::MakeEditorViewportCli
 	return EditorViewportClient.ToSharedRef();
 }
 
-TSharedPtr<SWidget> SMaterialEditorViewport::MakeViewportToolbar()
+void SMaterialEditorViewport::PopulateViewportOverlays(TSharedRef<class SOverlay> Overlay)
 {
-	return SNew(SMaterialEditorViewportToolBar)
-		.Viewport(SharedThis(this));
+	Overlay->AddSlot()
+		.VAlign(VAlign_Top)
+		[
+			SNew(SMaterialEditorViewportToolBar, SharedThis(this))
+		];
+
+	Overlay->AddSlot()
+		.VAlign(VAlign_Bottom)
+		[
+			SNew(SMaterialEditorViewportPreviewShapeToolBar, SharedThis(this))
+		];
 }
 
 EVisibility SMaterialEditorViewport::OnGetViewportContentVisibility() const
diff --git a/Engine/Source/Editor/MaterialEditor/Private/SMaterialEditorViewport.h b/Engine/Source/Editor/MaterialEditor/Private/SMaterialEditorViewport.h
--- a/Engine/Source/Editor/MaterialEditor/Private/SMaterialEditorViewport.h
+++ b/Engine/Source/Editor/MaterialEditor/Private/SMaterialEditorViewport.h
@@ -4,12 +4,12 @@
 
 #include "PreviewScene.h"
 #include "SEditorViewport.h"
-
+#include "Editor/UnrealEd/Public/SCommonEditorViewportToolbarBase.h"
 
 /**
  * Material Editor Preview viewport widget
  */
-class SMaterialEditorViewport : public SEditorViewport, public FGCObject
+class SMaterialEditorViewport : public SEditorViewport, public FGCObject, public ICommonEditorViewportToolbarInfoProvider
 {
 public:
 	SLATE_BEGIN_ARGS( SMaterialEditorViewport ){}
@@ -41,9 +41,6 @@ public:
 	bool SetPreviewAssetByName(const TCHAR* InAssetName);
 
 	void SetPreviewMaterial(UMaterialInterface* InMaterialInterface);
-	
-	/** @return The list of commands known by the material editor */
-	TSharedRef<FUICommandList> GetMaterialEditorCommands() const;
 
 	/** Component for the preview mesh. */
 	UMeshComponent* PreviewMeshComponent;
@@ -76,10 +73,16 @@ public:
 	void TogglePreviewBackground();
 	bool IsTogglePreviewBackgroundChecked() const;
 
+	// ICommonEditorViewportToolbarInfoProvider interface
+	virtual TSharedRef<class SEditorViewport> GetViewportWidget() override;
+	virtual TSharedPtr<FExtender> GetExtenders() const override;
+	virtual void OnFloatingButtonClicked() override;
+	// End of ICommonEditorViewportToolbarInfoProvider interface
+
 protected:
 	/** SEditorViewport interface */
 	virtual TSharedRef<FEditorViewportClient> MakeEditorViewportClient() override;
-	virtual TSharedPtr<SWidget> MakeViewportToolbar() override;
+	virtual void PopulateViewportOverlays(TSharedRef<class SOverlay> Overlay) override;
 	virtual EVisibility OnGetViewportContentVisibility() const override;
 	virtual void BindCommands() override;
 	virtual void OnFocusViewportToSelection() override;
diff --git a/Engine/Source/Editor/MaterialEditor/Private/SMaterialEditorViewportToolBar.cpp b/Engine/Source/Editor/MaterialEditor/Private/SMaterialEditorViewportToolBar.cpp
--- a/Engine/Source/Editor/MaterialEditor/Private/SMaterialEditorViewportToolBar.cpp
+++ b/Engine/Source/Editor/MaterialEditor/Private/SMaterialEditorViewportToolBar.cpp
@@ -5,14 +5,18 @@
 #include "SMaterialEditorViewport.h"
 #include "EditorViewportCommands.h"
 #include "MaterialEditorActions.h"
+#include "Editor/UnrealEd/Public/SViewportToolBar.h"
 
 #define LOCTEXT_NAMESPACE "MaterialEditorViewportToolBar"
 
-void SMaterialEditorViewportToolBar::Construct(const FArguments& InArgs)
+///////////////////////////////////////////////////////////
+// SMaterialEditorViewportPreviewShapeToolBar
+
+void SMaterialEditorViewportPreviewShapeToolBar::Construct(const FArguments& InArgs, TSharedPtr<class SMaterialEditorViewport> InViewport)
 {
 	// Force this toolbar to have small icons, as the preview panel is only small so we have limited space
 	const bool bForceSmallIcons = true;
-	FToolBarBuilder ToolbarBuilder(InArgs._Viewport->GetMaterialEditorCommands(), FMultiBoxCustomization::None, nullptr, Orient_Horizontal, bForceSmallIcons);
+	FToolBarBuilder ToolbarBuilder(InViewport->GetCommandList(), FMultiBoxCustomization::None, nullptr, Orient_Horizontal, bForceSmallIcons);
 
 	// Use a custom style
 	ToolbarBuilder.SetStyle(&FEditorStyle::Get(), "ViewportMenu");
@@ -26,8 +30,6 @@ void SMaterialEditorViewportToolBar::Construct(const FArguments& InArgs)
 		ToolbarBuilder.AddToolBarButton(FMaterialEditorCommands::Get().SetPlanePreview);
 		ToolbarBuilder.AddToolBarButton(FMaterialEditorCommands::Get().SetCubePreview);
 		ToolbarBuilder.AddToolBarButton(FMaterialEditorCommands::Get().SetPreviewMeshFromSelection);
-		ToolbarBuilder.AddToolBarButton(FMaterialEditorCommands::Get().TogglePreviewGrid);
-		ToolbarBuilder.AddToolBarButton(FEditorViewportCommands::Get().ToggleRealTime);
 	}
 	ToolbarBuilder.EndSection();
 
@@ -49,4 +51,35 @@ void SMaterialEditorViewportToolBar::Construct(const FArguments& InArgs)
 	SViewportToolBar::Construct(SViewportToolBar::FArguments());
 }
 
+///////////////////////////////////////////////////////////
+// SMaterialEditorViewportToolBar
+
+void SMaterialEditorViewportToolBar::Construct(const FArguments& InArgs, TSharedPtr<class SMaterialEditorViewport> InViewport)
+{
+	SCommonEditorViewportToolbarBase::Construct(SCommonEditorViewportToolbarBase::FArguments(), InViewport);
+}
+
+TSharedRef<SWidget> SMaterialEditorViewportToolBar::GenerateShowMenu() const
+{
+	GetInfoProvider().OnFloatingButtonClicked();
+
+	TSharedRef<SEditorViewport> ViewportRef = GetInfoProvider().GetViewportWidget();
+
+	const bool bInShouldCloseWindowAfterMenuSelection = true;
+	FMenuBuilder ShowMenuBuilder(bInShouldCloseWindowAfterMenuSelection, ViewportRef->GetCommandList());
+	{
+		auto Commands = FMaterialEditorCommands::Get();
+
+		ShowMenuBuilder.AddMenuEntry(Commands.ToggleMaterialStats);
+		ShowMenuBuilder.AddMenuEntry(Commands.ToggleMobileStats);
+
+		ShowMenuBuilder.AddMenuSeparator();
+
+		ShowMenuBuilder.AddMenuEntry(Commands.TogglePreviewGrid);
+		ShowMenuBuilder.AddMenuEntry(Commands.TogglePreviewBackground);
+	}
+
+	return ShowMenuBuilder.MakeWidget();
+}
+
 #undef LOCTEXT_NAMESPACE
diff --git a/Engine/Source/Editor/MaterialEditor/Private/SMaterialEditorViewportToolBar.h b/Engine/Source/Editor/MaterialEditor/Private/SMaterialEditorViewportToolBar.h
--- a/Engine/Source/Editor/MaterialEditor/Private/SMaterialEditorViewportToolBar.h
+++ b/Engine/Source/Editor/MaterialEditor/Private/SMaterialEditorViewportToolBar.h
@@ -1,22 +1,34 @@
 // Copyright 1998-2015 Epic Games, Inc. All Rights Reserved.
 
+#pragma once
 
-#ifndef __SLevelViewportToolBar_h__
-#define __SLevelViewportToolBar_h__
+#include "Editor/UnrealEd/Public/SCommonEditorViewportToolbarBase.h"
 
-#include "Editor/UnrealEd/Public/SViewportToolBar.h"
+///////////////////////////////////////////////////////////
+// SMaterialEditorViewportToolBar
 
-/**
- * A material editor viewport toolbar widget that is placed in a viewport
- */
-class SMaterialEditorViewportToolBar : public SViewportToolBar
+// In-viewport toolbar widget used in the material editor
+class SMaterialEditorViewportToolBar : public SCommonEditorViewportToolbarBase
 {
 public:
-	SLATE_BEGIN_ARGS( SMaterialEditorViewportToolBar ){}
-		SLATE_ARGUMENT( TSharedPtr<class SMaterialEditorViewport>, Viewport )
+	SLATE_BEGIN_ARGS(SMaterialEditorViewportToolBar) {}
 	SLATE_END_ARGS()
 
-	void Construct( const FArguments& InArgs );
+	void Construct(const FArguments& InArgs, TSharedPtr<class SMaterialEditorViewport> InViewport);
+
+	// SCommonEditorViewportToolbarBase interface
+	virtual TSharedRef<SWidget> GenerateShowMenu() const override;
+	// End of SCommonEditorViewportToolbarBase
 };
 
-#endif
+///////////////////////////////////////////////////////////
+// SMaterialEditorViewportPreviewShapeToolBar
+
+class SMaterialEditorViewportPreviewShapeToolBar : public SViewportToolBar
+{
+public:
+	SLATE_BEGIN_ARGS(SMaterialEditorViewportPreviewShapeToolBar){}
+	SLATE_END_ARGS()
+
+	void Construct(const FArguments& InArgs, TSharedPtr<class SMaterialEditorViewport> InViewport);
+};
