Fix so that material thumbnail is updated when a dependent texture is reimported.

diff --git a/Engine/Source/Editor/MaterialEditor/Private/MaterialEditor.cpp b/Engine/Source/Editor/MaterialEditor/Private/MaterialEditor.cpp
--- a/Engine/Source/Editor/MaterialEditor/Private/MaterialEditor.cpp
+++ b/Engine/Source/Editor/MaterialEditor/Private/MaterialEditor.cpp
@@ -2451,7 +2451,7 @@ void FMaterialEditor::OnAssetPostImport(UFactory* InFactory, UObject* InObject)
 			{
 				if (ExpressionTexture->Texture == Texture)
 				{
-					UpdatePreviewMaterial(true);
+					UpdateOriginalMaterial();
 					break;
 				}
 			}
@@ -2487,7 +2487,7 @@ void FMaterialEditor::OnAssetPostImport(UFactory* InFactory, UObject* InObject)
 				if (ExpressionFunctionCall->MaterialFunction != nullptr &&
 					Local::ReferencesImportedObject(ExpressionFunctionCall->MaterialFunction, Texture))
 				{
-					UpdatePreviewMaterial(true);
+					UpdateOriginalMaterial();
 					break;
 				}
 			}
