UE4 - Fix ES3.1 emulation shader issues

diff --git a/Engine/Shaders/LocalVertexFactoryCommon.usf b/Engine/Shaders/LocalVertexFactoryCommon.usf
--- a/Engine/Shaders/LocalVertexFactoryCommon.usf
+++ b/Engine/Shaders/LocalVertexFactoryCommon.usf
@@ -6,7 +6,7 @@
 
 // FVertexFactoryInterpolantsVSToPS is split into a PC and ES2 version, since the ES2 version has more restrictions on packing
 // Then accessors are provided for each property which try to encapsulate the packing differences
-#if FEATURE_LEVEL >= FEATURE_LEVEL_SM4
+#if FEATURE_LEVEL >= FEATURE_LEVEL_ES3_1
 
 struct FVertexFactoryInterpolantsVSToPS
 {
@@ -105,7 +105,7 @@ void SetTangents(inout FVertexFactoryInterpolantsVSToPS Interpolants, float3 InT
 }
 
 
-#else // #if FEATURE_LEVEL >= FEATURE_LEVEL_SM4
+#else // #if FEATURE_LEVEL >= FEATURE_LEVEL_ES3_1
 
 #define NUM_PACKED_UVS (NUM_MATERIAL_TEXCOORDS > 1 ? 3 : 1)
 
diff --git a/Engine/Source/Developer/ShaderFormatOpenGL/Private/ShaderFormatOpenGL.cpp b/Engine/Source/Developer/ShaderFormatOpenGL/Private/ShaderFormatOpenGL.cpp
--- a/Engine/Source/Developer/ShaderFormatOpenGL/Private/ShaderFormatOpenGL.cpp
+++ b/Engine/Source/Developer/ShaderFormatOpenGL/Private/ShaderFormatOpenGL.cpp
@@ -15,7 +15,7 @@ static FName NAME_GLSL_ES2(TEXT("GLSL_ES2"));
 static FName NAME_GLSL_ES2_WEBGL(TEXT("GLSL_ES2_WEBGL"));
 static FName NAME_GLSL_150_ES2(TEXT("GLSL_150_ES2"));
 static FName NAME_GLSL_150_ES2_NOUB(TEXT("GLSL_150_ES2_NOUB"));
-static FName NAME_GLSL_150_ES3_1(TEXT("GLSL_150_ES3_1"));
+static FName NAME_GLSL_150_ES3_1(TEXT("GLSL_150_ES31"));
 static FName NAME_GLSL_ES2_IOS(TEXT("GLSL_ES2_IOS"));
 static FName NAME_GLSL_310_ES_EXT(TEXT("GLSL_310_ES_EXT"));
  
diff --git a/Engine/Source/Developer/Windows/ShaderFormatD3D/Private/ShaderFormatD3D.cpp b/Engine/Source/Developer/Windows/ShaderFormatD3D/Private/ShaderFormatD3D.cpp
--- a/Engine/Source/Developer/Windows/ShaderFormatD3D/Private/ShaderFormatD3D.cpp
+++ b/Engine/Source/Developer/Windows/ShaderFormatD3D/Private/ShaderFormatD3D.cpp
@@ -8,7 +8,7 @@
 
 static FName NAME_PCD3D_SM5(TEXT("PCD3D_SM5"));
 static FName NAME_PCD3D_SM4(TEXT("PCD3D_SM4"));
-static FName NAME_PCD3D_ES3_1(TEXT("PCD3D_ES3_1"));
+static FName NAME_PCD3D_ES3_1(TEXT("PCD3D_ES31"));
 static FName NAME_PCD3D_ES2(TEXT("PCD3D_ES2"));
 
 class FShaderFormatD3D : public IShaderFormat
diff --git a/Engine/Source/Editor/LevelEditor/Private/LevelEditorActions.cpp b/Engine/Source/Editor/LevelEditor/Private/LevelEditorActions.cpp
--- a/Engine/Source/Editor/LevelEditor/Private/LevelEditorActions.cpp
+++ b/Engine/Source/Editor/LevelEditor/Private/LevelEditorActions.cpp
@@ -3077,7 +3077,7 @@ void FLevelEditorCommands::RegisterCommands()
 	static const FText FeatureLevelLabels[ERHIFeatureLevel::Num] = 
 	{
 		NSLOCTEXT("LevelEditorCommands", "FeatureLevelPreviewType_ES2", "Mobile / HTML5"),
-		NSLOCTEXT("LevelEditorCommands", "FeatureLevelPreviewType_ES3_1", "High-End Mobile / Metal"),
+		NSLOCTEXT("LevelEditorCommands", "FeatureLevelPreviewType_ES31", "High-End Mobile / Metal"),
 		NSLOCTEXT("LevelEditorCommands", "FeatureLevelPreviewType_SM4", "Shader Model 4"),
 		NSLOCTEXT("LevelEditorCommands", "FeatureLevelPreviewType_SM5", "Shader Model 5"),
 	};
diff --git a/Engine/Source/Editor/MaterialEditor/Private/MaterialEditor.cpp b/Engine/Source/Editor/MaterialEditor/Private/MaterialEditor.cpp
--- a/Engine/Source/Editor/MaterialEditor/Private/MaterialEditor.cpp
+++ b/Engine/Source/Editor/MaterialEditor/Private/MaterialEditor.cpp
@@ -926,7 +926,7 @@ void FMaterialEditor::DrawMaterialInfoStrings(
 			Canvas->DrawShadowedString(
 				5,
 				DrawPositionY,
-				*FString::Printf(TEXT("%s samplers: %u/%u"), FeatureLevel == ERHIFeatureLevel::ES2 ? TEXT("Mobile texture") : TEXT("Texture"), SamplersUsed, MaxSamplers),
+				*FString::Printf(TEXT("%s samplers: %u/%u"), FeatureLevel <= ERHIFeatureLevel::ES3_1 ? TEXT("Mobile texture") : TEXT("Texture"), SamplersUsed, MaxSamplers),
 				FontToUse,
 				SamplersUsed > MaxSamplers ? FLinearColor(1,0,0) : FLinearColor(1,1,0)
 				);
@@ -1597,7 +1597,7 @@ void FMaterialEditor::UpdateMaterialInfoList(bool bForceDisplay)
 				if (SamplersUsed >= 0)
 				{
 					int32 MaxSamplers = GetFeatureLevelMaxTextureSamplers(MaterialResource->GetFeatureLevel());
-					FString SamplersString = FString::Printf(TEXT("%s samplers: %u/%u"), FeatureLevel == ERHIFeatureLevel::ES2 ? TEXT("Mobile texture") : TEXT("Texture"), SamplersUsed, MaxSamplers);
+					FString SamplersString = FString::Printf(TEXT("%s samplers: %u/%u"), FeatureLevel <= ERHIFeatureLevel::ES3_1 ? TEXT("Mobile texture") : TEXT("Texture"), SamplersUsed, MaxSamplers);
 					TempMaterialInfoList.Add(MakeShareable(new FMaterialInfo(SamplersString, FLinearColor::Yellow)));
 					TSharedRef<FTokenizedMessage> Line = FTokenizedMessage::Create( EMessageSeverity::Info );
 					Line->AddToken( FTextToken::Create( FText::FromString( SamplersString ) ) );
diff --git a/Engine/Source/Runtime/Engine/Private/MeshParticleVertexFactory.cpp b/Engine/Source/Runtime/Engine/Private/MeshParticleVertexFactory.cpp
--- a/Engine/Source/Runtime/Engine/Private/MeshParticleVertexFactory.cpp
+++ b/Engine/Source/Runtime/Engine/Private/MeshParticleVertexFactory.cpp
@@ -42,7 +42,7 @@ public:
 
 	virtual void SetMesh(FRHICommandList& RHICmdList, FShader* Shader,const FVertexFactory* VertexFactory,const FSceneView& View,const FMeshBatchElement& BatchElement,uint32 DataFlags) const override
 	{
-		const bool bInstanced = View.GetFeatureLevel() >= ERHIFeatureLevel::SM4;
+		const bool bInstanced = View.GetFeatureLevel() >= ERHIFeatureLevel::ES3_1;
 		FMeshParticleVertexFactory* MeshParticleVF = (FMeshParticleVertexFactory*)VertexFactory;
 		FVertexShaderRHIParamRef VertexShaderRHI = Shader->GetVertexShader();
 		SetUniformBufferParameter(RHICmdList, VertexShaderRHI, Shader->GetUniformBufferParameter<FMeshParticleUniformParameters>(), MeshParticleVF->GetUniformBuffer() );
@@ -86,7 +86,7 @@ void FMeshParticleVertexFactory::InitRHI()
 {
 	FVertexDeclarationElementList Elements;
 
-	const bool bInstanced = GetFeatureLevel() >= ERHIFeatureLevel::SM4;
+	const bool bInstanced = GetFeatureLevel() >= ERHIFeatureLevel::ES3_1;
 
 	if (Data.bInitialized)
 	{
diff --git a/Engine/Source/Runtime/Engine/Public/MeshParticleVertexFactory.h b/Engine/Source/Runtime/Engine/Public/MeshParticleVertexFactory.h
--- a/Engine/Source/Runtime/Engine/Public/MeshParticleVertexFactory.h
+++ b/Engine/Source/Runtime/Engine/Public/MeshParticleVertexFactory.h
@@ -101,7 +101,7 @@ public:
 	{
 		FParticleVertexFactoryBase::ModifyCompilationEnvironment(Platform, Material, OutEnvironment);
 
-		const bool bInstanced = IsFeatureLevelSupported(Platform, ERHIFeatureLevel::SM4);
+		const bool bInstanced = IsFeatureLevelSupported(Platform, ERHIFeatureLevel::ES3_1);
 
 		// Set a define so we can tell in MaterialTemplate.usf when we are compiling a mesh particle vertex factory
 		OutEnvironment.SetDefine(TEXT("PARTICLE_MESH_FACTORY"),TEXT("1"));
diff --git a/Engine/Source/Runtime/RHI/Private/RHI.cpp b/Engine/Source/Runtime/RHI/Private/RHI.cpp
--- a/Engine/Source/Runtime/RHI/Private/RHI.cpp
+++ b/Engine/Source/Runtime/RHI/Private/RHI.cpp
@@ -270,7 +270,7 @@ RHI_API void GetFeatureLevelName(ERHIFeatureLevel::Type InFeatureLevel, FName& O
 
 static FName NAME_PCD3D_SM5(TEXT("PCD3D_SM5"));
 static FName NAME_PCD3D_SM4(TEXT("PCD3D_SM4"));
-static FName NAME_PCD3D_ES3_1(TEXT("PCD3D_ES3_1"));
+static FName NAME_PCD3D_ES3_1(TEXT("PCD3D_ES31"));
 static FName NAME_PCD3D_ES2(TEXT("PCD3D_ES2"));
 static FName NAME_GLSL_150(TEXT("GLSL_150"));
 static FName NAME_GLSL_150_MAC(TEXT("GLSL_150_MAC"));
@@ -279,7 +279,7 @@ static FName NAME_SF_XBOXONE(TEXT("SF_XBOXONE"));
 static FName NAME_GLSL_430(TEXT("GLSL_430"));
 static FName NAME_OPENGL_150_ES2(TEXT("GLSL_150_ES2"));
 static FName NAME_OPENGL_150_ES2_NOUB(TEXT("GLSL_150_ES2_NOUB"));
-static FName NAME_OPENGL_150_ES3_1(TEXT("GLSL_150_ES3_1"));
+static FName NAME_OPENGL_150_ES3_1(TEXT("GLSL_150_ES31"));
 static FName NAME_OPENGL_ES2(TEXT("GLSL_ES2"));
 static FName NAME_OPENGL_ES2_WEBGL(TEXT("GLSL_ES2_WEBGL"));
 static FName NAME_OPENGL_ES2_IOS(TEXT("GLSL_ES2_IOS"));
diff --git a/Engine/Source/Runtime/Renderer/Private/PostProcess/PostProcessVisualizeBuffer.cpp b/Engine/Source/Runtime/Renderer/Private/PostProcess/PostProcessVisualizeBuffer.cpp
--- a/Engine/Source/Runtime/Renderer/Private/PostProcess/PostProcessVisualizeBuffer.cpp
+++ b/Engine/Source/Runtime/Renderer/Private/PostProcess/PostProcessVisualizeBuffer.cpp
@@ -21,7 +21,7 @@ class FPostProcessVisualizeBufferPS : public FGlobalShader
 
 	static bool ShouldCache(EShaderPlatform Platform)
 	{
-		return IsFeatureLevelSupported(Platform, ERHIFeatureLevel::SM4);
+		return IsFeatureLevelSupported(Platform, ERHIFeatureLevel::ES3_1);
 	}
 
 	static void ModifyCompilationEnvironment(EShaderPlatform Platform, FShaderCompilerEnvironment& OutEnvironment)
