Fix so that material instance instances in placed meshes are not affected when a parent's parameter name is changed.

diff --git a/Engine/Source/Editor/MaterialEditor/Private/MaterialEditor.cpp b/Engine/Source/Editor/MaterialEditor/Private/MaterialEditor.cpp
--- a/Engine/Source/Editor/MaterialEditor/Private/MaterialEditor.cpp
+++ b/Engine/Source/Editor/MaterialEditor/Private/MaterialEditor.cpp
@@ -1493,8 +1493,18 @@ void FMaterialEditor::UpdateOriginalMaterial()
 				It->bIsViewRelevanceDirty = true;
 			}
 
+			// Update parameter names on any child material instances
+			for (TObjectIterator<UMaterialInstance> It; It; ++It)
+			{
+				if (It->Parent == OriginalMaterial)
+				{
+					It->UpdateParameterNames();
+				}
+			}
+
 			// Leaving this scope will update all dependent material instances.
 		}
+
 		RebuildMaterialInstanceEditors(NULL);
 	}
 
