Backed out CL 2502504 and CL 2472018. This needs a different fix.

diff --git a/Engine/Source/Editor/MaterialEditor/Private/MaterialEditor.cpp b/Engine/Source/Editor/MaterialEditor/Private/MaterialEditor.cpp
--- a/Engine/Source/Editor/MaterialEditor/Private/MaterialEditor.cpp
+++ b/Engine/Source/Editor/MaterialEditor/Private/MaterialEditor.cpp
@@ -331,7 +331,6 @@ void FMaterialEditor::InitMaterialEditor( const EToolkitMode::Type Mode, const T
 	FEditorSupportDelegates::MaterialUsageFlagsChanged.AddRaw(this, &FMaterialEditor::OnMaterialUsageFlagsChanged);
 	FEditorSupportDelegates::VectorParameterDefaultChanged.AddRaw(this, &FMaterialEditor::OnVectorParameterDefaultChanged);
 	FEditorSupportDelegates::ScalarParameterDefaultChanged.AddRaw(this, &FMaterialEditor::OnScalarParameterDefaultChanged);
-	FEditorDelegates::OnAssetPostImport.AddRaw(this, &FMaterialEditor::OnAssetPostImport);
 
 	FAssetRegistryModule& AssetRegistryModule = FModuleManager::GetModuleChecked<FAssetRegistryModule>(TEXT("AssetRegistry"));
 	
@@ -550,7 +549,6 @@ FMaterialEditor::~FMaterialEditor()
 	FEditorSupportDelegates::MaterialUsageFlagsChanged.RemoveAll(this);
 	FEditorSupportDelegates::VectorParameterDefaultChanged.RemoveAll(this);
 	FEditorSupportDelegates::ScalarParameterDefaultChanged.RemoveAll(this);
-	FEditorDelegates::OnAssetPostImport.RemoveAll(this);
 
 	// Null out the expression preview material so they can be GC'ed
 	ExpressionPreviewMaterial = NULL;
@@ -2452,64 +2450,6 @@ void FMaterialEditor::RenameAssetFromRegistry(const FAssetData& InAddedAssetData
 	}
 }
 
-
-void FMaterialEditor::OnAssetPostImport(UFactory* InFactory, UObject* InObject)
-{
-	UTexture* Texture = Cast<UTexture>(InObject);
-
-	if (InFactory->IsA(UTextureFactory::StaticClass()) && Texture != nullptr)
-	{
-		// When a texture which is referenced in the material is imported, update the preview material
-		for (UMaterialExpression* Expression : Material->Expressions)
-		{
-			if (UMaterialExpressionTextureBase* ExpressionTexture = Cast<UMaterialExpressionTextureBase>(Expression))
-			{
-				if (ExpressionTexture->Texture == Texture)
-				{
-					UpdateOriginalMaterial();
-					break;
-				}
-			}
-			else if (UMaterialExpressionMaterialFunctionCall* ExpressionFunctionCall = Cast<UMaterialExpressionMaterialFunctionCall>(Expression))
-			{
-				struct Local
-				{
-					static bool ReferencesImportedObject(UMaterialFunction* Function, UTexture* InTexture)
-					{
-						for (UMaterialExpression* MaterialExpression : Function->FunctionExpressions)
-						{
-							if (UMaterialExpressionTextureBase* MaterialExpressionTexture = Cast<UMaterialExpressionTextureBase>(MaterialExpression))
-							{
-								if (MaterialExpressionTexture->Texture == InTexture)
-								{
-									return true;
-								}
-							}
-							else if (UMaterialExpressionMaterialFunctionCall* MaterialExpressionFunctionCall = Cast<UMaterialExpressionMaterialFunctionCall>(MaterialExpression))
-							{
-								if (MaterialExpressionFunctionCall->MaterialFunction != nullptr &&
-									ReferencesImportedObject(MaterialExpressionFunctionCall->MaterialFunction, InTexture))
-								{
-									return true;
-								}
-							}
-						}
-
-						return false;
-					}
-				};
-
-				if (ExpressionFunctionCall->MaterialFunction != nullptr &&
-					Local::ReferencesImportedObject(ExpressionFunctionCall->MaterialFunction, Texture))
-				{
-					UpdateOriginalMaterial();
-					break;
-				}
-			}
-		}
-	}
-}
-
 void FMaterialEditor::OnMaterialUsageFlagsChanged(UMaterial* MaterialThatChanged, int32 FlagThatChanged)
 {
 	EMaterialUsage Flag = static_cast<EMaterialUsage>(FlagThatChanged);
