Guard against crash drawing a material editor who has been destroyed before its viewport has been.

diff --git a/Engine/Source/Editor/MaterialEditor/Private/SMaterialEditorViewport.cpp b/Engine/Source/Editor/MaterialEditor/Private/SMaterialEditorViewport.cpp
--- a/Engine/Source/Editor/MaterialEditor/Private/SMaterialEditorViewport.cpp
+++ b/Engine/Source/Editor/MaterialEditor/Private/SMaterialEditorViewport.cpp
@@ -108,17 +108,20 @@ bool FMaterialEditorViewportClient::ShouldOrbitCamera() const
 FLinearColor FMaterialEditorViewportClient::GetBackgroundColor() const
 {
 	FLinearColor BackgroundColor = FLinearColor::Black;
-	UMaterialInterface* MaterialInterface = MaterialEditorPtr.Pin()->GetMaterialInterface();
-	if (MaterialInterface)
+	if( MaterialEditorPtr.IsValid() )
 	{
-		const EBlendMode PreviewBlendMode = (EBlendMode)MaterialInterface->GetBlendMode();
-		if (PreviewBlendMode == BLEND_Modulate)
-		{
-			BackgroundColor = FLinearColor::White;
-		}
-		else if (PreviewBlendMode == BLEND_Translucent)
+		UMaterialInterface* MaterialInterface = MaterialEditorPtr.Pin()->GetMaterialInterface();
+		if(MaterialInterface)
 		{
-			BackgroundColor = FColor(64, 64, 64);
+			const EBlendMode PreviewBlendMode = (EBlendMode)MaterialInterface->GetBlendMode();
+			if(PreviewBlendMode == BLEND_Modulate)
+			{
+				BackgroundColor = FLinearColor::White;
+			}
+			else if(PreviewBlendMode == BLEND_Translucent)
+			{
+				BackgroundColor = FColor(64, 64, 64);
+			}
 		}
 	}
 	return BackgroundColor;
