PR #1168: Added hooks to Material Editor for plugin support (Contributed by [REDACTED_PERSON])

diff --git a/Engine/Source/Editor/MaterialEditor/Private/MaterialEditor.cpp b/Engine/Source/Editor/MaterialEditor/Private/MaterialEditor.cpp
--- a/Engine/Source/Editor/MaterialEditor/Private/MaterialEditor.cpp
+++ b/Engine/Source/Editor/MaterialEditor/Private/MaterialEditor.cpp
@@ -214,6 +214,8 @@ void FMaterialEditor::RegisterTabSpawners(const TSharedRef<class FTabManager>& T
 		.SetDisplayName( LOCTEXT("HLSLCodeTab", "HLSL Code") )
 		.SetGroup( WorkspaceMenuCategoryRef )
 		.SetIcon(FSlateIcon(FEditorStyle::GetStyleSetName(), "MaterialEditor.Tabs.HLSLCode"));
+
+	OnRegisterTabSpawners().Broadcast(TabManager);
 }
 
 
@@ -228,6 +230,8 @@ void FMaterialEditor::UnregisterTabSpawners(const TSharedRef<class FTabManager>&
 	TabManager->UnregisterTabSpawner( StatsTabId );
 	TabManager->UnregisterTabSpawner( FindTabId );
 	TabManager->UnregisterTabSpawner( HLSLCodeTabId );
+
+	OnUnregisterTabSpawners().Broadcast(TabManager);
 }
 
 void FMaterialEditor::InitEditorForMaterial(UMaterial* InMaterial)
@@ -415,6 +419,8 @@ void FMaterialEditor::InitMaterialEditor( const EToolkitMode::Type Mode, const T
 	ObjectsToEdit.Add(Material);
 	FAssetEditorToolkit::InitAssetEditor( Mode, InitToolkitHost, MaterialEditorAppIdentifier, StandaloneDefaultLayout, bCreateDefaultStandaloneMenu, bCreateDefaultToolbar, ObjectsToEdit, false );
 
+	AddMenuExtender(GetMenuExtensibilityManager()->GetAllExtenders(GetToolkitCommands(), GetEditingObjects()));
+
 	IMaterialEditorModule* MaterialEditorModule = &FModuleManager::LoadModuleChecked<IMaterialEditorModule>( "MaterialEditor" );
 	AddMenuExtender(MaterialEditorModule->GetMenuExtensibilityManager()->GetAllExtenders(GetToolkitCommands(), GetEditingObjects()));
 
@@ -530,11 +536,16 @@ FMaterialEditor::FMaterialEditor()
 	, bShowStats(true)
 	, bShowBuiltinStats(false)
 	, bShowMobileStats(false)
+	, MenuExtensibilityManager(new FExtensibilityManager)
+	, ToolBarExtensibilityManager(new FExtensibilityManager)
 {
 }
 
 FMaterialEditor::~FMaterialEditor()
 {
+	// Broadcast that this editor is going down to all listeners
+	OnMaterialEditorClosed().Broadcast();
+
 	for (int32 ParameterIndex = 0; ParameterIndex < OverriddenVectorParametersToRevert.Num(); ParameterIndex++)
 	{
 		SetVectorParameterDefaultOnDependentMaterials(OverriddenVectorParametersToRevert[ParameterIndex], FLinearColor::Black, false);
@@ -793,6 +804,8 @@ void FMaterialEditor::ExtendToolbar()
 	
 	AddToolbarExtender(ToolbarExtender);
 
+	AddToolbarExtender(GetToolBarExtensibilityManager()->GetAllExtenders(GetToolkitCommands(), GetEditingObjects()));
+
 	IMaterialEditorModule* MaterialEditorModule = &FModuleManager::LoadModuleChecked<IMaterialEditorModule>( "MaterialEditor" );
 	AddToolbarExtender(MaterialEditorModule->GetToolBarExtensibilityManager()->GetAllExtenders(GetToolkitCommands(), GetEditingObjects()));
 }
diff --git a/Engine/Source/Editor/MaterialEditor/Private/MaterialEditor.h b/Engine/Source/Editor/MaterialEditor/Private/MaterialEditor.h
--- a/Engine/Source/Editor/MaterialEditor/Private/MaterialEditor.h
+++ b/Engine/Source/Editor/MaterialEditor/Private/MaterialEditor.h
@@ -326,6 +326,10 @@ public:
 
 	void UpdateStatsMaterials();
 
+	/** Gets the extensibility managers for outside entities to extend material editor's menus and toolbars */
+	virtual TSharedPtr<FExtensibilityManager> GetMenuExtensibilityManager() { return MenuExtensibilityManager; }
+	virtual TSharedPtr<FExtensibilityManager> GetToolBarExtensibilityManager() { return ToolBarExtensibilityManager; }
+
 public:
 	/** Set to true when modifications have been made to the material */
 	bool bMaterialDirty;
@@ -693,6 +697,9 @@ private:
 	/** Command list for this editor */
 	TSharedPtr<FUICommandList> GraphEditorCommands;
 
+	TSharedPtr<FExtensibilityManager> MenuExtensibilityManager;
+	TSharedPtr<FExtensibilityManager> ToolBarExtensibilityManager;
+
 	/**	The tab ids for the material editor */
 	static const FName PreviewTabId;		
 	static const FName GraphCanvasTabId;	
diff --git a/Engine/Source/Editor/MaterialEditor/Private/MaterialEditorModule.cpp b/Engine/Source/Editor/MaterialEditor/Private/MaterialEditorModule.cpp
--- a/Engine/Source/Editor/MaterialEditor/Private/MaterialEditorModule.cpp
+++ b/Engine/Source/Editor/MaterialEditor/Private/MaterialEditorModule.cpp
@@ -50,6 +50,7 @@ public:
 	{
 		TSharedRef<FMaterialEditor> NewMaterialEditor(new FMaterialEditor());
 		NewMaterialEditor->InitEditorForMaterial(Material);
+		OnMaterialEditorOpened().Broadcast(NewMaterialEditor);
 		NewMaterialEditor->InitMaterialEditor(Mode, InitToolkitHost, Material);
 		return NewMaterialEditor;
 	}
@@ -58,6 +59,7 @@ public:
 	{
 		TSharedRef<FMaterialEditor> NewMaterialEditor(new FMaterialEditor());
 		NewMaterialEditor->InitEditorForMaterialFunction(MaterialFunction);
+		OnMaterialFunctionEditorOpened().Broadcast(NewMaterialEditor);
 		NewMaterialEditor->InitMaterialEditor(Mode, InitToolkitHost, MaterialFunction);
 		return NewMaterialEditor;
 	}
@@ -65,6 +67,7 @@ public:
 	virtual TSharedRef<IMaterialEditor> CreateMaterialInstanceEditor( const EToolkitMode::Type Mode, const TSharedPtr< IToolkitHost >& InitToolkitHost, UMaterialInstance* MaterialInstance ) override
 	{
 		TSharedRef<FMaterialInstanceEditor> NewMaterialInstanceEditor(new FMaterialInstanceEditor());
+		OnMaterialInstanceEditorOpened().Broadcast(NewMaterialInstanceEditor);
 		NewMaterialInstanceEditor->InitMaterialInstanceEditor(Mode, InitToolkitHost, MaterialInstance);
 		return NewMaterialInstanceEditor;
 	}
diff --git a/Engine/Source/Editor/MaterialEditor/Private/MaterialInstanceEditor.cpp b/Engine/Source/Editor/MaterialEditor/Private/MaterialInstanceEditor.cpp
--- a/Engine/Source/Editor/MaterialEditor/Private/MaterialInstanceEditor.cpp
+++ b/Engine/Source/Editor/MaterialEditor/Private/MaterialInstanceEditor.cpp
@@ -130,6 +130,8 @@ void FMaterialInstanceEditor::RegisterTabSpawners(const TSharedRef<class FTabMan
 		.SetDisplayName( LOCTEXT("ParentsTab", "Parents") )
 		.SetGroup( WorkspaceMenuCategoryRef )
 		.SetIcon( FSlateIcon( FEditorStyle::GetStyleSetName(), "Kismet.Tabs.Palette" ) );
+
+	OnRegisterTabSpawners().Broadcast(TabManager);
 }
 
 void FMaterialInstanceEditor::UnregisterTabSpawners(const TSharedRef<class FTabManager>& TabManager)
@@ -139,6 +141,8 @@ void FMaterialInstanceEditor::UnregisterTabSpawners(const TSharedRef<class FTabM
 	TabManager->UnregisterTabSpawner( PreviewTabId );		
 	TabManager->UnregisterTabSpawner( PropertiesTabId );	
 	TabManager->UnregisterTabSpawner( ParentsTabId );		
+
+	OnUnregisterTabSpawners().Broadcast(TabManager);
 }
 
 
@@ -234,9 +238,17 @@ void FMaterialInstanceEditor::InitMaterialInstanceEditor( const EToolkitMode::Ty
 	}
 	SetPreviewAssetByName(*InstanceConstant->PreviewMesh.ToString());
 }
+FMaterialInstanceEditor::FMaterialInstanceEditor()
+: MenuExtensibilityManager(new FExtensibilityManager)
+, ToolBarExtensibilityManager(new FExtensibilityManager)
+{
+}
 
 FMaterialInstanceEditor::~FMaterialInstanceEditor()
 {
+	// Broadcast that this editor is going down to all listeners
+	OnMaterialEditorClosed().Broadcast();
+
 	GEditor->UnregisterForUndo( this );
 
 	MaterialEditorInstance = NULL;
@@ -499,6 +511,8 @@ void FMaterialInstanceEditor::ExtendToolbar()
 
 	AddToolbarExtender(ToolbarExtender);
 
+	AddToolbarExtender(GetToolBarExtensibilityManager()->GetAllExtenders(GetToolkitCommands(), GetEditingObjects()));
+
 	IMaterialEditorModule* MaterialEditorModule = &FModuleManager::LoadModuleChecked<IMaterialEditorModule>( "MaterialEditor" );
 	AddToolbarExtender(MaterialEditorModule->GetToolBarExtensibilityManager()->GetAllExtenders(GetToolkitCommands(), GetEditingObjects()));
 }
diff --git a/Engine/Source/Editor/MaterialEditor/Private/MaterialInstanceEditor.h b/Engine/Source/Editor/MaterialEditor/Private/MaterialInstanceEditor.h
--- a/Engine/Source/Editor/MaterialEditor/Private/MaterialInstanceEditor.h
+++ b/Engine/Source/Editor/MaterialEditor/Private/MaterialInstanceEditor.h
@@ -25,6 +25,8 @@ public:
 	 */
 	void InitMaterialInstanceEditor( const EToolkitMode::Type Mode, const TSharedPtr< class IToolkitHost >& InitToolkitHost, UObject* ObjectToEdit );
 
+	FMaterialInstanceEditor();
+
 	virtual ~FMaterialInstanceEditor();
 
 	/** FGCObject interface */
@@ -73,6 +75,10 @@ public:
 	/** Returns true if hidden parameters should be shown */
 	void GetShowHiddenParameters(bool& bShowHiddenParameters);
 
+	/** Gets the extensibility managers for outside entities to extend material editor's menus and toolbars */
+	virtual TSharedPtr<FExtensibilityManager> GetMenuExtensibilityManager() { return MenuExtensibilityManager; }
+	virtual TSharedPtr<FExtensibilityManager> GetToolBarExtensibilityManager() { return ToolBarExtensibilityManager; }
+
 protected:
 	/** Saves editor settings. */
 	void SaveSettings();
@@ -170,6 +176,9 @@ private:
 	/** Whether to show mobile material stats. */
 	bool bShowMobileStats;
 
+	TSharedPtr<FExtensibilityManager> MenuExtensibilityManager;
+	TSharedPtr<FExtensibilityManager> ToolBarExtensibilityManager;
+
 	/**	The ids for the tabs spawned by this toolkit */
 	static const FName PreviewTabId;		
 	static const FName PropertiesTabId;	
diff --git a/Engine/Source/Editor/MaterialEditor/Public/IMaterialEditor.h b/Engine/Source/Editor/MaterialEditor/Public/IMaterialEditor.h
--- a/Engine/Source/Editor/MaterialEditor/Public/IMaterialEditor.h
+++ b/Engine/Source/Editor/MaterialEditor/Public/IMaterialEditor.h
@@ -8,7 +8,7 @@
 /**
  * Public interface to Material Editor
  */
-class IMaterialEditor : public FAssetEditorToolkit
+class IMaterialEditor : public FAssetEditorToolkit, public IHasMenuExtensibility, public IHasToolBarExtensibility
 {
 public:
 	/**
@@ -91,6 +91,23 @@ public:
 
 	/** Draws messages on the specified viewport and canvas. */
 	virtual void DrawMessages( FViewport* Viewport, FCanvas* Canvas ) = 0;
+
+	/** Delegate to be called when the tabs are being registered **/
+	DECLARE_EVENT_OneParam(IMaterialEditorModule, FRegisterTabSpawnersEvent, const TSharedRef<class FTabManager>&);
+	virtual FRegisterTabSpawnersEvent& OnRegisterTabSpawners() { return RegisterTabSpawnersEvent; };
+
+	/** Delegate to be called when the tabs are being unregistered **/
+	DECLARE_EVENT_OneParam(IMaterialEditorModule, FUnregisterTabSpawnersEvent, const TSharedRef<class FTabManager>&);
+	virtual FUnregisterTabSpawnersEvent& OnUnregisterTabSpawners() { return UnregisterTabSpawnersEvent; };
+
+	/** Delegate to be called when this IMaterialEditor is about to be destroyed **/
+	DECLARE_EVENT(IMaterialEditorModule, FMaterialEditorClosedEvent);
+	virtual FMaterialEditorClosedEvent& OnMaterialEditorClosed() { return MaterialEditorClosedEvent; };
+
+private:
+	FMaterialEditorClosedEvent MaterialEditorClosedEvent;
+	FRegisterTabSpawnersEvent RegisterTabSpawnersEvent;
+	FUnregisterTabSpawnersEvent UnregisterTabSpawnersEvent;
 };
 
 
diff --git a/Engine/Source/Editor/MaterialEditor/Public/MaterialEditorModule.h b/Engine/Source/Editor/MaterialEditor/Public/MaterialEditorModule.h
--- a/Engine/Source/Editor/MaterialEditor/Public/MaterialEditorModule.h
+++ b/Engine/Source/Editor/MaterialEditor/Public/MaterialEditorModule.h
@@ -33,7 +33,7 @@ public:
 	 *
 	 * @param	Material			The material to retrieve the parameters from.
 	 * @param	MaterialInstance	The material instance that contains all parameter overrides.
-	 * @param	VisisbleExpressions	The array that will contain the id's of the visible parameter expressions.
+	 * @param	VisibleExpressions	The array that will contain the id's of the visible parameter expressions.
 	 */
 	virtual void GetVisibleMaterialParameters(const class UMaterial* Material, class UMaterialInstance* MaterialInstance, TArray<struct FGuid>& VisibleExpressions) = 0;
 
@@ -43,10 +43,26 @@ public:
 	virtual TArray<FMaterialMenuExtender>& GetAllMaterialCanvasMenuExtenders() {return MaterialCanvasMenuExtenders;}
 	virtual TArray<FMaterialMenuExtender_MaterialInterface>& GetAllMaterialDragDropContextMenuExtenders() {return MaterialInheritanceMenuExtenders;}
 
+	/** Delegate to be called when a Material Editor is created, for toolbar, tab, and menu extension **/
+	DECLARE_EVENT_OneParam(IMaterialEditorModule, FMaterialEditorOpenedEvent, TWeakPtr<IMaterialEditor>);
+	virtual FMaterialEditorOpenedEvent& OnMaterialEditorOpened() { return MaterialEditorOpenedEvent; };
+
+	/** Delegate to be called when a Material Function Editor is created, for toolbar, tab, and menu extension **/
+	DECLARE_EVENT_OneParam(IMaterialEditorModule, FMaterialFunctionEditorOpenedEvent, TWeakPtr<IMaterialEditor>);
+	virtual FMaterialFunctionEditorOpenedEvent& OnMaterialFunctionEditorOpened() { return MaterialFunctionEditorOpenedEvent; };
+
+	/** Delegate to be called when a Material Instance Editor is created, for toolbar, tab, and menu extension **/
+	DECLARE_EVENT_OneParam(IMaterialEditorModule, FMaterialInstanceEditorOpenedEvent, TWeakPtr<IMaterialEditor>);
+	virtual FMaterialInstanceEditorOpenedEvent& OnMaterialInstanceEditorOpened() { return MaterialInstanceEditorOpenedEvent; };
+
 private:
 	/** All extender delegates for the material menus */
 	TArray<FMaterialMenuExtender> MaterialCanvasMenuExtenders;
 	TArray<FMaterialMenuExtender_MaterialInterface> MaterialInheritanceMenuExtenders;
+
+	FMaterialEditorOpenedEvent MaterialEditorOpenedEvent;
+	FMaterialFunctionEditorOpenedEvent MaterialFunctionEditorOpenedEvent;
+	FMaterialInstanceEditorOpenedEvent MaterialInstanceEditorOpenedEvent;
 };
 
 #endif // __MaterialEditorModule_h__
