Fixed crash converting texture sample to parameter when MipValue is set to "Level" or "Bias".

diff --git a/Engine/Source/Editor/MaterialEditor/Private/MaterialEditor.cpp b/Engine/Source/Editor/MaterialEditor/Private/MaterialEditor.cpp
--- a/Engine/Source/Editor/MaterialEditor/Private/MaterialEditor.cpp
+++ b/Engine/Source/Editor/MaterialEditor/Private/MaterialEditor.cpp
@@ -2120,6 +2120,7 @@ void FMaterialEditor::OnConvertObjects()
 							NewTextureExpr->TextureObject = TextureSampleExpression->TextureObject;
 							NewTextureExpr->MipValue = TextureSampleExpression->MipValue;
 							NewTextureExpr->MipValueMode = TextureSampleExpression->MipValueMode;
+							NewGraphNode->ReconstructNode();
 						}
 						else if (ComponentMaskExpression)
 						{
