Fix Material Promotion test never finishing if settings are blank (UEAUTO-484)
Current waterfall structure depends on step return value to determine whether to progress to the next test step.

diff --git a/Engine/Source/Editor/MaterialEditor/Private/Tests/MaterialEditorTests.cpp b/Engine/Source/Editor/MaterialEditor/Private/Tests/MaterialEditorTests.cpp
--- a/Engine/Source/Editor/MaterialEditor/Private/Tests/MaterialEditorTests.cpp
+++ b/Engine/Source/Editor/MaterialEditor/Private/Tests/MaterialEditorTests.cpp
@@ -69,7 +69,7 @@ namespace MaterialEditorPromotionTestUtils
 		else
 		{
 			UE_LOG(LogEditorMaterialEditorPromotionTests, Error, TEXT("Could not add a texture sample to %s"), *InMaterial->GetName());
-			return false;
+			return true;
 		}
 	}
 
@@ -148,7 +148,7 @@ namespace MaterialEditorPromotionTestHelper
 			if (!(DiffuseTexturePackagePath.Len() > 0))
 			{
 				UE_LOG(LogEditorMaterialEditorPromotionTests, Error, TEXT("Failed to load diffuse texture."));
-				return false;
+				return true;
 			}
 			DiffuseTexture = MaterialEditorPromotionTestUtils::GetAssetFromPackagePath<UTexture>(DiffuseTexturePackagePath);
 
@@ -188,7 +188,7 @@ namespace MaterialEditorPromotionTestHelper
 				if (!(NormalTexturePackagePath.Len() > 0))
 				{
 					UE_LOG(LogEditorMaterialEditorPromotionTests, Error, TEXT("Failed to load normal texture."));
-					return false;
+					return true;
 				}
 				NormalTexture = MaterialEditorPromotionTestUtils::GetAssetFromPackagePath<UTexture>(NormalTexturePackagePath);
 
