Material Editor promotion tests warn instead of fail if no settings are defined.
The tests fail instead of warn if settings are defined but assets cannot be loaded.
Clarified test failure logging.

diff --git a/Engine/Source/Editor/MaterialEditor/Private/Tests/MaterialEditorTests.cpp b/Engine/Source/Editor/MaterialEditor/Private/Tests/MaterialEditorTests.cpp
--- a/Engine/Source/Editor/MaterialEditor/Private/Tests/MaterialEditorTests.cpp
+++ b/Engine/Source/Editor/MaterialEditor/Private/Tests/MaterialEditorTests.cpp
@@ -147,7 +147,7 @@ namespace MaterialEditorPromotionTestHelper
 			FString DiffuseTexturePackagePath = MaterialEditorPromotionTestUtils::TestSettings().DefaultDiffuseTexture.FilePath;
 			if (!(DiffuseTexturePackagePath.Len() > 0))
 			{
-				UE_LOG(LogEditorMaterialEditorPromotionTests, Error, TEXT("Failed to load diffuse texture."));
+				UE_LOG(LogEditorMaterialEditorPromotionTests, Warning, TEXT("Skipping material creation test: No texture asset defined."));
 				return true;
 			}
 			DiffuseTexture = MaterialEditorPromotionTestUtils::GetAssetFromPackagePath<UTexture>(DiffuseTexturePackagePath);
@@ -158,19 +158,19 @@ namespace MaterialEditorPromotionTestHelper
 				CreatedMaterial = FEditorPromotionTestUtilities::CreateMaterialFromTexture(DiffuseTexture);
 				if (CreatedMaterial)
 				{
-					UE_LOG(LogEditorMaterialEditorPromotionTests, Display, TEXT("Created new material (%s) from diffuse (%s)"), *CreatedMaterial->GetName(), *DiffuseTexture->GetName());
+					UE_LOG(LogEditorMaterialEditorPromotionTests, Display, TEXT("Created new material (%s) from texture (%s)"), *CreatedMaterial->GetName(), *DiffuseTexture->GetName());
 					// @TODO: Break out TEST 2: Open a material
 					FAssetEditorManager::Get().OpenEditorForAsset(CreatedMaterial);
 					UE_LOG(LogEditorMaterialEditorPromotionTests, Display, TEXT("Opened the material editor"));
 				}
 				else
 				{
-					UE_LOG(LogEditorMaterialEditorPromotionTests, Error, TEXT("Failed to create material fom diffuse"));
+					UE_LOG(LogEditorMaterialEditorPromotionTests, Error, TEXT("Failed to create material fom texture"));
 				}
 			}
 			else
 			{
-				UE_LOG(LogEditorMaterialEditorPromotionTests, Warning, TEXT("Skipping Create Material: DiffuseTexture is missing"));
+				UE_LOG(LogEditorMaterialEditorPromotionTests, Error, TEXT("Failed to load texture asset"));
 			}
 
 			return true;
@@ -187,7 +187,7 @@ namespace MaterialEditorPromotionTestHelper
 				FString NormalTexturePackagePath = MaterialEditorPromotionTestUtils::TestSettings().DefaultNormalTexture.FilePath;
 				if (!(NormalTexturePackagePath.Len() > 0))
 				{
-					UE_LOG(LogEditorMaterialEditorPromotionTests, Error, TEXT("Failed to load normal texture."));
+					UE_LOG(LogEditorMaterialEditorPromotionTests, Warning, TEXT("Skipping material change test: No normal texture asset defined."));
 					return true;
 				}
 				NormalTexture = MaterialEditorPromotionTestUtils::GetAssetFromPackagePath<UTexture>(NormalTexturePackagePath);
@@ -211,9 +211,13 @@ namespace MaterialEditorPromotionTestHelper
 				}
 				else
 				{
-					UE_LOG(LogEditorMaterialEditorPromotionTests, Warning, TEXT("No normal texture to apply to new material"));
+					UE_LOG(LogEditorMaterialEditorPromotionTests, Error, TEXT("Failed to load normal texture asset"));
 				}
 			}
+			else
+			{
+				UE_LOG(LogEditorMaterialEditorPromotionTests, Warning, TEXT("Skipping material change test: Previous test step either did not run, or did not succeed."));
+			}
 
 			return true;
 		}
