Fixing missing loctext undef

diff --git a/Engine/Source/Editor/MaterialEditor/Private/Tests/MaterialEditorTests.cpp b/Engine/Source/Editor/MaterialEditor/Private/Tests/MaterialEditorTests.cpp
--- a/Engine/Source/Editor/MaterialEditor/Private/Tests/MaterialEditorTests.cpp
+++ b/Engine/Source/Editor/MaterialEditor/Private/Tests/MaterialEditorTests.cpp
@@ -245,4 +245,6 @@ bool FMaterialEditorPromotionTest::RunTest(const FString& Parameters)
 	TSharedPtr<MaterialEditorPromotionTestHelper::FMaterialEditorPromotionTest> MaterialEditorPromotionTest = MakeShareable(new MaterialEditorPromotionTestHelper::FMaterialEditorPromotionTest(&ExecutionInfo));
 	ADD_LATENT_AUTOMATION_COMMAND(FRunMaterialEditorPromotionTestCommand(MaterialEditorPromotionTest));
 	return true;
-}
\ No newline at end of file
+}
+
+#undef LOCTEXT_NAMESPACE
\ No newline at end of file
