Added error message if a material expression parameter name is too long.
Added the ability to display an error popup on the node if text verification fails.  The code existed for this but was non-functional

diff --git a/Engine/Source/Editor/GraphEditor/Private/SGraphNode.cpp b/Engine/Source/Editor/GraphEditor/Private/SGraphNode.cpp
--- a/Engine/Source/Editor/GraphEditor/Private/SGraphNode.cpp
+++ b/Engine/Source/Editor/GraphEditor/Private/SGraphNode.cpp
@@ -1262,10 +1262,13 @@ bool SGraphNode::OnVerifyNameTextChanged(const FText& InText, FText& OutErrorMes
 
 	if ((GetEditableNodeTitle() != InText.ToString()) && OnVerifyTextCommit.IsBound())
 	{
-		bValid = OnVerifyTextCommit.Execute(InText, GraphNode);
+		bValid = OnVerifyTextCommit.Execute(InText, GraphNode, OutErrorMessage);
 	}
 
-	OutErrorMessage = FText::FromString(TEXT("Error"));
+	if( OutErrorMessage.IsEmpty() )
+	{
+		OutErrorMessage = FText::FromString(TEXT("Error"));
+	}
 
 	//UpdateErrorInfo();
 	//ErrorReporting->SetError(ErrorMsg);
diff --git a/Engine/Source/Editor/Kismet/Private/BlueprintEditor.cpp b/Engine/Source/Editor/Kismet/Private/BlueprintEditor.cpp
--- a/Engine/Source/Editor/Kismet/Private/BlueprintEditor.cpp
+++ b/Engine/Source/Editor/Kismet/Private/BlueprintEditor.cpp
@@ -6909,7 +6909,7 @@ bool FBlueprintEditor::CanRenameNodes() const
 	return false;
 }
 
-bool FBlueprintEditor::OnNodeVerifyTitleCommit(const FText& NewText, UEdGraphNode* NodeBeingChanged)
+bool FBlueprintEditor::OnNodeVerifyTitleCommit(const FText& NewText, UEdGraphNode* NodeBeingChanged, FText& OutErrorMessage)
 {
 	bool bValid(false);
 
diff --git a/Engine/Source/Editor/MaterialEditor/Private/MaterialEditor.cpp b/Engine/Source/Editor/MaterialEditor/Private/MaterialEditor.cpp
--- a/Engine/Source/Editor/MaterialEditor/Private/MaterialEditor.cpp
+++ b/Engine/Source/Editor/MaterialEditor/Private/MaterialEditor.cpp
@@ -3815,6 +3815,7 @@ TSharedRef<SGraphEditor> FMaterialEditor::CreateGraphEditorWidget()
 	InEvents.OnSelectionChanged = SGraphEditor::FOnSelectionChanged::CreateSP(this, &FMaterialEditor::OnSelectedNodesChanged);
 	InEvents.OnNodeDoubleClicked = FSingleNodeEvent::CreateSP(this, &FMaterialEditor::OnNodeDoubleClicked);
 	InEvents.OnTextCommitted = FOnNodeTextCommitted::CreateSP(this, &FMaterialEditor::OnNodeTitleCommitted);
+	InEvents.OnVerifyTextCommit = FOnNodeVerifyTextCommit::CreateSP(this, &FMaterialEditor::OnVerifyNodeTextCommit);
 	InEvents.OnSpawnNodeByShortcut = SGraphEditor::FOnSpawnNodeByShortcut::CreateSP(this, &FMaterialEditor::OnSpawnGraphNodeByShortcut, CastChecked<UEdGraph>(Material->MaterialGraph));
 
 	// Create the title bar widget
@@ -4111,6 +4112,23 @@ void FMaterialEditor::OnNodeTitleCommitted(const FText& NewText, ETextCommit::Ty
 	}
 }
 
+bool FMaterialEditor::OnVerifyNodeTextCommit(const FText& NewText, UEdGraphNode* NodeBeingChanged, FText& OutErrorMessage)
+{
+	bool bValid = true;
+
+	UMaterialGraphNode* MaterialNode = Cast<UMaterialGraphNode>(NodeBeingChanged);
+	if( MaterialNode && MaterialNode->MaterialExpression && MaterialNode->MaterialExpression->IsA<UMaterialExpressionParameter>() )
+	{
+		if( NewText.ToString().Len() > NAME_SIZE )
+		{
+			OutErrorMessage = FText::Format( LOCTEXT("MaterialEditorExpressionError_NameTooLong", "Parameter names must be less than {0} characters"), FText::AsNumber(NAME_SIZE));
+			bValid = false;
+		}
+	}
+
+	return bValid;
+}
+
 FReply FMaterialEditor::OnSpawnGraphNodeByShortcut(FInputChord InChord, const FVector2D& InPosition, UEdGraph* InGraph)
 {
 	UEdGraph* Graph = InGraph;
diff --git a/Engine/Source/Editor/MaterialEditor/Private/MaterialEditor.h b/Engine/Source/Editor/MaterialEditor/Private/MaterialEditor.h
--- a/Engine/Source/Editor/MaterialEditor/Private/MaterialEditor.h
+++ b/Engine/Source/Editor/MaterialEditor/Private/MaterialEditor.h
@@ -391,6 +391,16 @@ protected:
 	 */
 	void OnNodeTitleCommitted(const FText& NewText, ETextCommit::Type CommitInfo, UEdGraphNode* NodeBeingChanged);
 
+	/**
+	 * Verifies that the node text entered is valid for the node
+	 *
+	 * @param	NewText			New node text
+	 * @param	NodeBeingChanged	The node being changed
+	 * @param	OutErrorMessage		Error message to display if text is invalid
+	 * @return	True if the text is valid, false otherwise
+	 */
+	bool OnVerifyNodeTextCommit(const FText& NewText, UEdGraphNode* NodeBeingChanged, FText& OutErrorMessage);
+
 	/**
 	 * Handles spawning a graph node in the current graph using the passed in chord
 	 *
diff --git a/Engine/Source/Editor/UnrealEd/Public/GraphEditor.h b/Engine/Source/Editor/UnrealEd/Public/GraphEditor.h
--- a/Engine/Source/Editor/UnrealEd/Public/GraphEditor.h
+++ b/Engine/Source/Editor/UnrealEd/Public/GraphEditor.h
@@ -10,7 +10,7 @@ class UEdGraph;
 struct FNotificationInfo;
 
 DECLARE_DELEGATE_ThreeParams( FOnNodeTextCommitted, const FText&, ETextCommit::Type, UEdGraphNode* );
-DECLARE_DELEGATE_RetVal_TwoParams( bool, FOnNodeVerifyTextCommit, const FText&, UEdGraphNode* );
+DECLARE_DELEGATE_RetVal_ThreeParams( bool, FOnNodeVerifyTextCommit, const FText&, UEdGraphNode*, FText& );
 
 typedef TSet<class UObject*> FGraphPanelSelectionSet;
 
