Added a UI material domain for all Slate and UMG materials
- Automatically upgraded all materials

When using the UI material domain the material editor and material instance editor is streamlined and only displays parts of the UI that are relevant:
- Changes to a preview material rendered with the UI shader
- Removes non-UI specific settings
- Renames some output pins and hides irrelevant ones
- Shows stats for UI shader

diff --git a/Engine/Source/Editor/DetailCustomizations/Private/SlateBrushCustomization.cpp b/Engine/Source/Editor/DetailCustomizations/Private/SlateBrushCustomization.cpp
--- a/Engine/Source/Editor/DetailCustomizations/Private/SlateBrushCustomization.cpp
+++ b/Engine/Source/Editor/DetailCustomizations/Private/SlateBrushCustomization.cpp
@@ -3,7 +3,8 @@
 #include "DetailCustomizationsPrivatePCH.h"
 #include "SlateBrushCustomization.h"
 #include "AssetData.h"
-
+#include "SHyperlink.h"
+#include "ScopedTransaction.h"
 /**
  * Slate Brush Preview widget
  */
@@ -1043,6 +1044,28 @@ private:
 
 float SSlateBrushStaticPreview::TargetHeight = 18.0f;
 
+class SBrushResourceError : public SBorder
+{
+public:
+	SLATE_BEGIN_ARGS( SBrushResourceError ) {}
+		SLATE_DEFAULT_SLOT( FArguments, Content )
+	SLATE_END_ARGS()
+
+	void Construct( const FArguments& InArgs )
+	{
+		SBorder::Construct( SBorder::FArguments()
+			.BorderBackgroundColor( FCoreStyle::Get().GetColor("ErrorReporting.BackgroundColor") )
+			.BorderImage( FCoreStyle::Get().GetBrush("ErrorReporting.Box") )
+			.HAlign(HAlign_Center)
+			.VAlign(VAlign_Center)
+			.Padding( FMargin(3,0) )
+			[
+				InArgs._Content.Widget
+			]
+		);
+	}
+};
+
 // SBrushResourceObjectBox
 ////////////////////////////////////////////////////////////////////////////////
 
@@ -1069,13 +1092,28 @@ class SBrushResourceObjectBox : public SCompoundWidget
 				SNew(SObjectPropertyEntryBox)
 				.PropertyHandle(InResourceObjectProperty)
 				.ThumbnailPool(StructCustomizationUtils->GetThumbnailPool())
-				.OnObjectChanged(this, &SBrushResourceObjectBox::OnAssetPicked)
 			]
 			+ SVerticalBox::Slot()
 			.AutoHeight()
 			.Padding( 0.0f, 3.0f )
 			[
-				SAssignNew(ResourceErrorText, SErrorText)
+				SAssignNew(ResourceError, SBrushResourceError )
+				[
+					SNew( SVerticalBox )
+					+ SVerticalBox::Slot()
+					.HAlign( HAlign_Left )
+					[
+						SNew( STextBlock )
+						.Text( NSLOCTEXT("FSlateBrushStructCustomization", "ResourceErrorText", "This material does not use the UI material domain" ) )
+					]
+					+ SVerticalBox::Slot()
+					.HAlign( HAlign_Left )
+					[
+						SNew( SHyperlink )
+						.Text( NSLOCTEXT("FSlateBrushStructCustomization", "ChangeMaterialDomain_ErrorMessage", "Change the Material Domain?" ) )
+						.OnNavigate( this, &SBrushResourceObjectBox::OnErrorLinkClicked )
+					]
+				]
 			]
 		];
 	}
@@ -1088,23 +1126,22 @@ class SBrushResourceObjectBox : public SCompoundWidget
 		{
 			UMaterialInterface* MaterialInterface = Cast<UMaterialInterface>( Resource );
 			UMaterial* BaseMaterial = MaterialInterface->GetBaseMaterial();
-			if( BaseMaterial && !BaseMaterial->bUsedWithUI )
+			if( BaseMaterial && !BaseMaterial->IsUIMaterial() )
 			{
-				ResourceErrorText->SetError( NSLOCTEXT("FSlateBrushStructCustomization", "ResourceErrorText", "This material is not supported in UI.  Please check \"Used with UI\" in the material editor" ) );
+				ResourceError->SetVisibility( EVisibility::Visible );
 			}
 			else
 			{
-				ResourceErrorText->SetError( FText::GetEmpty() );
+				ResourceError->SetVisibility( EVisibility::Collapsed );
 			}
 		}
-		else if( ResourceErrorText->HasError() )
+		else if( ResourceError->GetVisibility() != EVisibility::Collapsed )
 		{
-			ResourceErrorText->SetError( FText::GetEmpty() );
+			ResourceError->SetVisibility( EVisibility::Collapsed );
 		}
 	}
 
 private:
-
 	void OnBrushResourceChanged()
 	{
 		UObject* ResourceObject;
@@ -1130,22 +1167,26 @@ private:
 		}
 	}
 
-	/**
-	 * When the asset is picked if it's a material being used in the UI we 
-	 * automatically set the bUsedWithUI flag if it isn't already set.
-	 */
-	void OnAssetPicked(const FAssetData& InAssetData) const
+	void OnErrorLinkClicked()
 	{
 		UObject* Resource = nullptr;
 
-		if ( ResourceObjectProperty->GetValue(Resource) == FPropertyAccess::Success && Resource && Resource->IsA<UMaterialInterface>() )
+		if( ResourceObjectProperty->GetValue(Resource) == FPropertyAccess::Success && Resource && Resource->IsA<UMaterialInterface>() )
 		{
-			UMaterialInterface* MaterialInterface = Cast<UMaterialInterface>(Resource);
+			UMaterialInterface* MaterialInterface = Cast<UMaterialInterface>( Resource );
 			UMaterial* BaseMaterial = MaterialInterface->GetBaseMaterial();
-			if ( BaseMaterial && !BaseMaterial->bUsedWithUI )
+			if ( BaseMaterial && !BaseMaterial->IsUIMaterial() )
 			{
-				bool bNeedsRecompile = true;
-				BaseMaterial->SetMaterialUsage(bNeedsRecompile, MATUSAGE_UI);
+				UProperty* MaterialDomainProp = FindField<UProperty>(UMaterial::StaticClass(), GET_MEMBER_NAME_CHECKED(UMaterial,MaterialDomain) );
+
+				FScopedTransaction Transaction( FText::Format( NSLOCTEXT("FSlateBrushStructCustomization", "ChangeMaterialDomainTransaction", "Changed {0} to use the UI material domain"), FText::FromString( BaseMaterial->GetName() ) ) );
+
+				BaseMaterial->PreEditChange( MaterialDomainProp );
+
+				BaseMaterial->MaterialDomain = MD_UI;
+
+				FPropertyChangedEvent ChangeEvent( MaterialDomainProp );
+				BaseMaterial->PostEditChangeProperty( ChangeEvent );
 			}
 		}
 	}
@@ -1153,7 +1194,7 @@ private:
 private:
 	TSharedPtr<IPropertyHandle> ResourceObjectProperty;
 	TSharedPtr<IPropertyHandle> ImageSizeProperty;
-	TSharedPtr<SErrorText> ResourceErrorText;
+	TSharedPtr<SBrushResourceError> ResourceError;
 };
 
 // FSlateBrushStructCustomization
diff --git a/Engine/Source/Editor/MaterialEditor/Private/MaterialEditor.cpp b/Engine/Source/Editor/MaterialEditor/Private/MaterialEditor.cpp
--- a/Engine/Source/Editor/MaterialEditor/Private/MaterialEditor.cpp
+++ b/Engine/Source/Editor/MaterialEditor/Private/MaterialEditor.cpp
@@ -347,6 +347,8 @@ void FMaterialEditor::InitMaterialEditor( const EToolkitMode::Type Mode, const T
 	Material->bIsPreviewMaterial = true;
 	FMaterialEditorUtilities::InitExpressions(Material);
 
+	UpdatePreviewViewportsVisibility();
+
 	BindCommands();
 
 	const TSharedRef<FTabManager::FLayout> StandaloneDefaultLayout = FTabManager::NewLayout("Standalone_MaterialEditor_Layout_v5")
@@ -603,12 +605,28 @@ void FMaterialEditor::GetAllMaterialExpressionGroups(TArray<FString>* OutGroups)
 	}
 }
 
+void FMaterialEditor::UpdatePreviewViewportsVisibility()
+{
+	if( Material->IsUIMaterial() )
+	{
+		PreviewViewport->SetVisibility(EVisibility::Collapsed);
+		PreviewUIViewport->SetVisibility(EVisibility::Visible);
+	}
+	else
+	{
+		PreviewViewport->SetVisibility(EVisibility::Visible);
+		PreviewUIViewport->SetVisibility(EVisibility::Collapsed);
+	}
+}
+
 BEGIN_SLATE_FUNCTION_BUILD_OPTIMIZATION
 void FMaterialEditor::CreateInternalWidgets()
 {
-	Viewport = SNew(SMaterialEditorViewport)
+	PreviewViewport = SNew(SMaterialEditor3DPreviewViewport)
 		.MaterialEditor(SharedThis(this));
 
+	PreviewUIViewport = SNew(SMaterialEditorUIPreviewViewport, Material);
+
 	FPropertyEditorModule& PropertyEditorModule = FModuleManager::GetModuleChecked<FPropertyEditorModule>( "PropertyEditor" );
 
 	GraphEditor = CreateGraphEditorWidget();
@@ -643,6 +661,8 @@ void FMaterialEditor::CreateInternalWidgets()
 		LayoutCollectionParameterDetails
 		);
 
+	PropertyEditorModule.RegisterCustomClassLayout( UMaterial::StaticClass()->GetFName(), FOnGetDetailCustomizationInstance::CreateStatic(&FMaterialDetailCustomization::MakeInstance ) );
+
 	Palette = SNew(SMaterialPalette, SharedThis(this));
 
 	FMessageLogModule& MessageLogModule = FModuleManager::LoadModuleChecked<FMessageLogModule>("MessageLog");
@@ -870,10 +890,6 @@ bool FMaterialEditor::OnRequestClose()
 {
 	DestroyColorPicker();
 
-	// @todo DB: Store off the viewport camera position/orientation to the material.
-	//AnimTree->PreviewCamPos = PreviewVC->ViewLocation;
-	//AnimTree->PreviewCamRot = PreviewVC->ViewRotation;
-
 	if (bMaterialDirty)
 	{
 		// find out the user wants to do with this dirty material
@@ -1048,35 +1064,35 @@ void FMaterialEditor::RecenterEditor()
 
 bool FMaterialEditor::SetPreviewAsset(UObject* InAsset)
 {
-	if (Viewport.IsValid())
+	if (PreviewViewport.IsValid())
 	{
-		return Viewport->SetPreviewAsset(InAsset);
+		return PreviewViewport->SetPreviewAsset(InAsset);
 	}
 	return false;
 }
 
 bool FMaterialEditor::SetPreviewAssetByName(const TCHAR* InAssetName)
 {
-	if (Viewport.IsValid())
+	if (PreviewViewport.IsValid())
 	{
-		return Viewport->SetPreviewAssetByName(InAssetName);
+		return PreviewViewport->SetPreviewAssetByName(InAssetName);
 	}
 	return false;
 }
 
 void FMaterialEditor::SetPreviewMaterial(UMaterialInterface* InMaterialInterface)
 {
-	if (Viewport.IsValid())
+	if (PreviewViewport.IsValid())
 	{
-		Viewport->SetPreviewMaterial(InMaterialInterface);
+		PreviewViewport->SetPreviewMaterial(InMaterialInterface);
 	}
 }
 
 void FMaterialEditor::RefreshPreviewViewport()
 {
-	if (Viewport.IsValid())
+	if (PreviewViewport.IsValid())
 	{
-		Viewport->RefreshViewport();
+		PreviewViewport->RefreshViewport();
 	}
 }
 
@@ -1092,14 +1108,25 @@ void FMaterialEditor::LoadEditorSettings()
 	if (EditorOptions->bAlwaysRefreshAllPreviews) {OnAlwaysRefreshAllPreviews();}
 	if (EditorOptions->bRealtimeExpressionViewport) {ToggleRealTimeExpressions();}
 
-	if ( Viewport.IsValid() )
+	if ( PreviewViewport.IsValid() )
 	{
-		if (EditorOptions->bShowGrid) {Viewport->TogglePreviewGrid();}
-		if (EditorOptions->bShowBackground) {Viewport->TogglePreviewBackground();}
-		if (EditorOptions->bRealtimeMaterialViewport) {Viewport->OnToggleRealtime();}
+		if (EditorOptions->bShowGrid)
+		{
+			PreviewViewport->TogglePreviewGrid();
+		}
+
+		if (EditorOptions->bShowBackground)
+		{
+			PreviewViewport->TogglePreviewBackground();
+		}
+
+		if (EditorOptions->bRealtimeMaterialViewport)
+		{
+			PreviewViewport->OnToggleRealtime();
+		}
 
 		// Load the preview scene
-		Viewport->PreviewScene.LoadSettings(TEXT("MaterialEditor"));
+		PreviewViewport->PreviewScene.LoadSettings(TEXT("MaterialEditor"));
 	}
 
 	if (EditorOptions->bShowMobileStats)
@@ -1112,21 +1139,21 @@ void FMaterialEditor::LoadEditorSettings()
 	int32 PrimType;
 	if(GConfig->GetInt(TEXT("MaterialEditor"), TEXT("PrimType"), PrimType, GEditorPerProjectIni))
 	{
-		Viewport->OnSetPreviewPrimitive((EThumbnailPrimType)PrimType);
+		PreviewViewport->OnSetPreviewPrimitive((EThumbnailPrimType)PrimType);
 	}
 }
 
 void FMaterialEditor::SaveEditorSettings()
 {
 	// Save the preview scene
-	check( Viewport.IsValid() );
-	Viewport->PreviewScene.SaveSettings(TEXT("MaterialEditor"));
+	check( PreviewViewport.IsValid() );
+	PreviewViewport->PreviewScene.SaveSettings(TEXT("MaterialEditor"));
 
 	if ( EditorOptions )
 	{
-		EditorOptions->bShowGrid					= Viewport->IsTogglePreviewGridChecked();
-		EditorOptions->bShowBackground				= Viewport->IsTogglePreviewBackgroundChecked();
-		EditorOptions->bRealtimeMaterialViewport	= Viewport->IsRealtime();
+		EditorOptions->bShowGrid					= PreviewViewport->IsTogglePreviewGridChecked();
+		EditorOptions->bShowBackground				= PreviewViewport->IsTogglePreviewBackgroundChecked();
+		EditorOptions->bRealtimeMaterialViewport	= PreviewViewport->IsRealtime();
 		EditorOptions->bShowMobileStats				= bShowMobileStats;
 		EditorOptions->bHideUnusedConnectors		= !IsOnShowConnectorsChecked();
 		EditorOptions->bAlwaysRefreshAllPreviews	= IsOnAlwaysRefreshAllPreviews();
@@ -1135,7 +1162,7 @@ void FMaterialEditor::SaveEditorSettings()
 		EditorOptions->SaveConfig();
 	}
 
-	GConfig->SetInt(TEXT("MaterialEditor"), TEXT("PrimType"), Viewport->PreviewPrimType, GEditorPerProjectIni);
+	GConfig->SetInt(TEXT("MaterialEditor"), TEXT("PrimType"), PreviewViewport->PreviewPrimType, GEditorPerProjectIni);
 }
 
 FText FMaterialEditor::GetCodeViewText() const
@@ -1753,9 +1780,9 @@ void FMaterialEditor::BindCommands()
 
 	ToolkitCommands->MapAction(
 		FEditorViewportCommands::Get().ToggleRealTime,
-		FExecuteAction::CreateSP( Viewport.ToSharedRef(), &SMaterialEditorViewport::OnToggleRealtime ),
+		FExecuteAction::CreateSP( PreviewViewport.ToSharedRef(), &SMaterialEditor3DPreviewViewport::OnToggleRealtime ),
 		FCanExecuteAction(),
-		FIsActionChecked::CreateSP( Viewport.ToSharedRef(), &SMaterialEditorViewport::IsRealtime ) );
+		FIsActionChecked::CreateSP( PreviewViewport.ToSharedRef(), &SMaterialEditor3DPreviewViewport::IsRealtime ) );
 
 	ToolkitCommands->MapAction(
 		FGenericCommands::Get().Undo,
@@ -2641,10 +2668,18 @@ TSharedRef<SDockTab> FMaterialEditor::SpawnTab_Preview(const FSpawnTabArgs& Args
 		SNew(SDockTab)
 		.Label(LOCTEXT("ViewportTabTitle", "Viewport"))
 		[
-			Viewport.ToSharedRef()
+			SNew( SOverlay )
+			+ SOverlay::Slot()
+			[
+				PreviewViewport.ToSharedRef()
+			]
+			+ SOverlay::Slot()
+			[
+				PreviewUIViewport.ToSharedRef()
+			]
 		];
 
-	Viewport->OnAddedToTab( SpawnedTab );
+	PreviewViewport->OnAddedToTab( SpawnedTab );
 
 	return SpawnedTab;
 }
@@ -3494,6 +3529,14 @@ void FMaterialEditor::NotifyPostChange( const FPropertyChangedEvent& PropertyCha
 				SetPreviewAsset(GUnrealEd->GetThumbnailManager()->EditorSphere);
 			}
 		}
+		else if ( NameOfPropertyThatChanged == GET_MEMBER_NAME_CHECKED(UMaterial, MaterialDomain) )
+		{
+			Material->MaterialGraph->RebuildGraph();
+			TArray<TWeakObjectPtr<UObject>> SelectedObjects = MaterialDetailsView->GetSelectedObjects();
+			MaterialDetailsView->SetObjects( SelectedObjects, true );
+
+			UpdatePreviewViewportsVisibility();
+		}
 
 		FGraphPanelSelectionSet SelectedNodes = GraphEditor->GetSelectedNodes();
 
diff --git a/Engine/Source/Editor/MaterialEditor/Private/MaterialEditor.h b/Engine/Source/Editor/MaterialEditor/Private/MaterialEditor.h
--- a/Engine/Source/Editor/MaterialEditor/Private/MaterialEditor.h
+++ b/Engine/Source/Editor/MaterialEditor/Private/MaterialEditor.h
@@ -461,6 +461,9 @@ private:
 	/** Collects all groups for all material expressions */
 	void GetAllMaterialExpressionGroups(TArray<FString>* OutGroups);
 
+	/** Updates the 3D and UI preview viewport visibility based on material domain */
+	void UpdatePreviewViewportsVisibility();
+
 public:
 
 private:
@@ -649,10 +652,13 @@ private:
 	TSharedPtr<class SGraphEditor> GraphEditor;
 
 	/** Preview Viewport widget */
-	TSharedPtr<class SMaterialEditorViewport> Viewport;
+	TSharedPtr<class SMaterialEditor3DPreviewViewport> PreviewViewport;
+
+	/** Preview viewport widget used for UI materials */
+	TSharedPtr<class SMaterialEditorUIPreviewViewport> PreviewUIViewport;
 
 	/** Widget to hold utility components for the HLSL Code View */
-	 TSharedPtr<SWidget> CodeViewUtility;
+	TSharedPtr<SWidget> CodeViewUtility;
 
 	/** Widget for the HLSL Code View */
 	TSharedPtr<class SScrollBox> CodeView;
diff --git a/Engine/Source/Editor/MaterialEditor/Private/MaterialEditorDetailCustomization.cpp b/Engine/Source/Editor/MaterialEditor/Private/MaterialEditorDetailCustomization.cpp
--- a/Engine/Source/Editor/MaterialEditor/Private/MaterialEditorDetailCustomization.cpp
+++ b/Engine/Source/Editor/MaterialEditor/Private/MaterialEditorDetailCustomization.cpp
@@ -299,5 +299,62 @@ void FMaterialExpressionCollectionParameterDetails::OnSelectionChanged( TSharedP
 	}
 }
 
+TSharedRef<class IDetailCustomization> FMaterialDetailCustomization::MakeInstance()
+{
+	return MakeShareable( new FMaterialDetailCustomization );
+}
+
+void FMaterialDetailCustomization::CustomizeDetails( IDetailLayoutBuilder& DetailLayout )
+{
+	TArray<TWeakObjectPtr<UObject> > Objects;
+	DetailLayout.GetObjectsBeingCustomized( Objects );
+
+	bool bUIMaterial = true;
+	for( TWeakObjectPtr<UObject>& Object : Objects )
+	{
+		UMaterial* Material = Cast<UMaterial>( Object.Get() );
+		if( Material )
+		{
+			bUIMaterial &= Material->IsUIMaterial();
+		}
+		else
+		{
+			// this shouldn't happen but just in case, let all properties through
+			bUIMaterial = false;
+		}
+	}
+
+	if( bUIMaterial )
+	{
+		DetailLayout.HideCategory( TEXT("TranslucencySelfShadowing") );
+		DetailLayout.HideCategory( TEXT("Translucency") );
+		DetailLayout.HideCategory( TEXT("Tessellation") );
+		DetailLayout.HideCategory( TEXT("Mobile") );
+		DetailLayout.HideCategory( TEXT("PostProcessMaterial") );
+		DetailLayout.HideCategory( TEXT("Lightmass") );
+		DetailLayout.HideCategory( TEXT("Thumbnail") );
+		DetailLayout.HideCategory( TEXT("MaterialInterface") );
+		DetailLayout.HideCategory( TEXT("PhysicalMaterial") );
+		DetailLayout.HideCategory( TEXT("Usage") );
+
+		IDetailCategoryBuilder& MaterialCategory = DetailLayout.EditCategory( TEXT("Material") );
+
+		TArray<TSharedRef<IPropertyHandle>> AllProperties;
+		MaterialCategory.GetDefaultProperties( AllProperties );
+
+		for( TSharedRef<IPropertyHandle>& PropertyHandle : AllProperties )
+		{
+			UProperty* Property = PropertyHandle->GetProperty();
+			FName PropertyName = Property->GetFName();
+
+			if(		PropertyName != GET_MEMBER_NAME_CHECKED(UMaterial, MaterialDomain) 
+				&&	PropertyName != GET_MEMBER_NAME_CHECKED(UMaterial, BlendMode) 
+				&&	PropertyName != GET_MEMBER_NAME_CHECKED(UMaterial, OpacityMaskClipValue) )
+			{
+				DetailLayout.HideProperty( PropertyHandle );
+			}
+		}
+	}
+}
 
 #undef LOCTEXT_NAMESPACE
diff --git a/Engine/Source/Editor/MaterialEditor/Private/MaterialEditorDetailCustomization.h b/Engine/Source/Editor/MaterialEditor/Private/MaterialEditorDetailCustomization.h
--- a/Engine/Source/Editor/MaterialEditor/Private/MaterialEditorDetailCustomization.h
+++ b/Engine/Source/Editor/MaterialEditor/Private/MaterialEditorDetailCustomization.h
@@ -88,3 +88,14 @@ private:
 	/** A list of all parameter names to choose from */
 	TArray<TSharedPtr<FString>> ParametersSource;
 };
+
+class FMaterialDetailCustomization : public IDetailCustomization
+{
+public:
+
+	/** Makes a new instance of this detail layout class for a specific detail view requesting it */
+	static TSharedRef<class IDetailCustomization> MakeInstance();
+
+	/** IDetailCustomization interface */
+	virtual void CustomizeDetails( IDetailLayoutBuilder& DetailLayout ) override;
+};
diff --git a/Engine/Source/Editor/MaterialEditor/Private/MaterialInstanceEditor.cpp b/Engine/Source/Editor/MaterialEditor/Private/MaterialInstanceEditor.cpp
--- a/Engine/Source/Editor/MaterialEditor/Private/MaterialInstanceEditor.cpp
+++ b/Engine/Source/Editor/MaterialEditor/Private/MaterialInstanceEditor.cpp
@@ -174,6 +174,8 @@ void FMaterialInstanceEditor::InitMaterialInstanceEditor( const EToolkitMode::Ty
 
 	BindCommands();
 
+	UpdatePreviewViewportsVisibility();
+
 	TSharedRef<FTabManager::FLayout> StandaloneDefaultLayout = FTabManager::NewLayout( "Standalone_MaterialInstanceEditor_Layout_v2" )
 	->AddArea
 	(
@@ -285,9 +287,9 @@ void FMaterialInstanceEditor::BindCommands()
 
 	ToolkitCommands->MapAction(
 		FEditorViewportCommands::Get().ToggleRealTime,
-		FExecuteAction::CreateSP( PreviewVC.ToSharedRef(), &SMaterialEditorViewport::OnToggleRealtime ),
+		FExecuteAction::CreateSP( PreviewVC.ToSharedRef(), &SMaterialEditor3DPreviewViewport::OnToggleRealtime ),
 		FCanExecuteAction(),
-		FIsActionChecked::CreateSP( PreviewVC.ToSharedRef(), &SMaterialEditorViewport::IsRealtime ) );
+		FIsActionChecked::CreateSP( PreviewVC.ToSharedRef(), &SMaterialEditor3DPreviewViewport::IsRealtime ) );
 
 	ToolkitCommands->MapAction(
 		Commands.ToggleMobileStats,
@@ -458,9 +460,11 @@ UMaterialInterface* FMaterialInstanceEditor::GetSelectedParent() const
 
 void FMaterialInstanceEditor::CreateInternalWidgets()
 {
-	PreviewVC = SNew(SMaterialEditorViewport)
+	PreviewVC = SNew(SMaterialEditor3DPreviewViewport)
 		.MaterialEditor(SharedThis(this));
 
+	PreviewUIViewport = SNew(SMaterialEditorUIPreviewViewport, GetMaterialInterface());
+
 	FPropertyEditorModule& PropertyEditorModule = FModuleManager::GetModuleChecked<FPropertyEditorModule>( "PropertyEditor" );
 	const FDetailsViewArgs DetailsViewArgs( false, false, true, FDetailsViewArgs::HideNameArea, true, this );
 	MaterialInstanceDetails = PropertyEditorModule.CreateDetailView( DetailsViewArgs );
@@ -482,6 +486,23 @@ void FMaterialInstanceEditor::CreateInternalWidgets()
 			+SHeaderRow::Column(FName(TEXT("Name")))
 			.DefaultLabel(NSLOCTEXT("MaterialInstanceEditor", "Name", "Name"))
 		);
+
+}
+
+
+void FMaterialInstanceEditor::UpdatePreviewViewportsVisibility()
+{
+	UMaterial* PreviewMaterial = GetMaterialInterface()->GetBaseMaterial();
+	if( PreviewMaterial->IsUIMaterial() )
+	{
+		PreviewVC->SetVisibility(EVisibility::Collapsed);
+		PreviewUIViewport->SetVisibility(EVisibility::Visible);
+	}
+	else
+	{
+		PreviewVC->SetVisibility(EVisibility::Visible);
+		PreviewUIViewport->SetVisibility(EVisibility::Collapsed);
+	}
 }
 
 void FMaterialInstanceEditor::ExtendToolbar()
@@ -524,7 +545,15 @@ TSharedRef<SDockTab> FMaterialInstanceEditor::SpawnTab_Preview( const FSpawnTabA
 	TSharedRef<SDockTab> SpawnedTab = SNew(SDockTab)
 		.Label(LOCTEXT("ViewportTabTitle", "Viewport"))
 		[
-			PreviewVC.ToSharedRef()
+			SNew( SOverlay )
+			+ SOverlay::Slot()
+			[
+				PreviewVC.ToSharedRef()
+			]
+			+ SOverlay::Slot()
+			[
+				PreviewUIViewport.ToSharedRef()
+			]
 		];
 
 	PreviewVC->OnAddedToTab( SpawnedTab );
@@ -715,7 +744,7 @@ void FMaterialInstanceEditor::DrawMessages( FViewport* Viewport, FCanvas* Canvas
 	{
 		const FMaterialResource* MaterialResource = MaterialEditorInstance->SourceInstance->GetMaterialResource(GMaxRHIFeatureLevel);
 		UMaterial* BaseMaterial = MaterialEditorInstance->SourceInstance->GetMaterial();
-		int32 DrawPositionY = 5;
+		int32 DrawPositionY = 50;
 		if ( BaseMaterial && MaterialResource )
 		{
 			FMaterialEditor::DrawMaterialInfoStrings( Canvas, BaseMaterial, MaterialResource, MaterialResource->GetCompileErrors(), DrawPositionY, true );
diff --git a/Engine/Source/Editor/MaterialEditor/Private/MaterialInstanceEditor.h b/Engine/Source/Editor/MaterialEditor/Private/MaterialInstanceEditor.h
--- a/Engine/Source/Editor/MaterialEditor/Private/MaterialInstanceEditor.h
+++ b/Engine/Source/Editor/MaterialEditor/Private/MaterialInstanceEditor.h
@@ -124,6 +124,9 @@ private:
 	/** Creates all internal widgets for the tabs to point at */
 	void CreateInternalWidgets();
 
+	/** Updates the 3D and UI preview viewport visibility based on material domain */
+	void UpdatePreviewViewportsVisibility();
+
 	/** Builds the toolbar widget for the material editor */
 	void ExtendToolbar();
 
@@ -156,7 +159,10 @@ private:
 	TMap< FName, TWeakPtr<class SDockTab> > SpawnedToolPanels;
 
 	/** Preview Viewport widget */
-	TSharedPtr<class SMaterialEditorViewport> PreviewVC;
+	TSharedPtr<class SMaterialEditor3DPreviewViewport> PreviewVC;
+
+	/** Preview viewport widget used for UI materials */
+	TSharedPtr<class SMaterialEditorUIPreviewViewport> PreviewUIViewport;
 
 	/** Property View */
 	TSharedPtr<class IDetailsView> MaterialInstanceDetails;
diff --git a/Engine/Source/Editor/MaterialEditor/Private/SMaterialEditorViewport.cpp b/Engine/Source/Editor/MaterialEditor/Private/SMaterialEditorViewport.cpp
--- a/Engine/Source/Editor/MaterialEditor/Private/SMaterialEditorViewport.cpp
+++ b/Engine/Source/Editor/MaterialEditor/Private/SMaterialEditorViewport.cpp
@@ -16,12 +16,17 @@
 #include "Engine/Selection.h"
 #include "Engine/TextureCube.h"
 #include "ComponentAssetBroker.h"
+#include "SlateMaterialBrush.h"
+#include "SNumericEntryBox.h"
+
+
+#define LOCTEXT_NAMESPACE "MaterialEditor"
 
 /** Viewport Client for the preview viewport */
 class FMaterialEditorViewportClient : public FEditorViewportClient
 {
 public:
-	FMaterialEditorViewportClient(TWeakPtr<IMaterialEditor> InMaterialEditor, FPreviewScene& InPreviewScene, const TSharedRef<SMaterialEditorViewport>& InMaterialEditorViewport);
+	FMaterialEditorViewportClient(TWeakPtr<IMaterialEditor> InMaterialEditor, FPreviewScene& InPreviewScene, const TSharedRef<SMaterialEditor3DPreviewViewport>& InMaterialEditorViewport);
 
 	// FEditorViewportClient interface
 	virtual FLinearColor GetBackgroundColor() const override;
@@ -46,7 +51,7 @@ private:
 	TWeakPtr<IMaterialEditor> MaterialEditorPtr;
 };
 
-FMaterialEditorViewportClient::FMaterialEditorViewportClient(TWeakPtr<IMaterialEditor> InMaterialEditor, FPreviewScene& InPreviewScene, const TSharedRef<SMaterialEditorViewport>& InMaterialEditorViewport)
+FMaterialEditorViewportClient::FMaterialEditorViewportClient(TWeakPtr<IMaterialEditor> InMaterialEditor, FPreviewScene& InPreviewScene, const TSharedRef<SMaterialEditor3DPreviewViewport>& InMaterialEditorViewport)
 	: FEditorViewportClient(nullptr, &InPreviewScene, StaticCastSharedRef<SEditorViewport>(InMaterialEditorViewport))
 	, MaterialEditorPtr(InMaterialEditor)
 {
@@ -185,7 +190,7 @@ void FMaterialEditorViewportClient::FocusViewportOnBounds(const FBoxSphereBounds
 	Invalidate();
 }
 
-void SMaterialEditorViewport::Construct(const FArguments& InArgs)
+void SMaterialEditor3DPreviewViewport::Construct(const FArguments& InArgs)
 {
 	MaterialEditorPtr = InArgs._MaterialEditor;
 
@@ -201,7 +206,8 @@ void SMaterialEditorViewport::Construct(const FArguments& InArgs)
 	PreviewMaterial = nullptr;
 	PreviewMeshComponent = nullptr;
 
-	if (UMaterialInterface* Material = MaterialEditorPtr.Pin()->GetMaterialInterface())
+	UMaterialInterface* Material = MaterialEditorPtr.Pin()->GetMaterialInterface();
+	if (Material)
 	{
 		SetPreviewMaterial(Material);
 	}
@@ -209,7 +215,7 @@ void SMaterialEditorViewport::Construct(const FArguments& InArgs)
 	SetPreviewAsset( GUnrealEd->GetThumbnailManager()->EditorSphere );
 }
 
-SMaterialEditorViewport::~SMaterialEditorViewport()
+SMaterialEditor3DPreviewViewport::~SMaterialEditor3DPreviewViewport()
 {
 	if (PreviewMeshComponent != nullptr)
 	{
@@ -222,13 +228,13 @@ SMaterialEditorViewport::~SMaterialEditorViewport()
 	}
 }
 
-void SMaterialEditorViewport::AddReferencedObjects( FReferenceCollector& Collector )
+void SMaterialEditor3DPreviewViewport::AddReferencedObjects( FReferenceCollector& Collector )
 {
 	Collector.AddReferencedObject( PreviewMeshComponent );
 	Collector.AddReferencedObject( PreviewMaterial );
 }
 
-void SMaterialEditorViewport::RefreshViewport()
+void SMaterialEditor3DPreviewViewport::RefreshViewport()
 {
 	// reregister the preview components, so if the preview material changed it will be propagated to the render thread
 	if (PreviewMeshComponent != nullptr)
@@ -238,7 +244,7 @@ void SMaterialEditorViewport::RefreshViewport()
 	SceneViewport->InvalidateDisplay();
 }
 
-bool SMaterialEditorViewport::SetPreviewAsset(UObject* InAsset)
+bool SMaterialEditor3DPreviewViewport::SetPreviewAsset(UObject* InAsset)
 {
 	if (!MaterialEditorPtr.Pin()->ApproveSetPreviewAsset(InAsset))
 	{
@@ -314,7 +320,7 @@ bool SMaterialEditorViewport::SetPreviewAsset(UObject* InAsset)
 	return (PreviewMeshComponent != nullptr);
 }
 
-bool SMaterialEditorViewport::SetPreviewAssetByName(const TCHAR* InAssetName)
+bool SMaterialEditor3DPreviewViewport::SetPreviewAssetByName(const TCHAR* InAssetName)
 {
 	bool bSuccess = false;
 	if ((InAssetName != nullptr) && (*InAssetName != 0))
@@ -327,7 +333,7 @@ bool SMaterialEditorViewport::SetPreviewAssetByName(const TCHAR* InAssetName)
 	return bSuccess;
 }
 
-void SMaterialEditorViewport::SetPreviewMaterial(UMaterialInterface* InMaterialInterface)
+void SMaterialEditor3DPreviewViewport::SetPreviewMaterial(UMaterialInterface* InMaterialInterface)
 {
 	PreviewMaterial = InMaterialInterface;
 
@@ -338,17 +344,17 @@ void SMaterialEditorViewport::SetPreviewMaterial(UMaterialInterface* InMaterialI
 	}
 }
 
-void SMaterialEditorViewport::OnAddedToTab( const TSharedRef<SDockTab>& OwnerTab )
+void SMaterialEditor3DPreviewViewport::OnAddedToTab( const TSharedRef<SDockTab>& OwnerTab )
 {
 	ParentTab = OwnerTab;
 }
 
-bool SMaterialEditorViewport::IsVisible() const
+bool SMaterialEditor3DPreviewViewport::IsVisible() const
 {
 	return ViewportWidget.IsValid() && (!ParentTab.IsValid() || ParentTab.Pin()->IsForeground()) && SEditorViewport::IsVisible() ;
 }
 
-void SMaterialEditorViewport::BindCommands()
+void SMaterialEditor3DPreviewViewport::BindCommands()
 {
 	SEditorViewport::BindCommands();
 
@@ -360,48 +366,48 @@ void SMaterialEditorViewport::BindCommands()
 	// Add the commands to the toolkit command list so that the toolbar buttons can find them
 	CommandList->MapAction(
 		Commands.SetCylinderPreview,
-		FExecuteAction::CreateSP( this, &SMaterialEditorViewport::OnSetPreviewPrimitive, TPT_Cylinder ),
+		FExecuteAction::CreateSP( this, &SMaterialEditor3DPreviewViewport::OnSetPreviewPrimitive, TPT_Cylinder ),
 		FCanExecuteAction(),
-		FIsActionChecked::CreateSP( this, &SMaterialEditorViewport::IsPreviewPrimitiveChecked, TPT_Cylinder ) );
+		FIsActionChecked::CreateSP( this, &SMaterialEditor3DPreviewViewport::IsPreviewPrimitiveChecked, TPT_Cylinder ) );
 
 	CommandList->MapAction(
 		Commands.SetSpherePreview,
-		FExecuteAction::CreateSP( this, &SMaterialEditorViewport::OnSetPreviewPrimitive, TPT_Sphere ),
+		FExecuteAction::CreateSP( this, &SMaterialEditor3DPreviewViewport::OnSetPreviewPrimitive, TPT_Sphere ),
 		FCanExecuteAction(),
-		FIsActionChecked::CreateSP( this, &SMaterialEditorViewport::IsPreviewPrimitiveChecked, TPT_Sphere ) );
+		FIsActionChecked::CreateSP( this, &SMaterialEditor3DPreviewViewport::IsPreviewPrimitiveChecked, TPT_Sphere ) );
 
 	CommandList->MapAction(
 		Commands.SetPlanePreview,
-		FExecuteAction::CreateSP( this, &SMaterialEditorViewport::OnSetPreviewPrimitive, TPT_Plane ),
+		FExecuteAction::CreateSP( this, &SMaterialEditor3DPreviewViewport::OnSetPreviewPrimitive, TPT_Plane ),
 		FCanExecuteAction(),
-		FIsActionChecked::CreateSP( this, &SMaterialEditorViewport::IsPreviewPrimitiveChecked, TPT_Plane ) );
+		FIsActionChecked::CreateSP( this, &SMaterialEditor3DPreviewViewport::IsPreviewPrimitiveChecked, TPT_Plane ) );
 
 	CommandList->MapAction(
 		Commands.SetCubePreview,
-		FExecuteAction::CreateSP( this, &SMaterialEditorViewport::OnSetPreviewPrimitive, TPT_Cube ),
+		FExecuteAction::CreateSP( this, &SMaterialEditor3DPreviewViewport::OnSetPreviewPrimitive, TPT_Cube ),
 		FCanExecuteAction(),
-		FIsActionChecked::CreateSP( this, &SMaterialEditorViewport::IsPreviewPrimitiveChecked, TPT_Cube ) );
+		FIsActionChecked::CreateSP( this, &SMaterialEditor3DPreviewViewport::IsPreviewPrimitiveChecked, TPT_Cube ) );
 
 	CommandList->MapAction(
 		Commands.SetPreviewMeshFromSelection,
-		FExecuteAction::CreateSP( this, &SMaterialEditorViewport::OnSetPreviewMeshFromSelection ),
+		FExecuteAction::CreateSP( this, &SMaterialEditor3DPreviewViewport::OnSetPreviewMeshFromSelection ),
 		FCanExecuteAction(),
-		FIsActionChecked::CreateSP( this, &SMaterialEditorViewport::IsPreviewMeshFromSelectionChecked ) );
+		FIsActionChecked::CreateSP( this, &SMaterialEditor3DPreviewViewport::IsPreviewMeshFromSelectionChecked ) );
 
 	CommandList->MapAction(
 		Commands.TogglePreviewGrid,
-		FExecuteAction::CreateSP( this, &SMaterialEditorViewport::TogglePreviewGrid ),
+		FExecuteAction::CreateSP( this, &SMaterialEditor3DPreviewViewport::TogglePreviewGrid ),
 		FCanExecuteAction(),
-		FIsActionChecked::CreateSP( this, &SMaterialEditorViewport::IsTogglePreviewGridChecked ) );
+		FIsActionChecked::CreateSP( this, &SMaterialEditor3DPreviewViewport::IsTogglePreviewGridChecked ) );
 
 	CommandList->MapAction(
 		Commands.TogglePreviewBackground,
-		FExecuteAction::CreateSP( this, &SMaterialEditorViewport::TogglePreviewBackground ),
+		FExecuteAction::CreateSP( this, &SMaterialEditor3DPreviewViewport::TogglePreviewBackground ),
 		FCanExecuteAction(),
-		FIsActionChecked::CreateSP( this, &SMaterialEditorViewport::IsTogglePreviewBackgroundChecked ) );
+		FIsActionChecked::CreateSP( this, &SMaterialEditor3DPreviewViewport::IsTogglePreviewBackgroundChecked ) );
 }
 
-void SMaterialEditorViewport::OnFocusViewportToSelection()
+void SMaterialEditor3DPreviewViewport::OnFocusViewportToSelection()
 {
 	if( PreviewMeshComponent != nullptr )
 	{
@@ -409,7 +415,7 @@ void SMaterialEditorViewport::OnFocusViewportToSelection()
 	}
 }
 
-void SMaterialEditorViewport::OnSetPreviewPrimitive(EThumbnailPrimType PrimType)
+void SMaterialEditor3DPreviewViewport::OnSetPreviewPrimitive(EThumbnailPrimType PrimType)
 {
 	if (SceneViewport.IsValid())
 	{
@@ -430,12 +436,12 @@ void SMaterialEditorViewport::OnSetPreviewPrimitive(EThumbnailPrimType PrimType)
 	}
 }
 
-bool SMaterialEditorViewport::IsPreviewPrimitiveChecked(EThumbnailPrimType PrimType) const
+bool SMaterialEditor3DPreviewViewport::IsPreviewPrimitiveChecked(EThumbnailPrimType PrimType) const
 {
 	return PreviewPrimType == PrimType;
 }
 
-void SMaterialEditorViewport::OnSetPreviewMeshFromSelection()
+void SMaterialEditor3DPreviewViewport::OnSetPreviewMeshFromSelection()
 {
 	bool bFoundPreviewMesh = false;
 	FEditorDelegates::LoadSelectedAssetsIfNeeded.Broadcast();
@@ -488,51 +494,51 @@ void SMaterialEditorViewport::OnSetPreviewMeshFromSelection()
 	}
 }
 
-bool SMaterialEditorViewport::IsPreviewMeshFromSelectionChecked() const
+bool SMaterialEditor3DPreviewViewport::IsPreviewMeshFromSelectionChecked() const
 {
 	return (PreviewPrimType == TPT_None && PreviewMeshComponent != nullptr);
 }
 
-void SMaterialEditorViewport::TogglePreviewGrid()
+void SMaterialEditor3DPreviewViewport::TogglePreviewGrid()
 {
 	bShowGrid = !bShowGrid;
 	EditorViewportClient->SetShowGrid(bShowGrid);
 	RefreshViewport();
 }
 
-bool SMaterialEditorViewport::IsTogglePreviewGridChecked() const
+bool SMaterialEditor3DPreviewViewport::IsTogglePreviewGridChecked() const
 {
 	return bShowGrid;
 }
 
-void SMaterialEditorViewport::TogglePreviewBackground()
+void SMaterialEditor3DPreviewViewport::TogglePreviewBackground()
 {
 	bShowBackground = !bShowBackground;
 	// @todo DB: Set the background mesh for the preview viewport.
 	RefreshViewport();
 }
 
-bool SMaterialEditorViewport::IsTogglePreviewBackgroundChecked() const
+bool SMaterialEditor3DPreviewViewport::IsTogglePreviewBackgroundChecked() const
 {
 	return bShowBackground;
 }
 
-TSharedRef<class SEditorViewport> SMaterialEditorViewport::GetViewportWidget()
+TSharedRef<class SEditorViewport> SMaterialEditor3DPreviewViewport::GetViewportWidget()
 {
 	return SharedThis(this);
 }
 
-TSharedPtr<FExtender> SMaterialEditorViewport::GetExtenders() const
+TSharedPtr<FExtender> SMaterialEditor3DPreviewViewport::GetExtenders() const
 {
 	TSharedPtr<FExtender> Result(MakeShareable(new FExtender));
 	return Result;
 }
 
-void SMaterialEditorViewport::OnFloatingButtonClicked()
+void SMaterialEditor3DPreviewViewport::OnFloatingButtonClicked()
 {
 }
 
-TSharedRef<FEditorViewportClient> SMaterialEditorViewport::MakeEditorViewportClient() 
+TSharedRef<FEditorViewportClient> SMaterialEditor3DPreviewViewport::MakeEditorViewportClient() 
 {
 	EditorViewportClient = MakeShareable( new FMaterialEditorViewportClient(MaterialEditorPtr, PreviewScene, SharedThis(this)) );
 	
@@ -542,12 +548,12 @@ TSharedRef<FEditorViewportClient> SMaterialEditorViewport::MakeEditorViewportCli
 	EditorViewportClient->bSetListenerPosition = false;
 
 	EditorViewportClient->SetRealtime( false );
-	EditorViewportClient->VisibilityDelegate.BindSP( this, &SMaterialEditorViewport::IsVisible );
+	EditorViewportClient->VisibilityDelegate.BindSP( this, &SMaterialEditor3DPreviewViewport::IsVisible );
 
 	return EditorViewportClient.ToSharedRef();
 }
 
-void SMaterialEditorViewport::PopulateViewportOverlays(TSharedRef<class SOverlay> Overlay)
+void SMaterialEditor3DPreviewViewport::PopulateViewportOverlays(TSharedRef<class SOverlay> Overlay)
 {
 	Overlay->AddSlot()
 		.VAlign(VAlign_Top)
@@ -562,7 +568,7 @@ void SMaterialEditorViewport::PopulateViewportOverlays(TSharedRef<class SOverlay
 		];
 }
 
-EVisibility SMaterialEditorViewport::OnGetViewportContentVisibility() const
+EVisibility SMaterialEditor3DPreviewViewport::OnGetViewportContentVisibility() const
 {
 	EVisibility BaseVisibility = SEditorViewport::OnGetViewportContentVisibility();
 	if (BaseVisibility != EVisibility::Visible)
@@ -571,3 +577,239 @@ EVisibility SMaterialEditorViewport::OnGetViewportContentVisibility() const
 	}
 	return IsVisible() ? EVisibility::Visible : EVisibility::Collapsed;
 }
+
+class SMaterialEditorUIPreviewZoomer : public SPanel
+{
+public:
+
+	class FMaterialPreviewPanelSlot : public TSupportsOneChildMixin<FMaterialPreviewPanelSlot>
+	{
+	public:
+		FMaterialPreviewPanelSlot()
+			: TSupportsOneChildMixin<FMaterialPreviewPanelSlot>()
+		{
+		}
+	};
+
+	SLATE_BEGIN_ARGS(SMaterialEditorUIPreviewViewport){}
+	SLATE_END_ARGS()
+
+	SMaterialEditorUIPreviewZoomer()
+	{
+	}
+
+	void Construct( const FArguments& InArgs, UMaterialInterface* PreviewMaterial );
+
+	virtual void OnArrangeChildren( const FGeometry& AllottedGeometry, FArrangedChildren& ArrangedChildren ) const override;
+	virtual FVector2D ComputeDesiredSize(float) const override;
+	virtual FChildren* GetChildren() override;
+	virtual FReply OnMouseWheel( const FGeometry& MyGeometry, const FPointerEvent& MouseEvent ) override;
+	virtual int32 OnPaint(const FPaintArgs& Args, const FGeometry& AllottedGeometry, const FSlateRect& MyClippingRect, FSlateWindowElementList& OutDrawElements, int32 LayerId, const FWidgetStyle& InWidgetStyle, bool bParentEnabled) const override;
+
+	bool ZoomBy( const float Amount );
+	float GetZoomLevel() const;
+
+	void SetPreviewSize( const FVector2D PreviewSize );
+private:
+	mutable FVector2D CachedSize;
+	float ZoomLevel;
+
+	FMaterialPreviewPanelSlot ChildSlot;
+
+	TWeakObjectPtr<UMaterial> PreviewMaterial;
+	TSharedPtr<FSlateMaterialBrush> PreviewBrush;
+};
+
+
+void SMaterialEditorUIPreviewZoomer::Construct( const FArguments& InArgs, UMaterialInterface* PreviewMaterial )
+{
+	PreviewBrush = MakeShareable( new FSlateMaterialBrush( *PreviewMaterial, FVector2D(250,250) ) );
+
+	ChildSlot
+	[
+		SNew( SImage )
+		.Image( PreviewBrush.Get() )
+	];
+
+	ZoomLevel = 1.0f;
+}
+
+void SMaterialEditorUIPreviewZoomer::OnArrangeChildren( const FGeometry& AllottedGeometry, FArrangedChildren& ArrangedChildren ) const 
+{
+	CachedSize = AllottedGeometry.GetLocalSize();
+
+	const TSharedRef<SWidget>& ChildWidget = ChildSlot.GetWidget();
+	if( ChildWidget->GetVisibility() != EVisibility::Collapsed )
+	{
+		const FVector2D& WidgetDesiredSize = ChildWidget->GetDesiredSize();
+
+		ArrangedChildren.AddWidget(AllottedGeometry.MakeChild(ChildWidget, FVector2D::ZeroVector, WidgetDesiredSize * ZoomLevel));
+	}
+}
+
+FVector2D SMaterialEditorUIPreviewZoomer::ComputeDesiredSize(float) const
+{
+	FVector2D ThisDesiredSize = FVector2D::ZeroVector;
+
+	const TSharedRef<SWidget>& ChildWidget = ChildSlot.GetWidget();
+	if( ChildWidget->GetVisibility() != EVisibility::Collapsed )
+	{
+		ThisDesiredSize = ChildWidget->GetDesiredSize() * ZoomLevel;
+	}
+
+	return ThisDesiredSize;
+}
+
+FChildren* SMaterialEditorUIPreviewZoomer::GetChildren()
+{
+	return &ChildSlot;
+}
+
+FReply SMaterialEditorUIPreviewZoomer::OnMouseWheel( const FGeometry& MyGeometry, const FPointerEvent& MouseEvent )
+{
+	 ZoomBy(MouseEvent.GetWheelDelta());
+
+	 return FReply::Handled();
+}
+
+int32 SMaterialEditorUIPreviewZoomer::OnPaint(const FPaintArgs& Args, const FGeometry& AllottedGeometry, const FSlateRect& MyClippingRect, FSlateWindowElementList& OutDrawElements, int32 LayerId, const FWidgetStyle& InWidgetStyle, bool bParentEnabled) const
+{
+	LayerId = SPanel::OnPaint(Args, AllottedGeometry, MyClippingRect, OutDrawElements, LayerId, InWidgetStyle, bParentEnabled);
+	return LayerId;
+}
+
+bool SMaterialEditorUIPreviewZoomer::ZoomBy( const float Amount )
+{
+	static const float MinZoomLevel = 0.2f;
+	static const float MaxZoomLevel = 4.0f;
+
+	const float PrevZoomLevel = ZoomLevel;
+	ZoomLevel = FMath::Clamp(ZoomLevel + (Amount * 0.05f), MinZoomLevel, MaxZoomLevel);
+	return ZoomLevel != PrevZoomLevel;
+}
+
+float SMaterialEditorUIPreviewZoomer::GetZoomLevel() const
+{
+	return ZoomLevel;
+}
+
+
+void SMaterialEditorUIPreviewZoomer::SetPreviewSize( const FVector2D PreviewSize )
+{
+	PreviewBrush->ImageSize = PreviewSize;
+}
+
+
+void SMaterialEditorUIPreviewViewport::Construct( const FArguments& InArgs, UMaterialInterface* PreviewMaterial )
+{
+
+	ChildSlot
+	[
+		SNew( SVerticalBox )
+		+ SVerticalBox::Slot()
+		.AutoHeight()
+		[
+			SNew( SBorder )
+			.HAlign(HAlign_Right)
+			.VAlign(VAlign_Top)
+			.BorderImage( FEditorStyle::GetBrush("ToolPanel.GroupBorder") )
+			[
+				SNew( SHorizontalBox )
+				+ SHorizontalBox::Slot()
+				.VAlign(VAlign_Center)
+				.Padding( 3.f )
+				.AutoWidth()
+				[
+					SNew( STextBlock )
+					.Text( LOCTEXT("PreviewSize", "Preview Size" ) )
+				]
+				+ SHorizontalBox::Slot()
+				.VAlign(VAlign_Center)
+				.Padding( 3.f )
+				.MaxWidth( 75 )
+				[
+					SNew( SNumericEntryBox<int32> )
+					.AllowSpin( true )
+					.MinValue(1)
+					.MaxSliderValue( 4096 )
+					.OnValueChanged( this, &SMaterialEditorUIPreviewViewport::OnPreviewXChanged )
+					.OnValueCommitted( this, &SMaterialEditorUIPreviewViewport::OnPreviewXCommitted )
+					.Value( this, &SMaterialEditorUIPreviewViewport::OnGetPreviewXValue )
+					.MinDesiredValueWidth( 75 )
+					.Label()
+					[	
+						SNew( SBox )
+						.VAlign( VAlign_Center )
+						[
+							SNew( STextBlock )
+							.Text( LOCTEXT("PreviewSize_X", "X") )
+						]
+					]
+				]
+				+ SHorizontalBox::Slot()
+				.VAlign(VAlign_Center)
+				.Padding( 3.f )
+				.MaxWidth( 75 )
+				[
+					SNew( SNumericEntryBox<int32> )
+					.AllowSpin( true )
+					.MinValue(1)
+					.MaxSliderValue( 4096 )
+					.MinDesiredValueWidth( 75 )
+					.OnValueChanged( this, &SMaterialEditorUIPreviewViewport::OnPreviewYChanged )
+					.OnValueCommitted( this, &SMaterialEditorUIPreviewViewport::OnPreviewYCommitted )
+					.Value( this, &SMaterialEditorUIPreviewViewport::OnGetPreviewYValue )
+					.Label()
+					[	
+						SNew( SBox )
+						.VAlign( VAlign_Center )
+						[
+							SNew( STextBlock )
+							.Text( LOCTEXT("PreviewSize_Y", "Y") )
+						]
+					]
+				]
+			]
+		]
+		+ SVerticalBox::Slot()
+		[
+			SNew( SBorder )
+			.HAlign( HAlign_Center )
+			.VAlign( VAlign_Center )
+			.BorderImage( FEditorStyle::GetBrush("BlackBrush") )
+			[
+				SAssignNew( PreviewZoomer, SMaterialEditorUIPreviewZoomer, PreviewMaterial )
+			]
+		]
+	];
+
+	PreviewSize = FIntPoint(250,250);
+	PreviewZoomer->SetPreviewSize( FVector2D(PreviewSize) );
+}
+
+void SMaterialEditorUIPreviewViewport::OnPreviewXChanged( int32 NewValue )
+{
+	PreviewSize.X = NewValue;
+	PreviewZoomer->SetPreviewSize( FVector2D( PreviewSize ) );
+
+}
+
+void SMaterialEditorUIPreviewViewport::OnPreviewXCommitted( int32 NewValue, ETextCommit::Type )
+{
+	OnPreviewXChanged( NewValue );
+}
+
+
+void SMaterialEditorUIPreviewViewport::OnPreviewYChanged( int32 NewValue )
+{
+	PreviewSize.Y = NewValue;
+	PreviewZoomer->SetPreviewSize( FVector2D( PreviewSize ) );
+
+}
+
+void SMaterialEditorUIPreviewViewport::OnPreviewYCommitted( int32 NewValue, ETextCommit::Type )
+{
+	OnPreviewYChanged( NewValue );
+}
+
+#undef LOC_NAMESPACE
\ No newline at end of file
diff --git a/Engine/Source/Editor/MaterialEditor/Private/SMaterialEditorViewport.h b/Engine/Source/Editor/MaterialEditor/Private/SMaterialEditorViewport.h
--- a/Engine/Source/Editor/MaterialEditor/Private/SMaterialEditorViewport.h
+++ b/Engine/Source/Editor/MaterialEditor/Private/SMaterialEditorViewport.h
@@ -6,18 +6,20 @@
 #include "SEditorViewport.h"
 #include "Editor/UnrealEd/Public/SCommonEditorViewportToolbarBase.h"
 
+struct FSlateMaterialBrush;
+
 /**
  * Material Editor Preview viewport widget
  */
-class SMaterialEditorViewport : public SEditorViewport, public FGCObject, public ICommonEditorViewportToolbarInfoProvider
+class SMaterialEditor3DPreviewViewport : public SEditorViewport, public FGCObject, public ICommonEditorViewportToolbarInfoProvider
 {
 public:
-	SLATE_BEGIN_ARGS( SMaterialEditorViewport ){}
+	SLATE_BEGIN_ARGS( SMaterialEditor3DPreviewViewport ){}
 		SLATE_ARGUMENT(TWeakPtr<IMaterialEditor>, MaterialEditor)
 	SLATE_END_ARGS()
 
 	void Construct(const FArguments& InArgs);
-	~SMaterialEditorViewport();
+	~SMaterialEditor3DPreviewViewport();
 	
 	virtual void AddReferencedObjects( FReferenceCollector& Collector ) override;
 
@@ -94,7 +96,30 @@ private:
 
 	/** Pointer back to the material editor tool that owns us */
 	TWeakPtr<IMaterialEditor> MaterialEditorPtr;
-	
+
 	/** Level viewport client */
 	TSharedPtr<class FMaterialEditorViewportClient> EditorViewportClient;
 };
+
+/**
+ * A preview viewport used for 2D UI materials
+ */
+class SMaterialEditorUIPreviewViewport : public SCompoundWidget
+{
+public:
+	SLATE_BEGIN_ARGS( SMaterialEditorUIPreviewViewport ){}
+	SLATE_END_ARGS()
+
+	void Construct( const FArguments& InArgs, UMaterialInterface* PreviewMaterial );
+
+private:
+	void OnPreviewXChanged( int32 NewValue );
+	void OnPreviewXCommitted( int32 NewValue, ETextCommit::Type );
+	void OnPreviewYChanged( int32 NewValue );
+	void OnPreviewYCommitted( int32 NewValue, ETextCommit::Type );
+	TOptional<int32> OnGetPreviewXValue() const { return PreviewSize.X; }
+	TOptional<int32> OnGetPreviewYValue() const { return PreviewSize.Y; }
+private:
+	FIntPoint PreviewSize;
+	TSharedPtr<class SMaterialEditorUIPreviewZoomer> PreviewZoomer;
+};
\ No newline at end of file
diff --git a/Engine/Source/Editor/MaterialEditor/Private/SMaterialEditorViewportToolBar.cpp b/Engine/Source/Editor/MaterialEditor/Private/SMaterialEditorViewportToolBar.cpp
--- a/Engine/Source/Editor/MaterialEditor/Private/SMaterialEditorViewportToolBar.cpp
+++ b/Engine/Source/Editor/MaterialEditor/Private/SMaterialEditorViewportToolBar.cpp
@@ -12,7 +12,7 @@
 ///////////////////////////////////////////////////////////
 // SMaterialEditorViewportPreviewShapeToolBar
 
-void SMaterialEditorViewportPreviewShapeToolBar::Construct(const FArguments& InArgs, TSharedPtr<class SMaterialEditorViewport> InViewport)
+void SMaterialEditorViewportPreviewShapeToolBar::Construct(const FArguments& InArgs, TSharedPtr<class SMaterialEditor3DPreviewViewport> InViewport)
 {
 	// Force this toolbar to have small icons, as the preview panel is only small so we have limited space
 	const bool bForceSmallIcons = true;
@@ -54,7 +54,7 @@ void SMaterialEditorViewportPreviewShapeToolBar::Construct(const FArguments& InA
 ///////////////////////////////////////////////////////////
 // SMaterialEditorViewportToolBar
 
-void SMaterialEditorViewportToolBar::Construct(const FArguments& InArgs, TSharedPtr<class SMaterialEditorViewport> InViewport)
+void SMaterialEditorViewportToolBar::Construct(const FArguments& InArgs, TSharedPtr<class SMaterialEditor3DPreviewViewport> InViewport)
 {
 	SCommonEditorViewportToolbarBase::Construct(SCommonEditorViewportToolbarBase::FArguments(), InViewport);
 }
diff --git a/Engine/Source/Editor/MaterialEditor/Private/SMaterialEditorViewportToolBar.h b/Engine/Source/Editor/MaterialEditor/Private/SMaterialEditorViewportToolBar.h
--- a/Engine/Source/Editor/MaterialEditor/Private/SMaterialEditorViewportToolBar.h
+++ b/Engine/Source/Editor/MaterialEditor/Private/SMaterialEditorViewportToolBar.h
@@ -14,7 +14,7 @@ public:
 	SLATE_BEGIN_ARGS(SMaterialEditorViewportToolBar) {}
 	SLATE_END_ARGS()
 
-	void Construct(const FArguments& InArgs, TSharedPtr<class SMaterialEditorViewport> InViewport);
+	void Construct(const FArguments& InArgs, TSharedPtr<class SMaterialEditor3DPreviewViewport> InViewport);
 
 	// SCommonEditorViewportToolbarBase interface
 	virtual TSharedRef<SWidget> GenerateShowMenu() const override;
@@ -30,5 +30,5 @@ public:
 	SLATE_BEGIN_ARGS(SMaterialEditorViewportPreviewShapeToolBar){}
 	SLATE_END_ARGS()
 
-	void Construct(const FArguments& InArgs, TSharedPtr<class SMaterialEditorViewport> InViewport);
+	void Construct(const FArguments& InArgs, TSharedPtr<class SMaterialEditor3DPreviewViewport> InViewport);
 };
diff --git a/Engine/Source/Editor/PropertyEditor/Private/DetailLayoutBuilderImpl.cpp b/Engine/Source/Editor/PropertyEditor/Private/DetailLayoutBuilderImpl.cpp
--- a/Engine/Source/Editor/PropertyEditor/Private/DetailLayoutBuilderImpl.cpp
+++ b/Engine/Source/Editor/PropertyEditor/Private/DetailLayoutBuilderImpl.cpp
@@ -125,9 +125,9 @@ void FDetailLayoutBuilderImpl::BuildCategories( const FCategoryMap& CategoryMap,
 	{
 		TSharedRef<FDetailCategoryImpl> DetailCategory = It.Value().ToSharedRef();
 
-		const bool bCategoryHiddenByClass = GetDetailsView().IsCategoryHiddenByClass( DetailCategory->GetCategoryName() );
+		const bool bCategoryHidden = GetDetailsView().IsCategoryHiddenByClass( DetailCategory->GetCategoryName() ) || ForceHiddenCategories.Contains( DetailCategory->GetCategoryName() );
 
-		if( !bCategoryHiddenByClass )
+		if( !bCategoryHidden )
 		{
 			DetailCategory->GenerateLayout();
 
@@ -429,6 +429,11 @@ bool FDetailLayoutBuilderImpl::IsPropertyVisible( const struct FPropertyAndParen
 	return DetailsView.IsPropertyVisible( PropertyAndParent );
 }
 
+void FDetailLayoutBuilderImpl::HideCategory( FName CategoryName )
+{
+	ForceHiddenCategories.Add(CategoryName);
+}
+
 const IDetailsView& FDetailLayoutBuilderImpl::GetDetailsView() const
 { 
 	return DetailsView; 
diff --git a/Engine/Source/Editor/PropertyEditor/Private/DetailLayoutBuilderImpl.h b/Engine/Source/Editor/PropertyEditor/Private/DetailLayoutBuilderImpl.h
--- a/Engine/Source/Editor/PropertyEditor/Private/DetailLayoutBuilderImpl.h
+++ b/Engine/Source/Editor/PropertyEditor/Private/DetailLayoutBuilderImpl.h
@@ -19,7 +19,7 @@ public:
 	virtual TSharedPtr<FAssetThumbnailPool> GetThumbnailPool() const override;
 	virtual bool IsPropertyVisible( TSharedRef<IPropertyHandle> PropertyHandle ) const override;
 	virtual bool IsPropertyVisible( const struct FPropertyAndParent& PropertyAndParent ) const override;
-
+	virtual void HideCategory( FName CategoryName ) override;
 	virtual const TSharedRef< IPropertyUtilities >& GetPropertyUtilities() const override; 
 
 	/**
@@ -165,6 +165,8 @@ private:
 	FCategoryMap DefaultCategoryMap;
 	/** A mapping of classes to top level properties in that class */
 	FClassToPropertyMap& PropertyMap;
+	/** Force hidden categories set by the user */
+	TSet<FName> ForceHiddenCategories;
 	/** Nodes that require ticking */
 	TSet<IDetailTreeNode*> TickableNodes;
 	/** Current filter applied to the view */
diff --git a/Engine/Source/Editor/PropertyEditor/Public/DetailLayoutBuilder.h b/Engine/Source/Editor/PropertyEditor/Public/DetailLayoutBuilder.h
--- a/Engine/Source/Editor/PropertyEditor/Public/DetailLayoutBuilder.h
+++ b/Engine/Source/Editor/PropertyEditor/Public/DetailLayoutBuilder.h
@@ -73,6 +73,12 @@ public:
 	 */
 	virtual IDetailCategoryBuilder& EditCategory( FName CategoryName, const FText& NewLocalizedDisplayName = FText::GetEmpty(), ECategoryPriority::Type CategoryType = ECategoryPriority::Default ) = 0;
 
+	/**
+	 * Hides an entire category
+	 *
+	 * @param CategoryName	The name of the category to hide
+	 */
+	virtual void HideCategory( FName CategoryName ) = 0;
 	
 	/**
 	 * Gets a handle to a property which can be used to read and write the property value and identify the property in other detail customization interfaces.
diff --git a/Engine/Source/Editor/UnrealEd/Classes/MaterialGraph/MaterialGraph.h b/Engine/Source/Editor/UnrealEd/Classes/MaterialGraph/MaterialGraph.h
--- a/Engine/Source/Editor/UnrealEd/Classes/MaterialGraph/MaterialGraph.h
+++ b/Engine/Source/Editor/UnrealEd/Classes/MaterialGraph/MaterialGraph.h
@@ -14,11 +14,6 @@ DECLARE_DELEGATE_OneParam( FToggleExpressionCollapsed, UMaterialExpression* );
  */
 struct FMaterialInputInfo
 {
-	/** Name of the input shown to user */
-	FText Name;
-	/** Type of the input */
-	EMaterialProperty Property;
-
 	/** Constructor */
 	FMaterialInputInfo()
 	{
@@ -49,6 +44,10 @@ struct FMaterialInputInfo
 		{
 			return Property == MP_MaterialAttributes;
 		}
+		else if( Material->IsUIMaterial() )
+		{
+			return Property == MP_EmissiveColor || Property == MP_Opacity || Property == MP_OpacityMask;
+		}
 		else
 		{
 			if(Property == MP_MaterialAttributes)
@@ -64,6 +63,17 @@ struct FMaterialInputInfo
 			return true;
 		}
 	}
+
+	const FText& GetName() const { return Name; }
+
+	EMaterialProperty GetProperty() const { return Property; }
+
+private:
+	/** Name of the input shown to user */
+	FText Name;
+	/** Type of the input */
+	EMaterialProperty Property;
+
 };
 
 UCLASS()
@@ -155,4 +165,7 @@ private:
 	 * @param	Input	Input we are finding an output index for
 	 */
 	int32 GetValidOutputIndex(FExpressionInput* Input) const;
+
+	FText GetEmissivePinName() const;
+	FText GetBaseColorPinName() const;
 };
diff --git a/Engine/Source/Editor/UnrealEd/Private/MaterialGraph.cpp b/Engine/Source/Editor/UnrealEd/Private/MaterialGraph.cpp
--- a/Engine/Source/Editor/UnrealEd/Private/MaterialGraph.cpp
+++ b/Engine/Source/Editor/UnrealEd/Private/MaterialGraph.cpp
@@ -31,11 +31,11 @@ void UMaterialGraph::RebuildGraph()
 	if (!MaterialFunction)
 	{
 		// Initialize the material input list.
-		MaterialInputs.Add( FMaterialInputInfo( LOCTEXT("BaseColor", "Base Color"), MP_BaseColor ) );	
+		MaterialInputs.Add( FMaterialInputInfo( GetBaseColorPinName(), MP_BaseColor ) );	
 		MaterialInputs.Add( FMaterialInputInfo( LOCTEXT("Metallic", "Metallic"), MP_Metallic ) );
 		MaterialInputs.Add( FMaterialInputInfo( LOCTEXT("Specular", "Specular"), MP_Specular ) );
 		MaterialInputs.Add( FMaterialInputInfo( LOCTEXT("Roughness", "Roughness"), MP_Roughness ) );
-		MaterialInputs.Add( FMaterialInputInfo( LOCTEXT("EmissiveColor", "Emissive Color"), MP_EmissiveColor ) );
+		MaterialInputs.Add( FMaterialInputInfo( GetEmissivePinName(), MP_EmissiveColor ) );
 		MaterialInputs.Add( FMaterialInputInfo( LOCTEXT("Opacity", "Opacity"), MP_Opacity ) );
 		MaterialInputs.Add( FMaterialInputInfo( LOCTEXT("OpacityMask", "Opacity Mask"), MP_OpacityMask ) );
 		MaterialInputs.Add( FMaterialInputInfo( LOCTEXT("Normal", "Normal"), MP_Normal ) );
@@ -301,7 +301,7 @@ bool UMaterialGraph::IsInputActive(UEdGraphPin* GraphPin) const
 		{
 			if (RootNode->Pins[Index] == GraphPin)
 			{
-				return Material->IsPropertyActive(MaterialInputs[Index].Property);
+				return Material->IsPropertyActive(MaterialInputs[Index].GetProperty());
 			}
 		}
 	}
@@ -453,4 +453,14 @@ int32 UMaterialGraph::GetValidOutputIndex(FExpressionInput* Input) const
 	return OutputIndex;
 }
 
+FText UMaterialGraph::GetEmissivePinName() const
+{
+	return Material->IsUIMaterial() ? LOCTEXT("UIOutputColor", "Final Color") : LOCTEXT("EmissiveColor", "Emissive Color");
+}
+
+FText UMaterialGraph::GetBaseColorPinName() const
+{
+	return LOCTEXT("BaseColor", "Base Color");
+}
+
 #undef LOCTEXT_NAMESPACE
diff --git a/Engine/Source/Editor/UnrealEd/Private/MaterialGraphNode_Root.cpp b/Engine/Source/Editor/UnrealEd/Private/MaterialGraphNode_Root.cpp
--- a/Engine/Source/Editor/UnrealEd/Private/MaterialGraphNode_Root.cpp
+++ b/Engine/Source/Editor/UnrealEd/Private/MaterialGraphNode_Root.cpp
@@ -50,7 +50,7 @@ void UMaterialGraphNode_Root::CreateInputPins()
 
 	for (int32 Index = 0; Index < MaterialGraph->MaterialInputs.Num(); ++Index)
 	{
-		UEdGraphPin* InputPin = CreatePin(EGPD_Input, Schema->PC_MaterialInput, TEXT(""), NULL, /*bIsArray=*/ false, /*bIsReference=*/ false, MaterialGraph->MaterialInputs[Index].Name.ToString());
+		UEdGraphPin* InputPin = CreatePin(EGPD_Input, Schema->PC_MaterialInput, TEXT(""), NULL, /*bIsArray=*/ false, /*bIsReference=*/ false, MaterialGraph->MaterialInputs[Index].GetName().ToString());
 	}
 }
 
@@ -70,7 +70,7 @@ int32 UMaterialGraphNode_Root::GetInputIndex(const UEdGraphPin* InputPin) const
 uint32 UMaterialGraphNode_Root::GetInputType(const UEdGraphPin* InputPin) const
 {
 	UMaterialGraph* MaterialGraph = CastChecked<UMaterialGraph>(GetGraph());
-	return GetMaterialPropertyType(MaterialGraph->MaterialInputs[GetInputIndex(InputPin)].Property);
+	return GetMaterialPropertyType(MaterialGraph->MaterialInputs[GetInputIndex(InputPin)].GetProperty());
 }
 
 #undef LOCTEXT_NAMESPACE
diff --git a/Engine/Source/Editor/UnrealEd/Private/MaterialGraphSchema.cpp b/Engine/Source/Editor/UnrealEd/Private/MaterialGraphSchema.cpp
--- a/Engine/Source/Editor/UnrealEd/Private/MaterialGraphSchema.cpp
+++ b/Engine/Source/Editor/UnrealEd/Private/MaterialGraphSchema.cpp
@@ -397,7 +397,7 @@ void UMaterialGraphSchema::GetContextMenuActions(const UEdGraph* CurrentGraph, c
 						if(MaterialGraph->MaterialInputs[Index].IsVisiblePin(MaterialGraph->Material))
 						{
 							FFormatNamedArguments Arguments;
-							Arguments.Add(TEXT("Name"), MaterialGraph->MaterialInputs[Index].Name);
+							Arguments.Add(TEXT("Name"), MaterialGraph->MaterialInputs[Index].GetName());
 							const FText Label = FText::Format( LOCTEXT( "ConnectToInput", "Connect To {Name}" ), Arguments );
 							const FText ToolTip = FText::Format( LOCTEXT( "ConnectToInputTooltip", "Connects to the material input {Name}" ), Arguments );
 							MenuBuilder->AddMenuEntry(Label,
diff --git a/Engine/Source/Editor/UnrealEd/Private/PreviewMaterial.cpp b/Engine/Source/Editor/UnrealEd/Private/PreviewMaterial.cpp
--- a/Engine/Source/Editor/UnrealEd/Private/PreviewMaterial.cpp
+++ b/Engine/Source/Editor/UnrealEd/Private/PreviewMaterial.cpp
@@ -37,83 +37,96 @@ public:
 		// only generate the needed shaders (which should be very restrictive for fast recompiling during editing)
 		// @todo: Add a FindShaderType by fname or something
 
-		bool bEditorStatsMaterial = Material->bIsMaterialEditorStatsMaterial;
-
-		// Always allow HitProxy shaders.
-		if (FCString::Stristr(ShaderType->GetName(), TEXT("HitProxy")))
+		if( Material->IsUIMaterial() )
 		{
-			return true;
+			if (FCString::Stristr(ShaderType->GetName(), TEXT("TSlateMaterialShaderPS")))
+			{
+				return true;
+			}
+	
 		}
 
-		// we only need local vertex factory for the preview static mesh
-		if (VertexFactoryType != FindVertexFactoryType(FName(TEXT("FLocalVertexFactory"), FNAME_Find)))
+
 		{
-			//cache for gpu skinned vertex factory if the material allows it
-			//this way we can have a preview skeletal mesh
-			if (bEditorStatsMaterial ||
-				!IsUsedWithSkeletalMesh() ||
-				(VertexFactoryType != FindVertexFactoryType(FName(TEXT("TGPUSkinVertexFactoryfalse"), FNAME_Find)) &&
-				VertexFactoryType != FindVertexFactoryType(FName(TEXT("TGPUSkinVertexFactorytrue"), FNAME_Find))))
+			bool bEditorStatsMaterial = Material->bIsMaterialEditorStatsMaterial;
+
+			// Always allow HitProxy shaders.
+			if (FCString::Stristr(ShaderType->GetName(), TEXT("HitProxy")))
 			{
-				return false;
+				return true;
 			}
-		}
 
-		if (bEditorStatsMaterial)
-		{
-			TArray<FString> ShaderTypeNames;
-			TArray<FString> ShaderTypeDescriptions;
-			GetRepresentativeShaderTypesAndDescriptions(ShaderTypeNames, ShaderTypeDescriptions);
+			// we only need local vertex factory for the preview static mesh
+			if (VertexFactoryType != FindVertexFactoryType(FName(TEXT("FLocalVertexFactory"), FNAME_Find)))
+			{
+				//cache for gpu skinned vertex factory if the material allows it
+				//this way we can have a preview skeletal mesh
+				if (bEditorStatsMaterial ||
+					!IsUsedWithSkeletalMesh() ||
+					(VertexFactoryType != FindVertexFactoryType(FName(TEXT("TGPUSkinVertexFactoryfalse"), FNAME_Find)) &&
+					VertexFactoryType != FindVertexFactoryType(FName(TEXT("TGPUSkinVertexFactorytrue"), FNAME_Find))))
+				{
+					return false;
+				}
+			}
 
-			//Only allow shaders that are used in the stats.
-			return ShaderTypeNames.Find(ShaderType->GetName()) != INDEX_NONE;
-		}
+			if (bEditorStatsMaterial)
+			{
+				TArray<FString> ShaderTypeNames;
+				TArray<FString> ShaderTypeDescriptions;
+				GetRepresentativeShaderTypesAndDescriptions(ShaderTypeNames, ShaderTypeDescriptions);
 
-		// look for any of the needed type
-		bool bShaderTypeMatches = false;
+				//Only allow shaders that are used in the stats.
+				return ShaderTypeNames.Find(ShaderType->GetName()) != INDEX_NONE;
+			}
 
-		// For FMaterialResource::GetRepresentativeInstructionCounts
-		if (FCString::Stristr(ShaderType->GetName(), TEXT("BasePassPSTDistanceFieldShadowsAndLightMapPolicyHQ")))
-		{
-			bShaderTypeMatches = true;
-		}
-		else if (FCString::Stristr(ShaderType->GetName(), TEXT("BasePassPSFNoLightMapPolicy")))
-		{
-			bShaderTypeMatches = true;
-		}
-		else if (FCString::Stristr(ShaderType->GetName(), TEXT("CachedPointIndirectLightingPolicy")))
-		{
-			bShaderTypeMatches = true;
-		}
-		else if (FCString::Stristr(ShaderType->GetName(), TEXT("BasePassPSFSelfShadowedTranslucencyPolicy")))
-		{
-			bShaderTypeMatches = true;
-		}
-		// Pick tessellation shader based on material settings
-		else if(FCString::Stristr(ShaderType->GetName(), TEXT("BasePassVSFNoLightMapPolicy")) ||
-			FCString::Stristr(ShaderType->GetName(), TEXT("BasePassHSFNoLightMapPolicy")) ||
-			FCString::Stristr(ShaderType->GetName(), TEXT("BasePassDSFNoLightMapPolicy")))
-		{
-			bShaderTypeMatches = true;
-		}
-		else if (FCString::Stristr(ShaderType->GetName(), TEXT("DepthOnly")))
-		{
-			bShaderTypeMatches = true;
-		}
-		else if (FCString::Stristr(ShaderType->GetName(), TEXT("ShadowDepth")))
-		{
-			bShaderTypeMatches = true;
-		}
-		else if (FCString::Stristr(ShaderType->GetName(), TEXT("TDistortion")))
-		{
-			bShaderTypeMatches = true;
-		}
-		else if (FCString::Stristr(ShaderType->GetName(), TEXT("TBasePassForForwardShading")))
-		{
-			bShaderTypeMatches = true;
-		}
+			// look for any of the needed type
+			bool bShaderTypeMatches = false;
 
-		return bShaderTypeMatches;
+			// For FMaterialResource::GetRepresentativeInstructionCounts
+			if (FCString::Stristr(ShaderType->GetName(), TEXT("BasePassPSTDistanceFieldShadowsAndLightMapPolicyHQ")))
+			{
+				bShaderTypeMatches = true;
+			}
+			else if (FCString::Stristr(ShaderType->GetName(), TEXT("BasePassPSFNoLightMapPolicy")))
+			{
+				bShaderTypeMatches = true;
+			}
+			else if (FCString::Stristr(ShaderType->GetName(), TEXT("CachedPointIndirectLightingPolicy")))
+			{
+				bShaderTypeMatches = true;
+			}
+			else if (FCString::Stristr(ShaderType->GetName(), TEXT("BasePassPSFSelfShadowedTranslucencyPolicy")))
+			{
+				bShaderTypeMatches = true;
+			}
+			// Pick tessellation shader based on material settings
+			else if(FCString::Stristr(ShaderType->GetName(), TEXT("BasePassVSFNoLightMapPolicy")) ||
+				FCString::Stristr(ShaderType->GetName(), TEXT("BasePassHSFNoLightMapPolicy")) ||
+				FCString::Stristr(ShaderType->GetName(), TEXT("BasePassDSFNoLightMapPolicy")))
+			{
+				bShaderTypeMatches = true;
+			}
+			else if (FCString::Stristr(ShaderType->GetName(), TEXT("DepthOnly")))
+			{
+				bShaderTypeMatches = true;
+			}
+			else if (FCString::Stristr(ShaderType->GetName(), TEXT("ShadowDepth")))
+			{
+				bShaderTypeMatches = true;
+			}
+			else if (FCString::Stristr(ShaderType->GetName(), TEXT("TDistortion")))
+			{
+				bShaderTypeMatches = true;
+			}
+			else if (FCString::Stristr(ShaderType->GetName(), TEXT("TBasePassForForwardShading")))
+			{
+				bShaderTypeMatches = true;
+			}
+
+			return bShaderTypeMatches;
+		}
+	
 	}
 
 	/**
diff --git a/Engine/Source/Editor/UnrealEd/Private/ThumbnailHelpers.cpp b/Engine/Source/Editor/UnrealEd/Private/ThumbnailHelpers.cpp
--- a/Engine/Source/Editor/UnrealEd/Private/ThumbnailHelpers.cpp
+++ b/Engine/Source/Editor/UnrealEd/Private/ThumbnailHelpers.cpp
@@ -253,6 +253,7 @@ void FMaterialThumbnailScene::SetMaterialInterface(UMaterialInterface* InMateria
 	check(PreviewActor);
 	check(PreviewActor->GetStaticMeshComponent());
 
+	bIsUIMaterial = false;
 	if ( InMaterial )
 	{
 		// Transform the preview mesh as necessary
@@ -264,7 +265,14 @@ void FMaterialThumbnailScene::SetMaterialInterface(UMaterialInterface* InMateria
 			ThumbnailInfo = USceneThumbnailInfoWithPrimitive::StaticClass()->GetDefaultObject<USceneThumbnailInfoWithPrimitive>();
 		}
 
-		switch( ThumbnailInfo->PrimitiveType )
+		UMaterial* BaseMaterial = InMaterial->GetBaseMaterial();
+
+		// UI material thumbnails always get a 2D plane centered at the camera which is a better representation of the
+		// what the material will look like on UI
+		bIsUIMaterial = BaseMaterial && BaseMaterial->IsUIMaterial();
+		EThumbnailPrimType PrimitiveType = bIsUIMaterial ? TPT_Plane : ThumbnailInfo->PrimitiveType;
+
+		switch( PrimitiveType )
 		{
 		case TPT_None:
 			{
@@ -355,7 +363,7 @@ void FMaterialThumbnailScene::GetViewMatrixParameters(const float InFOVDegrees,
 	}
 
 	OutOrigin = FVector(0, 0, -BoundsZOffset);
-	OutOrbitPitch = ThumbnailInfo->OrbitPitch;
+	OutOrbitPitch = bIsUIMaterial ? 0.0f : ThumbnailInfo->OrbitPitch;
 	OutOrbitYaw = ThumbnailInfo->OrbitYaw;
 	OutOrbitZoom = TargetDistance + ThumbnailInfo->OrbitZoom;
 }
diff --git a/Engine/Source/Editor/UnrealEd/Public/ThumbnailHelpers.h b/Engine/Source/Editor/UnrealEd/Public/ThumbnailHelpers.h
--- a/Engine/Source/Editor/UnrealEd/Public/ThumbnailHelpers.h
+++ b/Engine/Source/Editor/UnrealEd/Public/ThumbnailHelpers.h
@@ -70,6 +70,8 @@ protected:
 protected:
 	/** The static mesh actor used to display all material thumbnails */
 	AStaticMeshActor* PreviewActor;
+	/** Material being rendered is for UI */
+	bool bIsUIMaterial;
 };
 
 class UNREALED_API FSkeletalMeshThumbnailScene : public FThumbnailPreviewScene
diff --git a/Engine/Source/Runtime/Core/Public/UObject/ObjectVersion.h b/Engine/Source/Runtime/Core/Public/UObject/ObjectVersion.h
--- a/Engine/Source/Runtime/Core/Public/UObject/ObjectVersion.h
+++ b/Engine/Source/Runtime/Core/Public/UObject/ObjectVersion.h
@@ -536,6 +536,8 @@ enum EUnrealEngineObjectUE4Version
 	VER_UE4_NIAGARA_DATA_OBJECT_DEV_UI_FIX,
 	// Fixed the default orientation of widget component so it faces down +x
 	VER_UE4_FIXED_DEFAULT_ORIENTATION_OF_WIDGET_COMPONENT,
+	// Removed bUsedWithUI flag from UMaterial and replaced it with a new material domain for UI
+	VER_UE4_REMOVED_MATERIAL_USED_WITH_UI_FLAG,
 	// -----<new versions can be added before this line>-------------------------------------------------
 	// - this needs to be the last line (see note below)
 	VER_UE4_AUTOMATIC_VERSION_PLUS_ONE,
diff --git a/Engine/Source/Runtime/Engine/Classes/Materials/Material.h b/Engine/Source/Runtime/Engine/Classes/Materials/Material.h
--- a/Engine/Source/Runtime/Engine/Classes/Materials/Material.h
+++ b/Engine/Source/Runtime/Engine/Classes/Materials/Material.h
@@ -90,6 +90,9 @@ enum EMaterialDomain
 	MD_LightFunction UMETA(DisplayName="Light Function"),
 	/** The material will be used in a custom post process pass. */
 	MD_PostProcess UMETA(DisplayName="Post Process"),
+	/** The material wwill be used for UMG or Slate UI */
+	MD_UI UMETA(DisplayName="User Interface"),
+
 	MD_MAX
 };
 
@@ -596,8 +599,8 @@ public:
 	 * Indicates that the material and its instances can be use with Slate UI and UMG
 	 * This will result in the shaders required to support UI materials being compiled which will increase shader compile time and memory usage.
 	 */
-	UPROPERTY(EditAnywhere, BlueprintReadOnly, Category=Usage)
-	uint32 bUsedWithUI:1;
+	UPROPERTY()
+	uint32 bUsedWithUI_DEPRECATED:1;
 
 	/** 
 	 * Whether to automatically set usage flags based on what the material is applied to in the editor.
@@ -799,6 +802,7 @@ public:
 	ENGINE_API virtual bool IsTwoSided(bool bIsGameThread = IsInGameThread()) const override;
 	ENGINE_API virtual bool IsDitheredLODTransition(bool bIsGameThread = IsInGameThread()) const override;
 	ENGINE_API virtual bool IsMasked(bool bIsGameThread = IsInGameThread()) const override;
+	ENGINE_API virtual bool IsUIMaterial() const { return MaterialDomain == MD_UI; }
 	ENGINE_API virtual USubsurfaceProfile* GetSubsurfaceProfile_Internal() const override;
 
 	ENGINE_API void SetShadingModel(EMaterialShadingModel NewModel) {ShadingModel = NewModel;}
diff --git a/Engine/Source/Runtime/Engine/Classes/Materials/MaterialInterface.h b/Engine/Source/Runtime/Engine/Classes/Materials/MaterialInterface.h
--- a/Engine/Source/Runtime/Engine/Classes/Materials/MaterialInterface.h
+++ b/Engine/Source/Runtime/Engine/Classes/Materials/MaterialInterface.h
@@ -27,7 +27,6 @@ enum EMaterialUsage
 	MATUSAGE_Landscape,
 	MATUSAGE_InstancedStaticMeshes,
 	MATUSAGE_Clothing,
-	MATUSAGE_UI,
 	MATUSAGE_MAX,
 };
 
diff --git a/Engine/Source/Runtime/Engine/Private/Materials/Material.cpp b/Engine/Source/Runtime/Engine/Private/Materials/Material.cpp
--- a/Engine/Source/Runtime/Engine/Private/Materials/Material.cpp
+++ b/Engine/Source/Runtime/Engine/Private/Materials/Material.cpp
@@ -325,10 +325,16 @@ static UMaterial* GDefaultMaterials[MD_MAX] = {0};
 
 static const TCHAR* GDefaultMaterialNames[MD_MAX] =
 {
+	// Surface
 	TEXT("engine-ini:/Script/Engine.Engine.DefaultMaterialName"),
+	// Deferred Decal
 	TEXT("engine-ini:/Script/Engine.Engine.DefaultDeferredDecalMaterialName"),
+	// Light Function
 	TEXT("engine-ini:/Script/Engine.Engine.DefaultLightFunctionMaterialName"),
-	TEXT("engine-ini:/Script/Engine.Engine.DefaultPostProcessMaterialName")
+	// Post Process
+	TEXT("engine-ini:/Script/Engine.Engine.DefaultPostProcessMaterialName"),
+	// User Interface 
+	TEXT("engine-ini:/Script/Engine.Engine.DefaultMaterialName"),
 };
 
 void UMaterialInterface::InitDefaultMaterials()
@@ -784,7 +790,6 @@ bool UMaterial::GetUsageByFlag(EMaterialUsage Usage) const
 		case MATUSAGE_SplineMesh: UsageValue = bUsedWithSplineMeshes; break;
 		case MATUSAGE_InstancedStaticMeshes: UsageValue = bUsedWithInstancedStaticMeshes; break;
 		case MATUSAGE_Clothing: UsageValue = bUsedWithClothing; break;
-		case MATUSAGE_UI: UsageValue = bUsedWithUI; break;
 		default: UE_LOG(LogMaterial, Fatal,TEXT("Unknown material usage: %u"), (int32)Usage);
 	};
 	return UsageValue;
@@ -867,10 +872,6 @@ void UMaterial::SetUsageByFlag(EMaterialUsage Usage, bool NewValue)
 		{
 			bUsedWithClothing = NewValue; break;
 		}
-		case MATUSAGE_UI:
-		{
-			bUsedWithUI = NewValue; break;
-		}
 		default: UE_LOG(LogMaterial, Fatal,TEXT("Unknown material usage: %u"), (int32)Usage);
 	};
 #if WITH_EDITOR
@@ -894,7 +895,6 @@ FString UMaterial::GetUsageName(EMaterialUsage Usage) const
 		case MATUSAGE_SplineMesh: UsageName = TEXT("bUsedWithSplineMeshes"); break;
 		case MATUSAGE_InstancedStaticMeshes: UsageName = TEXT("bUsedWithInstancedStaticMeshes"); break;
 		case MATUSAGE_Clothing: UsageName = TEXT("bUsedWithClothing"); break;
-		case MATUSAGE_UI: UsageName = TEXT("bUsedWithUI"); break;
 		default: UE_LOG(LogMaterial, Fatal,TEXT("Unknown material usage: %u"), (int32)Usage);
 	};
 	return UsageName;
@@ -2174,6 +2174,12 @@ void UMaterial::PostLoad()
 		bIsMaterialEditorStatsMaterial = true;
 	}
 
+
+	if( GetLinkerUE4Version() < VER_UE4_REMOVED_MATERIAL_USED_WITH_UI_FLAG && bUsedWithUI_DEPRECATED == true )
+	{
+		MaterialDomain = MD_UI;
+	}
+
 	// Ensure expressions have been postloaded before we use them for compiling
 	// Any UObjects used by material compilation must be postloaded here
 	for (int32 ExpressionIndex = 0; ExpressionIndex < Expressions.Num(); ExpressionIndex++)
@@ -2489,7 +2495,7 @@ bool UMaterial::CanEditChange(const UProperty* InProperty) const
 
 		if (PropertyName == GET_MEMBER_NAME_STRING_CHECKED(UMaterial, BlendMode))
 		{
-			return MaterialDomain == MD_Surface;
+			return MaterialDomain == MD_Surface || MaterialDomain == MD_UI;
 		}
 	
 		if (PropertyName == GET_MEMBER_NAME_STRING_CHECKED(UMaterial, ShadingModel))
@@ -3773,6 +3779,7 @@ EMaterialShadingModel UMaterial::GetShadingModel(bool bIsInGameThread) const
 		// Post process and light function materials must be rendered with the unlit model.
 		case MD_PostProcess:
 		case MD_LightFunction:
+		case MD_UI:
 			return MSM_Unlit;
 
 		default:
@@ -3902,6 +3909,12 @@ bool UMaterial::IsPropertyActive(EMaterialProperty InProperty) const
 				return false;
 		}
 	}
+	else if ( MaterialDomain == MD_UI )
+	{
+		return InProperty == MP_EmissiveColor
+			|| ( InProperty == MP_OpacityMask && BlendMode == BLEND_Masked ) 
+			|| ( InProperty == MP_Opacity && IsTranslucentBlendMode((EBlendMode)BlendMode) && BlendMode != BLEND_Modulate );
+	}
 
 	bool Active = true;
 
@@ -3953,7 +3966,7 @@ bool UMaterial::IsPropertyActive(EMaterialProperty InProperty) const
 		Active = true;
 		break;
 	case MP_WorldPositionOffset:
-		Active = !bUsedWithUI;
+		Active = true;
 		break;
 	case MP_PixelDepthOffset:
 		Active = !IsTranslucentBlendMode((EBlendMode)BlendMode);
diff --git a/Engine/Source/Runtime/Engine/Private/Materials/MaterialShared.cpp b/Engine/Source/Runtime/Engine/Private/Materials/MaterialShared.cpp
--- a/Engine/Source/Runtime/Engine/Private/Materials/MaterialShared.cpp
+++ b/Engine/Source/Runtime/Engine/Private/Materials/MaterialShared.cpp
@@ -707,6 +707,7 @@ bool FMaterialResource::ShouldDisableDepthTest() const { return Material->bDisab
 bool FMaterialResource::ShouldEnableResponsiveAA() const { return Material->bEnableResponsiveAA; }
 bool FMaterialResource::ShouldDoSSR() const { return Material->bScreenSpaceReflections; }
 bool FMaterialResource::IsWireframe() const { return Material->Wireframe; }
+bool FMaterialResource::IsUIMaterial() const { return Material->MaterialDomain == MD_UI; }
 bool FMaterialResource::IsLightFunction() const { return Material->MaterialDomain == MD_LightFunction; }
 bool FMaterialResource::IsUsedWithEditorCompositing() const { return Material->bUsedWithEditorCompositing; }
 bool FMaterialResource::IsUsedWithDeferredDecal() const { return Material->MaterialDomain == MD_DeferredDecal; }
@@ -773,7 +774,7 @@ bool FMaterialResource::IsUsedWithAPEXCloth() const
 
 bool FMaterialResource::IsUsedWithUI() const
 {
-	return Material->bUsedWithUI;
+	return IsUIMaterial();
 }
 
 EMaterialTessellationMode FMaterialResource::GetTessellationMode() const 
@@ -922,23 +923,40 @@ void FMaterialResource::GetRepresentativeInstructionCounts(TArray<FString> &Desc
 	{
 		GetRepresentativeShaderTypesAndDescriptions(ShaderTypeNames, ShaderTypeDescriptions);
 
-		const FMeshMaterialShaderMap* MeshShaderMap = MaterialShaderMap->GetMeshShaderMap(&FLocalVertexFactory::StaticType);
-		if (MeshShaderMap)
+		if( IsUIMaterial() )
 		{
-			Descriptions.Empty();
-			InstructionCounts.Empty();
-
 			for (int32 InstructionIndex = 0; InstructionIndex < ShaderTypeNames.Num(); InstructionIndex++)
 			{
 				FShaderType* ShaderType = FindShaderTypeByName(*ShaderTypeNames[InstructionIndex]);
-				if (ShaderType)
+				const FShader* Shader = MaterialShaderMap->GetShader(ShaderType);
+				if (Shader && Shader->GetNumInstructions() > 0)
 				{
-					const FShader* Shader = MeshShaderMap->GetShader(ShaderType);
-					if (Shader && Shader->GetNumInstructions() > 0)
+					//if the shader was found, add it to the output arrays
+					InstructionCounts.Push(Shader->GetNumInstructions());
+					Descriptions.Push(ShaderTypeDescriptions[InstructionIndex]);
+				}
+			}
+		}
+		else
+		{
+			const FMeshMaterialShaderMap* MeshShaderMap = MaterialShaderMap->GetMeshShaderMap(&FLocalVertexFactory::StaticType);
+			if (MeshShaderMap)
+			{
+				Descriptions.Empty();
+				InstructionCounts.Empty();
+
+				for (int32 InstructionIndex = 0; InstructionIndex < ShaderTypeNames.Num(); InstructionIndex++)
+				{
+					FShaderType* ShaderType = FindShaderTypeByName(*ShaderTypeNames[InstructionIndex]);
+					if (ShaderType)
 					{
-						//if the shader was found, add it to the output arrays
-						InstructionCounts.Push(Shader->GetNumInstructions());
-						Descriptions.Push(ShaderTypeDescriptions[InstructionIndex]);
+						const FShader* Shader = MeshShaderMap->GetShader(ShaderType);
+						if (Shader && Shader->GetNumInstructions() > 0)
+						{
+							//if the shader was found, add it to the output arrays
+							InstructionCounts.Push(Shader->GetNumInstructions());
+							Descriptions.Push(ShaderTypeDescriptions[InstructionIndex]);
+						}
 					}
 				}
 			}
@@ -953,7 +971,12 @@ void FMaterialResource::GetRepresentativeShaderTypesAndDescriptions(TArray<FStri
 	static auto* MobileHDR = IConsoleManager::Get().FindTConsoleVariableDataInt(TEXT("r.MobileHDR"));
 	bool bMobileHDR = MobileHDR && MobileHDR->GetValueOnAnyThread() == 1;
 
-	if (GetFeatureLevel() >= ERHIFeatureLevel::SM4)
+	if( IsUIMaterial() )
+	{
+		new(ShaderTypeNames) FString(TEXT("TSlateMaterialShaderPSDefaultfalse"));
+		new(ShaderTypeDescriptions) FString(TEXT("Default UI Shader"));
+	}
+	else if (GetFeatureLevel() >= ERHIFeatureLevel::SM4)
 	{
 		if (GetShadingModel() == MSM_Unlit)
 		{
diff --git a/Engine/Source/Runtime/Engine/Public/MaterialShared.h b/Engine/Source/Runtime/Engine/Public/MaterialShared.h
--- a/Engine/Source/Runtime/Engine/Public/MaterialShared.h
+++ b/Engine/Source/Runtime/Engine/Public/MaterialShared.h
@@ -929,6 +929,7 @@ public:
 	virtual bool IsUsedWithEditorCompositing() const { return false; }
 	virtual bool IsUsedWithDeferredDecal() const = 0;
 	virtual bool IsWireframe() const = 0;
+	virtual bool IsUIMaterial() const { return false; }
 	virtual bool IsSpecialEngineMaterial() const = 0;
 	virtual bool IsUsedWithSkeletalMesh() const { return false; }
 	virtual bool IsUsedWithLandscape() const { return false; }
@@ -1488,6 +1489,7 @@ public:
 	ENGINE_API virtual bool IsUsedWithEditorCompositing() const override;
 	ENGINE_API virtual bool IsUsedWithDeferredDecal() const override;
 	ENGINE_API virtual bool IsWireframe() const override;
+	ENGINE_API virtual bool IsUIMaterial() const override;
 	ENGINE_API virtual bool IsSpecialEngineMaterial() const override;
 	ENGINE_API virtual bool IsUsedWithSkeletalMesh() const override;
 	ENGINE_API virtual bool IsUsedWithLandscape() const override;
@@ -1500,6 +1502,7 @@ public:
 	ENGINE_API virtual bool IsUsedWithSplineMeshes() const override;
 	ENGINE_API virtual bool IsUsedWithInstancedStaticMeshes() const override;
 	ENGINE_API virtual bool IsUsedWithAPEXCloth() const override;
+	DEPRECATED(4.9, "IsUsedWithUI is now replaced by IsUIMaterial")
 	ENGINE_API virtual bool IsUsedWithUI() const override;
 	ENGINE_API virtual enum EMaterialTessellationMode GetTessellationMode() const override;
 	ENGINE_API virtual bool IsCrackFreeDisplacementEnabled() const override;
diff --git a/Engine/Source/Runtime/SlateCore/Public/Styling/SlateBrush.h b/Engine/Source/Runtime/SlateCore/Public/Styling/SlateBrush.h
--- a/Engine/Source/Runtime/SlateCore/Public/Styling/SlateBrush.h
+++ b/Engine/Source/Runtime/SlateCore/Public/Styling/SlateBrush.h
@@ -271,7 +271,7 @@ protected:
 	/**
 	 * The image to render for this brush, can be a UTexture2D or Material.
 	 */
-	UPROPERTY(EditAnywhere, Category=Brush, meta=( DisplayThumbnail="true", ThumbnailSize="X=40 Y=40", DisplayName="Image", AllowedClasses="Texture2D,MaterialInterface" ))
+	UPROPERTY(EditAnywhere, Category=Brush, meta=( DisplayThumbnail="true", DisplayName="Image", AllowedClasses="Texture2D,MaterialInterface" ))
 	UObject* ResourceObject;
 
 	/** The name of the rendering resource to use */
diff --git a/Engine/Source/Runtime/SlateRHIRenderer/Private/SlateMaterialShader.cpp b/Engine/Source/Runtime/SlateRHIRenderer/Private/SlateMaterialShader.cpp
--- a/Engine/Source/Runtime/SlateRHIRenderer/Private/SlateMaterialShader.cpp
+++ b/Engine/Source/Runtime/SlateRHIRenderer/Private/SlateMaterialShader.cpp
@@ -6,7 +6,7 @@
 
 bool FSlateMaterialShaderPS::ShouldCache(EShaderPlatform Platform, const FMaterial* Material)
 {
-	return Material->IsUsedWithUI();
+	return Material->GetMaterialDomain() == MD_UI;
 }
 
 
