Undo/Redo for switching material domains now properly updates the views in the material editor

diff --git a/Engine/Source/Editor/MaterialEditor/Private/MaterialEditor.cpp b/Engine/Source/Editor/MaterialEditor/Private/MaterialEditor.cpp
--- a/Engine/Source/Editor/MaterialEditor/Private/MaterialEditor.cpp
+++ b/Engine/Source/Editor/MaterialEditor/Private/MaterialEditor.cpp
@@ -3499,6 +3499,8 @@ void FMaterialEditor::PostUndo(bool bSuccess)
 		// Update the current preview material.
 		UpdatePreviewMaterial();
 
+		UpdatePreviewViewportsVisibility();
+
 		RefreshExpressionPreviews();
 		GraphEditor->NotifyGraphChanged();
 		SetMaterialDirty();
