Fix variable shadowing issues

diff --git a/Engine/Source/Editor/MaterialEditor/Private/SMaterialEditorViewport.cpp b/Engine/Source/Editor/MaterialEditor/Private/SMaterialEditorViewport.cpp
--- a/Engine/Source/Editor/MaterialEditor/Private/SMaterialEditorViewport.cpp
+++ b/Engine/Source/Editor/MaterialEditor/Private/SMaterialEditorViewport.cpp
@@ -598,7 +598,7 @@ public:
 	{
 	}
 
-	void Construct( const FArguments& InArgs, UMaterialInterface* PreviewMaterial );
+	void Construct( const FArguments& InArgs, UMaterialInterface* InPreviewMaterial );
 
 	virtual void OnArrangeChildren( const FGeometry& AllottedGeometry, FArrangedChildren& ArrangedChildren ) const override;
 	virtual FVector2D ComputeDesiredSize(float) const override;
@@ -615,15 +615,13 @@ private:
 	float ZoomLevel;
 
 	FMaterialPreviewPanelSlot ChildSlot;
-
-	TWeakObjectPtr<UMaterial> PreviewMaterial;
 	TSharedPtr<FSlateMaterialBrush> PreviewBrush;
 };
 
 
-void SMaterialEditorUIPreviewZoomer::Construct( const FArguments& InArgs, UMaterialInterface* PreviewMaterial )
+void SMaterialEditorUIPreviewZoomer::Construct( const FArguments& InArgs, UMaterialInterface* InPreviewMaterial )
 {
-	PreviewBrush = MakeShareable( new FSlateMaterialBrush( *PreviewMaterial, FVector2D(250,250) ) );
+	PreviewBrush = MakeShareable( new FSlateMaterialBrush( *InPreviewMaterial, FVector2D(250,250) ) );
 
 	ChildSlot
 	[
