Fix crash copy and pasting material function call nodes while live preview is disabled.  We now properly update the node from the material function resource

https://jira.ol.epicgames.net/browse/UE-17746

diff --git a/Engine/Source/Editor/MaterialEditor/Private/MaterialEditor.cpp b/Engine/Source/Editor/MaterialEditor/Private/MaterialEditor.cpp
--- a/Engine/Source/Editor/MaterialEditor/Private/MaterialEditor.cpp
+++ b/Engine/Source/Editor/MaterialEditor/Private/MaterialEditor.cpp
@@ -3375,6 +3375,12 @@ void FMaterialEditor::PasteNodesHere(const FVector2D& Location)
 				FunctionOutput->ConditionallyGenerateId(true);
 				FunctionOutput->ValidateName();
 			}
+
+			UMaterialExpressionMaterialFunctionCall* FunctionCall = Cast<UMaterialExpressionMaterialFunctionCall>( NewExpression );
+			if( FunctionCall )
+			{
+				FunctionCall->UpdateFromFunctionResource();
+			}
 		}
 		else if (UMaterialGraphNode_Comment* CommentNode = Cast<UMaterialGraphNode_Comment>(Node))
 		{
