Removed ensure if expression is not found in MI editor, was spamming before, effectively causing a crash

diff --git a/Engine/Source/Editor/MaterialEditor/Private/MaterialInstanceEditor.cpp b/Engine/Source/Editor/MaterialEditor/Private/MaterialInstanceEditor.cpp
--- a/Engine/Source/Editor/MaterialEditor/Private/MaterialInstanceEditor.cpp
+++ b/Engine/Source/Editor/MaterialEditor/Private/MaterialInstanceEditor.cpp
@@ -795,9 +795,6 @@ void FMaterialInstanceEditor::DrawSamplerWarningStrings(FCanvas* Canvas, int32&
 						{
 							EMaterialSamplerType SamplerType = UMaterialExpressionTextureBase::GetSamplerTypeForTexture( Texture );
 							UMaterialExpressionTextureSampleParameter* Expression = BaseMaterial->FindExpressionByGUID<UMaterialExpressionTextureSampleParameter>( TextureParameterValue->ExpressionId );
-							
-							// We want a crash report if that happens
-							ensure(Expression);
 
 							if ( Expression && Expression->SamplerType != SamplerType )
 							{
