UE-11390 - Parameters not appearing in Material instances

This is now "fixed". This was only ever a problem for parameters that fed into a function input's Preview pin.
They will never have an impact on the final material's color and so are useless. But as it's such a minor edge case, it's not worth the effort of the real solution which would be to disallow connecting parameters to the preview in in the first place. Or something along those lines.

diff --git a/Engine/Source/Editor/MaterialEditor/Private/MaterialEditorUtilities.cpp b/Engine/Source/Editor/MaterialEditor/Private/MaterialEditorUtilities.cpp
--- a/Engine/Source/Editor/MaterialEditor/Private/MaterialEditorUtilities.cpp
+++ b/Engine/Source/Editor/MaterialEditor/Private/MaterialEditorUtilities.cpp
@@ -500,6 +500,8 @@ void FMaterialEditorUtilities::GetVisibleMaterialParametersFromExpression(
 	}
 	else if (FunctionInputExpression)
 	{
+		GetVisibleMaterialParametersFromExpression(FMaterialExpressionKey(FunctionInputExpression->Preview.Expression, FunctionInputExpression->Preview.OutputIndex), MaterialInstance, VisibleExpressions, FunctionStack);
+		
 		FGetVisibleMaterialParametersFunctionState* FunctionState = FunctionStack.Pop();
 
 		const FFunctionExpressionInput* MatchingInput = FindInputById(FunctionInputExpression, FunctionState->FunctionCall->FunctionInputs);
