Scalar parameter SliderMin/Max affects DefaultValue editing as well as MI parameter editing

diff --git a/Engine/Source/Editor/MaterialEditor/Private/MaterialEditorDetailCustomization.cpp b/Engine/Source/Editor/MaterialEditor/Private/MaterialEditorDetailCustomization.cpp
--- a/Engine/Source/Editor/MaterialEditor/Private/MaterialEditorDetailCustomization.cpp
+++ b/Engine/Source/Editor/MaterialEditor/Private/MaterialEditorDetailCustomization.cpp
@@ -26,6 +26,40 @@ void FMaterialExpressionParameterDetails::CustomizeDetails( IDetailLayoutBuilder
 	FName DefaultCategory = NAME_None;
 	IDetailCategoryBuilder& Category = DetailLayout.EditCategory( DefaultCategory );
 
+	TArray< TWeakObjectPtr<UObject> > Objects;
+	DetailLayout.GetObjectsBeingCustomized(Objects);
+
+	if (Objects.Num() > 0)
+	{
+		UMaterialExpressionScalarParameter* ScalarParameter = Cast<UMaterialExpressionScalarParameter>(Objects[0].Get());
+
+		if (ScalarParameter)
+		{
+			// Store these for OnSliderMinMaxEdited
+			ScalarParameterObject = Objects[0];
+			DefaultValueHandle = DetailLayout.GetProperty("DefaultValue", UMaterialExpressionScalarParameter::StaticClass());
+
+			TSharedPtr<IPropertyHandle> SliderMinHandle = DetailLayout.GetProperty("SliderMin", UMaterialExpressionScalarParameter::StaticClass());
+
+			if (SliderMinHandle.IsValid())
+			{
+				// Setup a callback when SliderMin changes to update the DefaultValue slider
+				FSimpleDelegate OnSliderMinMaxEditedDelegate = FSimpleDelegate::CreateSP(this, &FMaterialExpressionParameterDetails::OnSliderMinMaxEdited);
+				SliderMinHandle->SetOnPropertyValueChanged(OnSliderMinMaxEditedDelegate);
+			}
+
+			TSharedPtr<IPropertyHandle> SliderMaxHandle = DetailLayout.GetProperty("SliderMax", UMaterialExpressionScalarParameter::StaticClass());
+
+			if (SliderMaxHandle.IsValid())
+			{
+				FSimpleDelegate OnSliderMinMaxEditedDelegate = FSimpleDelegate::CreateSP(this, &FMaterialExpressionParameterDetails::OnSliderMinMaxEdited);
+				SliderMaxHandle->SetOnPropertyValueChanged(OnSliderMinMaxEditedDelegate);
+			}
+
+			OnSliderMinMaxEdited();
+		}
+	}
+	
 	Category.AddProperty("ParameterName");
 	
 	// Get a handle to the property we are about to edit
@@ -128,7 +162,26 @@ FText FMaterialExpressionParameterDetails::OnGetText() const
 	return FText::FromString(NewString);
 }
 
+void FMaterialExpressionParameterDetails::OnSliderMinMaxEdited()
+{
+	UMaterialExpressionScalarParameter* ScalarParameter = Cast<UMaterialExpressionScalarParameter>(ScalarParameterObject.Get());
 
+	if (ScalarParameter && DefaultValueHandle.IsValid())
+	{
+		if (ScalarParameter->SliderMax > ScalarParameter->SliderMin)
+		{
+			// Update the values that SPropertyEditorNumeric reads
+			// Unfortuantly there is no way to recreate the widget to actually update the UI with these new values
+			DefaultValueHandle->SetInstanceMetaData("UIMin",FString::Printf(TEXT("%f"), ScalarParameter->SliderMin));
+			DefaultValueHandle->SetInstanceMetaData("UIMax",FString::Printf(TEXT("%f"), ScalarParameter->SliderMax));
+		}
+		else
+		{
+			DefaultValueHandle->SetInstanceMetaData("UIMin", TEXT(""));
+			DefaultValueHandle->SetInstanceMetaData("UIMax", TEXT(""));
+		}
+	}
+}
 
 TSharedRef<IDetailCustomization> FMaterialExpressionCollectionParameterDetails::MakeInstance()
 {
diff --git a/Engine/Source/Editor/MaterialEditor/Private/MaterialEditorDetailCustomization.h b/Engine/Source/Editor/MaterialEditor/Private/MaterialEditorDetailCustomization.h
--- a/Engine/Source/Editor/MaterialEditor/Private/MaterialEditorDetailCustomization.h
+++ b/Engine/Source/Editor/MaterialEditor/Private/MaterialEditorDetailCustomization.h
@@ -28,6 +28,7 @@ private:
 	void OnTextCommitted( const FText& InText, ETextCommit::Type CommitInfo);
 	FString OnGetString() const;
 	FText OnGetText() const;
+	void OnSliderMinMaxEdited();
 
 private:
 	/** The property handle to the groups */
@@ -43,6 +44,9 @@ private:
 
 	/** A list of all group names to choose from */
 	TArray<TSharedPtr<FString>> GroupsSource;
+
+	TWeakObjectPtr<UObject> ScalarParameterObject;
+	TSharedPtr<IPropertyHandle> DefaultValueHandle;
 };
 
 /** 
