Fix non-unity build error in MaterialEditorDetailCustomization.cpp

diff --git a/Engine/Source/Editor/MaterialEditor/Private/MaterialEditorDetailCustomization.cpp b/Engine/Source/Editor/MaterialEditor/Private/MaterialEditorDetailCustomization.cpp
--- a/Engine/Source/Editor/MaterialEditor/Private/MaterialEditorDetailCustomization.cpp
+++ b/Engine/Source/Editor/MaterialEditor/Private/MaterialEditorDetailCustomization.cpp
@@ -5,6 +5,7 @@
 #include "MaterialEditor.h"
 #include "MaterialEditorDetailCustomization.h"
 #include "Materials/MaterialParameterCollection.h"
+#include "Materials/MaterialExpressionScalarParameter.h"
 
 #define LOCTEXT_NAMESPACE "MaterialEditor"
 
