Changed bOnlyRefreshOnOk paramater so double clicking the color parameter in the material node also changes the in-world instance color paramater.

UE-17439

diff --git a/Engine/Source/Editor/MaterialEditor/Private/MaterialEditor.cpp b/Engine/Source/Editor/MaterialEditor/Private/MaterialEditor.cpp
--- a/Engine/Source/Editor/MaterialEditor/Private/MaterialEditor.cpp
+++ b/Engine/Source/Editor/MaterialEditor/Private/MaterialEditor.cpp
@@ -4109,7 +4109,7 @@ void FMaterialEditor::OnNodeDoubleClicked(class UEdGraphNode* Node)
 			FColorPickerArgs PickerArgs;
 			PickerArgs.ParentWidget = GraphEditor;//AsShared();
 			PickerArgs.bUseAlpha = Constant4Expression != NULL || VectorExpression != NULL;
-			PickerArgs.bOnlyRefreshOnOk = true;
+			PickerArgs.bOnlyRefreshOnOk = false;
 			PickerArgs.bExpandAdvancedSection = true;
 			PickerArgs.DisplayGamma = TAttribute<float>::Create( TAttribute<float>::FGetter::CreateUObject(GEngine, &UEngine::GetDisplayGamma) );
 			PickerArgs.ColorChannelsArray = &Channels;
