Fixed issue where a group of material nodes containing a MaterialFunctionCall node lost their connections when copy+pasted.

diff --git a/Engine/Source/Editor/MaterialEditor/Private/MaterialEditor.cpp b/Engine/Source/Editor/MaterialEditor/Private/MaterialEditor.cpp
--- a/Engine/Source/Editor/MaterialEditor/Private/MaterialEditor.cpp
+++ b/Engine/Source/Editor/MaterialEditor/Private/MaterialEditor.cpp
@@ -3381,7 +3381,11 @@ void FMaterialEditor::PasteNodesHere(const FVector2D& Location)
 			UMaterialExpressionMaterialFunctionCall* FunctionCall = Cast<UMaterialExpressionMaterialFunctionCall>( NewExpression );
 			if( FunctionCall )
 			{
-				FunctionCall->UpdateFromFunctionResource();
+				// When pasting new nodes, we don't want to break all node links as this information is used by UpdateMaterialAfterGraphChange() below,
+				// to recreate all the connections in the pasted group.
+				// Just update the function input/outputs here.
+				const bool bRecreateAndLinkNode = false;
+				FunctionCall->UpdateFromFunctionResource(bRecreateAndLinkNode);
 			}
 		}
 		else if (UMaterialGraphNode_Comment* CommentNode = Cast<UMaterialGraphNode_Comment>(Node))
diff --git a/Engine/Source/Runtime/Engine/Classes/Materials/MaterialExpressionMaterialFunctionCall.h b/Engine/Source/Runtime/Engine/Classes/Materials/MaterialExpressionMaterialFunctionCall.h
--- a/Engine/Source/Runtime/Engine/Classes/Materials/MaterialExpressionMaterialFunctionCall.h
+++ b/Engine/Source/Runtime/Engine/Classes/Materials/MaterialExpressionMaterialFunctionCall.h
@@ -127,7 +127,7 @@ class UMaterialExpressionMaterialFunctionCall : public UMaterialExpression
 	 * Update FunctionInputs and FunctionOutputs from the MaterialFunction.  
 	 * This must be called to keep the inputs and outputs up to date with the function being used. 
 	 */
-	ENGINE_API void UpdateFromFunctionResource();
+	ENGINE_API void UpdateFromFunctionResource(bool bRecreateAndLinkNode = true);
 
 private:
 	
diff --git a/Engine/Source/Runtime/Engine/Private/Materials/MaterialExpressions.cpp b/Engine/Source/Runtime/Engine/Private/Materials/MaterialExpressions.cpp
--- a/Engine/Source/Runtime/Engine/Private/Materials/MaterialExpressions.cpp
+++ b/Engine/Source/Runtime/Engine/Private/Materials/MaterialExpressions.cpp
@@ -7190,7 +7190,7 @@ bool UMaterialExpressionMaterialFunctionCall::SetMaterialFunction(
 }
 
 
-void UMaterialExpressionMaterialFunctionCall::UpdateFromFunctionResource()
+void UMaterialExpressionMaterialFunctionCall::UpdateFromFunctionResource(bool bRecreateAndLinkNode)
 {
 	TArray<FFunctionExpressionInput> OriginalInputs = FunctionInputs;
 	TArray<FFunctionExpressionOutput> OriginalOutputs = FunctionOutputs;
@@ -7253,7 +7253,7 @@ void UMaterialExpressionMaterialFunctionCall::UpdateFromFunctionResource()
 	}
 
 #if WITH_EDITOR
-	if (GraphNode)
+	if (GraphNode && bRecreateAndLinkNode)
 	{
 		// Check whether number of input/outputs or transient pointers have changed
 		bool bUpdatedFromFunction = false;
