Fixed issue where closing the editor with a dirty material open in the Material Editor would lead to multiple prompts.

diff --git a/Engine/Source/Editor/MaterialEditor/Private/MaterialEditor.cpp b/Engine/Source/Editor/MaterialEditor/Private/MaterialEditor.cpp
--- a/Engine/Source/Editor/MaterialEditor/Private/MaterialEditor.cpp
+++ b/Engine/Source/Editor/MaterialEditor/Private/MaterialEditor.cpp
@@ -908,6 +908,7 @@ bool FMaterialEditor::OnRequestClose()
 				
 		case EAppReturnType::No:
 			// exit
+			bMaterialDirty = false;
 			break;
 				
 		case EAppReturnType::Cancel:
