Fixed cases where material editor and function instance editors would fail to update when a material function was applied.

diff --git a/Engine/Source/Editor/MaterialEditor/Private/MaterialEditingLibrary.cpp b/Engine/Source/Editor/MaterialEditor/Private/MaterialEditingLibrary.cpp
--- a/Engine/Source/Editor/MaterialEditor/Private/MaterialEditingLibrary.cpp
+++ b/Engine/Source/Editor/MaterialEditor/Private/MaterialEditingLibrary.cpp
@@ -191,7 +191,7 @@ void UMaterialEditingLibrary::RebuildMaterialInstanceEditors(UMaterialFunction*
 		if (FunctionInstance)
 		{
 			// Update function instances that are children of this material function	
-			if (BaseFunction && BaseFunction->ParentFunction == FunctionInstance->GetBaseFunction())
+			if (BaseFunction && BaseFunction == FunctionInstance->GetBaseFunction())
 			{
 				IAssetEditorInstance* EditorInstance = AssetEditorManager.FindEditorForAsset(EditedAsset, false);
 				if (EditorInstance)
@@ -611,7 +611,7 @@ void UMaterialEditingLibrary::UpdateMaterialFunction(UMaterialFunction* Material
 						for (const FMaterialFunctionInfo& FunctionInfo : CurrentMaterial->MaterialFunctionInfos)
 						{
 							UMaterialFunctionInstance* FunctionInstance = Cast<UMaterialFunctionInstance>(FunctionInfo.Function);
-							if (FunctionInfo.Function == MaterialFunction->ParentFunction || (FunctionInstance && FunctionInstance->GetBaseFunction() == MaterialFunction->ParentFunction))
+							if (FunctionInfo.Function == MaterialFunction || (FunctionInstance && FunctionInstance->GetBaseFunction() == MaterialFunction))
 							{
 								bRecompile = true;
 								break;
