Fixed material editor crash if reloading a texture used by the material

diff --git a/Engine/Source/Editor/MaterialEditor/Private/MaterialEditor.cpp b/Engine/Source/Editor/MaterialEditor/Private/MaterialEditor.cpp
--- a/Engine/Source/Editor/MaterialEditor/Private/MaterialEditor.cpp
+++ b/Engine/Source/Editor/MaterialEditor/Private/MaterialEditor.cpp
@@ -2082,6 +2082,10 @@ void FMaterialEditor::AddReferencedObjects( FReferenceCollector& Collector )
 	Collector.AddReferencedObject( ExpressionPreviewMaterial );
 	Collector.AddReferencedObject( EmptyMaterial );
 	Collector.AddReferencedObject( MaterialEditorInstance );
+	for (FMatExpressionPreview& ExpressionPreview : ExpressionPreviews)
+	{
+		ExpressionPreview.AddReferencedObjects(Collector);
+	}
 }
 
 void FMaterialEditor::BindCommands()
