Prevent GC of material editor active preview expression

diff --git a/Engine/Source/Editor/MaterialEditor/Private/MaterialEditor.cpp b/Engine/Source/Editor/MaterialEditor/Private/MaterialEditor.cpp
--- a/Engine/Source/Editor/MaterialEditor/Private/MaterialEditor.cpp
+++ b/Engine/Source/Editor/MaterialEditor/Private/MaterialEditor.cpp
@@ -2075,13 +2075,14 @@ void FMaterialEditor::UpdateGraphNodeStates()
 
 void FMaterialEditor::AddReferencedObjects( FReferenceCollector& Collector )
 {
-	Collector.AddReferencedObject( EditorOptions );
-	Collector.AddReferencedObject( Material );
-	Collector.AddReferencedObject( OriginalMaterial );
-	Collector.AddReferencedObject( MaterialFunction );
-	Collector.AddReferencedObject( ExpressionPreviewMaterial );
-	Collector.AddReferencedObject( EmptyMaterial );
-	Collector.AddReferencedObject( MaterialEditorInstance );
+	Collector.AddReferencedObject(EditorOptions);
+	Collector.AddReferencedObject(Material);
+	Collector.AddReferencedObject(OriginalMaterial);
+	Collector.AddReferencedObject(MaterialFunction);
+	Collector.AddReferencedObject(ExpressionPreviewMaterial);
+	Collector.AddReferencedObject(EmptyMaterial);
+	Collector.AddReferencedObject(MaterialEditorInstance);
+	Collector.AddReferencedObject(PreviewExpression);
 	for (FMatExpressionPreview& ExpressionPreview : ExpressionPreviews)
 	{
 		ExpressionPreview.AddReferencedObjects(Collector);
@@ -3477,8 +3478,7 @@ void FMaterialEditor::DeleteNodes(const TArray<UEdGraphNode*>& NodesToDelete)
 			if( bPreviewExpressionDeleted )
 			{
 				// The preview expression was deleted.  Null out our reference to it and reset to the normal preview material
-				PreviewExpression = NULL;
-				SetPreviewMaterial( Material );
+				SetPreviewExpression(nullptr);
 			}
 			RegenerateCodeView();
 		}
