Adding additional conditionals to GetParameterExpressionDescription

diff --git a/Engine/Source/Editor/MaterialEditor/Private/MaterialEditorDetailCustomization.cpp b/Engine/Source/Editor/MaterialEditor/Private/MaterialEditorDetailCustomization.cpp
--- a/Engine/Source/Editor/MaterialEditor/Private/MaterialEditorDetailCustomization.cpp
+++ b/Engine/Source/Editor/MaterialEditor/Private/MaterialEditorDetailCustomization.cpp
@@ -414,17 +414,19 @@ FString FMaterialEditorParameterDetails::GetCurvePath(UDEditorScalarParameterVal
 
 FText FMaterialEditorParameterDetails::GetParameterExpressionDescription(UDEditorParameterValue* Parameter) const
 {
-	UMaterial* BaseMaterial = MaterialEditorInstance->OriginalMaterial->GetMaterial();
-	if (BaseMaterial)
+	if (MaterialEditorInstance && MaterialEditorInstance->OriginalMaterial)
 	{
-		UMaterialExpression* MaterialExpression = BaseMaterial->FindExpressionByGUID<UMaterialExpression>(Parameter->ExpressionId);
-
-		if (MaterialExpression)
+		UMaterial* BaseMaterial = MaterialEditorInstance->OriginalMaterial->GetMaterial();
+		if (BaseMaterial)
 		{
-			return FText::FromString(MaterialExpression->Desc);
+			UMaterialExpression* MaterialExpression = BaseMaterial->FindExpressionByGUID<UMaterialExpression>(Parameter->ExpressionId);
+
+			if (MaterialExpression)
+			{
+				return FText::FromString(MaterialExpression->Desc);
+			}
 		}
 	}
-
 	return FText::GetEmpty();
 }
 
