When selecting a node while the details panel is behind a different panel in the same dock tab, the details panel is brought forward

diff --git a/Engine/Source/Editor/MaterialEditor/Private/MaterialEditor.cpp b/Engine/Source/Editor/MaterialEditor/Private/MaterialEditor.cpp
--- a/Engine/Source/Editor/MaterialEditor/Private/MaterialEditor.cpp
+++ b/Engine/Source/Editor/MaterialEditor/Private/MaterialEditor.cpp
@@ -3096,7 +3096,7 @@ TSharedRef<SDockTab> FMaterialEditor::SpawnTab_GraphCanvas(const FSpawnTabArgs&
 
 TSharedRef<SDockTab> FMaterialEditor::SpawnTab_MaterialProperties(const FSpawnTabArgs& Args)
 {
-	TSharedRef<SDockTab> SpawnedTab = SNew(SDockTab)
+	SpawnedDetailsTab = SNew(SDockTab)
 		.Icon( FEditorStyle::GetBrush("LevelEditor.Tabs.Details") )
 		.Label( LOCTEXT("MaterialDetailsTitle", "Details") )
 		[
@@ -3109,7 +3109,7 @@ TSharedRef<SDockTab> FMaterialEditor::SpawnTab_MaterialProperties(const FSpawnTa
 		GraphEditor->ClearSelectionSet();
 	}
 
-	return SpawnedTab;
+	return SpawnedDetailsTab.ToSharedRef();
 }
 
 TSharedRef<SDockTab> FMaterialEditor::SpawnTab_Palette(const FSpawnTabArgs& Args)
@@ -4427,6 +4427,7 @@ void FMaterialEditor::OnSelectedNodesChanged(const TSet<class UObject*>& NewSele
 	}
 
 	GetDetailView()->SetObjects( SelectedObjects, true );
+	FocusDetailsPanel();
 }
 
 void FMaterialEditor::OnNodeDoubleClicked(class UEdGraphNode* Node)
@@ -4607,4 +4608,12 @@ void FMaterialEditor::NotifyExternalMaterialChange()
 	MaterialStatsManager->SignalMaterialChanged();
 }
 
+void FMaterialEditor::FocusDetailsPanel()
+{
+	if (SpawnedDetailsTab)
+	{
+		SpawnedDetailsTab->DrawAttention();
+	}
+}
+
 #undef LOCTEXT_NAMESPACE
diff --git a/Engine/Source/Editor/MaterialEditor/Private/MaterialEditor.h b/Engine/Source/Editor/MaterialEditor/Private/MaterialEditor.h
--- a/Engine/Source/Editor/MaterialEditor/Private/MaterialEditor.h
+++ b/Engine/Source/Editor/MaterialEditor/Private/MaterialEditor.h
@@ -381,6 +381,9 @@ public:
 	/** call this to notify the editor that the edited material changed from outside */
 	virtual void NotifyExternalMaterialChange() override;
 
+	/** Called to bring focus to the details panel */
+	void FocusDetailsPanel();
+
 public:
 	/** Set to true when modifications have been made to the material */
 	bool bMaterialDirty;
@@ -778,4 +781,7 @@ private:
 
 	/** Object used as material statistics manager */
 	TSharedPtr<class FMaterialStats> MaterialStatsManager;
+
+	/** Tab that holds the details panel */
+	TSharedPtr<SDockTab> SpawnedDetailsTab;
 };
