Fix for new texture sample nodes not taking the selected texture from the content browser as the starting value.

diff --git a/Engine/Source/Editor/MaterialEditor/Private/MaterialEditor.cpp b/Engine/Source/Editor/MaterialEditor/Private/MaterialEditor.cpp
--- a/Engine/Source/Editor/MaterialEditor/Private/MaterialEditor.cpp
+++ b/Engine/Source/Editor/MaterialEditor/Private/MaterialEditor.cpp
@@ -3292,6 +3292,8 @@ UMaterialExpression* FMaterialEditor::CreateNewMaterialExpression(UClass* NewExp
 		UObject* SelectedAsset = nullptr;
 		if (bAutoAssignResource)
 		{
+			// Load selected assets
+			FEditorDelegates::LoadSelectedAssetsIfNeeded.Broadcast();
 			SelectedAsset = GEditor->GetSelectedObjects()->GetTop<UObject>();
 		}
 
@@ -4610,7 +4612,7 @@ void FMaterialEditor::NotifyExternalMaterialChange()
 
 void FMaterialEditor::FocusDetailsPanel()
 {
-	if (SpawnedDetailsTab)
+	if (SpawnedDetailsTab && !SpawnedDetailsTab->IsForeground())
 	{
 		SpawnedDetailsTab->DrawAttention();
 	}
