Fixed crashes when overriding material parameters. Every texture or parameter override was needlessly recreating uniform expression uniform buffer, which may be referenced by cached mesh draw commands. Now it is done only when required, so only when uniform buffer layout will change. This also allowed to back-out CL 4751942 "Fix for UE-68218. Crash altering material's scalar values" and similar fixes.

diff --git a/Engine/Source/Editor/MaterialEditor/Private/MaterialEditor.cpp b/Engine/Source/Editor/MaterialEditor/Private/MaterialEditor.cpp
--- a/Engine/Source/Editor/MaterialEditor/Private/MaterialEditor.cpp
+++ b/Engine/Source/Editor/MaterialEditor/Private/MaterialEditor.cpp
@@ -3098,9 +3098,6 @@ void FMaterialEditor::SetVectorParameterDefaultOnDependentMaterials(FName Parame
 		MaterialsToOverride.Add(OriginalMaterial);
 	}
 
-
-	FMaterialUpdateContext UpdateContext(FMaterialUpdateContext::EOptions::SyncWithRenderingThread);
-
 	const ERHIFeatureLevel::Type FeatureLevel = GEditor->GetEditorWorldContext().World()->FeatureLevel;
 
 	for (int32 MaterialIndex = 0; MaterialIndex < MaterialsToOverride.Num(); MaterialIndex++)
@@ -3108,7 +3105,6 @@ void FMaterialEditor::SetVectorParameterDefaultOnDependentMaterials(FName Parame
 		UMaterial* CurrentMaterial = MaterialsToOverride[MaterialIndex];
 
 		CurrentMaterial->OverrideVectorParameterDefault(ParameterName, Value, bOverride, FeatureLevel);
-		UpdateContext.AddMaterial(CurrentMaterial);
 	}
 
 	// Update MI's that reference any of the materials affected
@@ -3178,9 +3174,6 @@ void FMaterialEditor::SetScalarParameterDefaultOnDependentMaterials(FName Parame
 		MaterialsToOverride.Add(OriginalMaterial);
 	}
 
-
-	FMaterialUpdateContext UpdateContext(FMaterialUpdateContext::EOptions::SyncWithRenderingThread);
-
 	const ERHIFeatureLevel::Type FeatureLevel = GEditor->GetEditorWorldContext().World()->FeatureLevel;
 
 	for (int32 MaterialIndex = 0; MaterialIndex < MaterialsToOverride.Num(); MaterialIndex++)
@@ -3188,7 +3181,6 @@ void FMaterialEditor::SetScalarParameterDefaultOnDependentMaterials(FName Parame
 		UMaterial* CurrentMaterial = MaterialsToOverride[MaterialIndex];
 
 		CurrentMaterial->OverrideScalarParameterDefault(ParameterName, Value, bOverride, FeatureLevel);
-		UpdateContext.AddMaterial(CurrentMaterial);
 	}
 
 	// Update MI's that reference any of the materials affected
diff --git a/Engine/Source/Runtime/Engine/Classes/Materials/Material.h b/Engine/Source/Runtime/Engine/Classes/Materials/Material.h
--- a/Engine/Source/Runtime/Engine/Classes/Materials/Material.h
+++ b/Engine/Source/Runtime/Engine/Classes/Materials/Material.h
@@ -977,7 +977,7 @@ public:
 		TArray<FName>* OutTextureParamNames, struct FStaticParameterSet* InStaticParameterSet,
 		ERHIFeatureLevel::Type InFeatureLevel, EMaterialQualityLevel::Type InQuality) override;
 #endif
-	ENGINE_API virtual void RecacheUniformExpressions() const override;
+	ENGINE_API virtual void RecacheUniformExpressions(bool bRecreateUniformBuffer) const override;
 
 	ENGINE_API virtual float GetOpacityMaskClipValue() const override;
 	ENGINE_API virtual bool GetCastDynamicShadowAsMasked() const override;
diff --git a/Engine/Source/Runtime/Engine/Classes/Materials/MaterialInstance.h b/Engine/Source/Runtime/Engine/Classes/Materials/MaterialInstance.h
--- a/Engine/Source/Runtime/Engine/Classes/Materials/MaterialInstance.h
+++ b/Engine/Source/Runtime/Engine/Classes/Materials/MaterialInstance.h
@@ -450,7 +450,7 @@ public:
 		TArray<FName>* OutTextureParamNames, struct FStaticParameterSet* InStaticParameterSet,
 		ERHIFeatureLevel::Type InFeatureLevel, EMaterialQualityLevel::Type InQuality) override;
 #endif
-	virtual ENGINE_API void RecacheUniformExpressions() const override;
+	virtual ENGINE_API void RecacheUniformExpressions(bool bRecreateUniformBuffer) const override;
 	virtual ENGINE_API bool GetRefractionSettings(float& OutBiasValue) const override;
 
 #if WITH_EDITOR
diff --git a/Engine/Source/Runtime/Engine/Classes/Materials/MaterialInterface.h b/Engine/Source/Runtime/Engine/Classes/Materials/MaterialInterface.h
--- a/Engine/Source/Runtime/Engine/Classes/Materials/MaterialInterface.h
+++ b/Engine/Source/Runtime/Engine/Classes/Materials/MaterialInterface.h
@@ -729,13 +729,21 @@ public:
 
 	/**
 	 * Re-caches uniform expressions for all material interfaces
+	 * Set bRecreateUniformBuffer to true if uniform buffer layout will change (e.g. FMaterial is being recompiled).
+	 * In that case calling needs to use FMaterialUpdateContext to recreate the rendering state of primitives using this material.
+	 * 
+	 * @param bRecreateUniformBuffer - true forces uniform buffer recreation.
 	 */
-	ENGINE_API static void RecacheAllMaterialUniformExpressions();
+	ENGINE_API static void RecacheAllMaterialUniformExpressions(bool bRecreateUniformBuffer);
 
 	/**
 	 * Re-caches uniform expressions for this material interface                   
+	 * Set bRecreateUniformBuffer to true if uniform buffer layout will change (e.g. FMaterial is being recompiled).
+	 * In that case calling needs to use FMaterialUpdateContext to recreate the rendering state of primitives using this material.
+	 *
+	 * @param bRecreateUniformBuffer - true forces uniform buffer recreation.
 	 */
-	virtual void RecacheUniformExpressions() const {}
+	virtual void RecacheUniformExpressions(bool bRecreateUniformBuffer) const {}
 
 #if WITH_EDITOR
 	/** Clears the shader cache and recompiles the shader for rendering. */
diff --git a/Engine/Source/Runtime/Engine/Private/Materials/Material.cpp b/Engine/Source/Runtime/Engine/Private/Materials/Material.cpp
--- a/Engine/Source/Runtime/Engine/Private/Materials/Material.cpp
+++ b/Engine/Source/Runtime/Engine/Private/Materials/Material.cpp
@@ -1145,7 +1145,7 @@ void UMaterial::OverrideTexture(const UTexture* InTextureToOverride, UTexture* O
 
 	if (bShouldRecacheMaterialExpressions)
 	{
-		RecacheUniformExpressions();
+		RecacheUniformExpressions(false);
 		RecacheMaterialInstanceUniformExpressions(this);
 	}
 #endif // #if WITH_EDITOR
@@ -1176,7 +1176,7 @@ void UMaterial::OverrideVectorParameterDefault(const FMaterialParameterInfo& Par
 
 	if (bShouldRecacheMaterialExpressions)
 	{
-		RecacheUniformExpressions();
+		RecacheUniformExpressions(false);
 		RecacheMaterialInstanceUniformExpressions(this);
 	}
 #endif // #if WITH_EDITOR
@@ -1208,13 +1208,13 @@ void UMaterial::OverrideScalarParameterDefault(const FMaterialParameterInfo& Par
 
 	if (bShouldRecacheMaterialExpressions)
 	{
-		RecacheUniformExpressions();
+		RecacheUniformExpressions(false);
 		RecacheMaterialInstanceUniformExpressions(this);
 	}
 #endif // #if WITH_EDITOR
 }
 
-void UMaterial::RecacheUniformExpressions() const
+void UMaterial::RecacheUniformExpressions(bool bRecreateUniformBuffer) const
 {
 	bool bUsingNewLoader = EVENT_DRIVEN_ASYNC_LOAD_ACTIVE_AT_RUNTIME && GEventDrivenLoaderEnabled;
 
@@ -1226,7 +1226,7 @@ void UMaterial::RecacheUniformExpressions() const
 
 	if (DefaultMaterialInstance)
 	{
-		DefaultMaterialInstance->CacheUniformExpressions_GameThread(true);
+		DefaultMaterialInstance->CacheUniformExpressions_GameThread(bRecreateUniformBuffer);
 	}
 }
 
@@ -3125,7 +3125,7 @@ void UMaterial::CacheResourceShadersForRendering(bool bRegenerateId)
 			}
 		}
 
-		RecacheUniformExpressions();
+		RecacheUniformExpressions(true);
 	}
 }
 
diff --git a/Engine/Source/Runtime/Engine/Private/Materials/MaterialInstance.cpp b/Engine/Source/Runtime/Engine/Private/Materials/MaterialInstance.cpp
--- a/Engine/Source/Runtime/Engine/Private/Materials/MaterialInstance.cpp
+++ b/Engine/Source/Runtime/Engine/Private/Materials/MaterialInstance.cpp
@@ -1254,7 +1254,7 @@ void UMaterialInstance::OverrideTexture(const UTexture* InTextureToOverride, UTe
 
 	if (bShouldRecacheMaterialExpressions)
 	{
-		RecacheUniformExpressions();
+		RecacheUniformExpressions(false);
 		RecacheMaterialInstanceUniformExpressions(this);
 	}
 #endif // #if WITH_EDITOR
@@ -1289,7 +1289,7 @@ void UMaterialInstance::OverrideVectorParameterDefault(const FMaterialParameterI
 
 	if (bShouldRecacheMaterialExpressions)
 	{
-		RecacheUniformExpressions();
+		RecacheUniformExpressions(false);
 		RecacheMaterialInstanceUniformExpressions(this);
 	}
 #endif // #if WITH_EDITOR
@@ -1324,7 +1324,7 @@ void UMaterialInstance::OverrideScalarParameterDefault(const FMaterialParameterI
 
 	if (bShouldRecacheMaterialExpressions)
 	{
-		RecacheUniformExpressions();
+		RecacheUniformExpressions(false);
 		RecacheMaterialInstanceUniformExpressions(this);
 	}
 #endif // #if WITH_EDITOR
@@ -3582,7 +3582,7 @@ void UMaterialInstance::UpdateParameterNames()
 }
 #endif
 
-void UMaterialInstance::RecacheUniformExpressions() const
+void UMaterialInstance::RecacheUniformExpressions(bool bRecreateUniformBuffer) const
 {	
 	CacheMaterialInstanceUniformExpressions(this);
 }
diff --git a/Engine/Source/Runtime/Engine/Private/Materials/MaterialInterface.cpp b/Engine/Source/Runtime/Engine/Private/Materials/MaterialInterface.cpp
--- a/Engine/Source/Runtime/Engine/Private/Materials/MaterialInterface.cpp
+++ b/Engine/Source/Runtime/Engine/Private/Materials/MaterialInterface.cpp
@@ -210,12 +210,12 @@ void UMaterialInterface::SetForceMipLevelsToBeResident( bool OverrideForceMiplev
 	}
 }
 
-void UMaterialInterface::RecacheAllMaterialUniformExpressions()
+void UMaterialInterface::RecacheAllMaterialUniformExpressions(bool bRecreateUniformBuffer)
 {
 	// For each interface, reacache its uniform parameters
 	for( TObjectIterator<UMaterialInterface> MaterialIt; MaterialIt; ++MaterialIt )
 	{
-		MaterialIt->RecacheUniformExpressions();
+		MaterialIt->RecacheUniformExpressions(bRecreateUniformBuffer);
 	}
 }
 
diff --git a/Engine/Source/Runtime/Engine/Private/ShaderCompiler/ShaderCompiler.cpp b/Engine/Source/Runtime/Engine/Private/ShaderCompiler/ShaderCompiler.cpp
--- a/Engine/Source/Runtime/Engine/Private/ShaderCompiler/ShaderCompiler.cpp
+++ b/Engine/Source/Runtime/Engine/Private/ShaderCompiler/ShaderCompiler.cpp
@@ -3478,7 +3478,7 @@ public:
 			}
 
 			// fixup uniform expressions
-			UMaterialInterface::RecacheAllMaterialUniformExpressions();
+			UMaterialInterface::RecacheAllMaterialUniformExpressions(true);
 
 			// Need to recache all cached mesh draw commands, as they store pointers to material uniform buffers which we just invalidated.
 			GetRendererModule().UpdateStaticDrawLists();
diff --git a/Engine/Source/Runtime/Engine/Private/UnrealEngine.cpp b/Engine/Source/Runtime/Engine/Private/UnrealEngine.cpp
--- a/Engine/Source/Runtime/Engine/Private/UnrealEngine.cpp
+++ b/Engine/Source/Runtime/Engine/Private/UnrealEngine.cpp
@@ -764,10 +764,7 @@ void RefreshSamplerStatesCallback()
 			);
 		}
 
-		UMaterialInterface::RecacheAllMaterialUniformExpressions();
-
-		// Need to recache all cached mesh draw commands, as they store pointers to material uniform buffers which we just invalidated.
-		GetRendererModule().UpdateStaticDrawLists();
+		UMaterialInterface::RecacheAllMaterialUniformExpressions(false);
 	}
 }
 
