Fixed crash in case of material instance with no parent looking for parent values

diff --git a/Engine/Source/Editor/MaterialEditor/Private/MaterialEditorInstanceDetailCustomization.cpp b/Engine/Source/Editor/MaterialEditor/Private/MaterialEditorInstanceDetailCustomization.cpp
--- a/Engine/Source/Editor/MaterialEditor/Private/MaterialEditorInstanceDetailCustomization.cpp
+++ b/Engine/Source/Editor/MaterialEditor/Private/MaterialEditorInstanceDetailCustomization.cpp
@@ -761,12 +761,15 @@ void FMaterialInstanceParameterDetails::CreateLightmassOverrideWidgets(IDetailLa
 	TSharedPtr<IPropertyHandle> DiffuseBoostProperty = LightmassSettings->GetChildHandle("DiffuseBoost");
 	TSharedPtr<IPropertyHandle> ExportResolutionScaleProperty = LightmassSettings->GetChildHandle("ExportResolutionScale");
 
-
+	
 	FIsResetToDefaultVisible IsCastShadowAsMaskedPropertyResetVisible = FIsResetToDefaultVisible::CreateLambda([this](TSharedPtr<IPropertyHandle> InHandle) {
-		return MaterialEditorInstance->LightmassSettings.CastShadowAsMasked.ParameterValue != MaterialEditorInstance->Parent->GetCastShadowAsMasked();
+		return MaterialEditorInstance->Parent != nullptr ? MaterialEditorInstance->LightmassSettings.CastShadowAsMasked.ParameterValue != MaterialEditorInstance->Parent->GetCastShadowAsMasked() : false;
 	});
 	FResetToDefaultHandler ResetCastShadowAsMaskedPropertyHandler = FResetToDefaultHandler::CreateLambda([this](TSharedPtr<IPropertyHandle> InHandle) {
-		MaterialEditorInstance->LightmassSettings.CastShadowAsMasked.ParameterValue = MaterialEditorInstance->Parent->GetCastShadowAsMasked();
+		if (MaterialEditorInstance->Parent != nullptr)
+		{
+			MaterialEditorInstance->LightmassSettings.CastShadowAsMasked.ParameterValue = MaterialEditorInstance->Parent->GetCastShadowAsMasked();
+		}
 	});
 	FResetToDefaultOverride ResetCastShadowAsMaskedPropertyOverride = FResetToDefaultOverride::Create(IsCastShadowAsMaskedPropertyResetVisible, ResetCastShadowAsMaskedPropertyHandler);
 	IDetailPropertyRow& CastShadowAsMaskedPropertyRow = LightmassSettingsGroup.AddPropertyRow(CastShadowAsMaskedProperty->GetChildHandle(0).ToSharedRef());
@@ -782,10 +785,13 @@ void FMaterialInstanceParameterDetails::CreateLightmassOverrideWidgets(IDetailLa
 		.OverrideResetToDefault(ResetCastShadowAsMaskedPropertyOverride);
 
 	FIsResetToDefaultVisible IsEmissiveBoostPropertyResetVisible = FIsResetToDefaultVisible::CreateLambda([this](TSharedPtr<IPropertyHandle> InHandle) {
-		return MaterialEditorInstance->LightmassSettings.EmissiveBoost.ParameterValue != MaterialEditorInstance->Parent->GetEmissiveBoost();
+		return MaterialEditorInstance->Parent != nullptr ? MaterialEditorInstance->LightmassSettings.EmissiveBoost.ParameterValue != MaterialEditorInstance->Parent->GetEmissiveBoost() : false;
 	});
 	FResetToDefaultHandler ResetEmissiveBoostPropertyHandler = FResetToDefaultHandler::CreateLambda([this](TSharedPtr<IPropertyHandle> InHandle) {
-		MaterialEditorInstance->LightmassSettings.EmissiveBoost.ParameterValue = MaterialEditorInstance->Parent->GetEmissiveBoost();
+		if (MaterialEditorInstance->Parent != nullptr)
+		{
+			MaterialEditorInstance->LightmassSettings.EmissiveBoost.ParameterValue = MaterialEditorInstance->Parent->GetEmissiveBoost();
+		}
 	});
 	FResetToDefaultOverride ResetEmissiveBoostPropertyOverride = FResetToDefaultOverride::Create(IsEmissiveBoostPropertyResetVisible, ResetEmissiveBoostPropertyHandler);
 	IDetailPropertyRow& EmissiveBoostPropertyRow = LightmassSettingsGroup.AddPropertyRow(EmissiveBoostProperty->GetChildHandle(0).ToSharedRef());
@@ -801,10 +807,13 @@ void FMaterialInstanceParameterDetails::CreateLightmassOverrideWidgets(IDetailLa
 		.OverrideResetToDefault(ResetEmissiveBoostPropertyOverride);
 	
 	FIsResetToDefaultVisible IsDiffuseBoostPropertyResetVisible = FIsResetToDefaultVisible::CreateLambda([this](TSharedPtr<IPropertyHandle> InHandle) {
-		return MaterialEditorInstance->LightmassSettings.DiffuseBoost.ParameterValue != MaterialEditorInstance->Parent->GetDiffuseBoost();
+		return MaterialEditorInstance->Parent != nullptr ? MaterialEditorInstance->LightmassSettings.DiffuseBoost.ParameterValue != MaterialEditorInstance->Parent->GetDiffuseBoost() : false;
 	});
 	FResetToDefaultHandler ResetDiffuseBoostPropertyHandler = FResetToDefaultHandler::CreateLambda([this](TSharedPtr<IPropertyHandle> InHandle) {
-		MaterialEditorInstance->LightmassSettings.DiffuseBoost.ParameterValue = MaterialEditorInstance->Parent->GetDiffuseBoost();
+		if (MaterialEditorInstance->Parent != nullptr)
+		{
+			MaterialEditorInstance->LightmassSettings.DiffuseBoost.ParameterValue = MaterialEditorInstance->Parent->GetDiffuseBoost();
+		}
 	});
 	FResetToDefaultOverride ResetDiffuseBoostPropertyOverride = FResetToDefaultOverride::Create(IsDiffuseBoostPropertyResetVisible, ResetDiffuseBoostPropertyHandler);
 	IDetailPropertyRow& DiffuseBoostPropertyRow = LightmassSettingsGroup.AddPropertyRow(DiffuseBoostProperty->GetChildHandle(0).ToSharedRef());
@@ -820,10 +829,13 @@ void FMaterialInstanceParameterDetails::CreateLightmassOverrideWidgets(IDetailLa
 		.OverrideResetToDefault(ResetDiffuseBoostPropertyOverride);
 
 	FIsResetToDefaultVisible IsExportResolutionScalePropertyResetVisible = FIsResetToDefaultVisible::CreateLambda([this](TSharedPtr<IPropertyHandle> InHandle) {
-		return MaterialEditorInstance->LightmassSettings.ExportResolutionScale.ParameterValue != MaterialEditorInstance->Parent->GetDiffuseBoost();
+		return MaterialEditorInstance->Parent != nullptr ? MaterialEditorInstance->LightmassSettings.ExportResolutionScale.ParameterValue != MaterialEditorInstance->Parent->GetDiffuseBoost()  : false;
 	});
 	FResetToDefaultHandler ResetExportResolutionScalePropertyHandler = FResetToDefaultHandler::CreateLambda([this](TSharedPtr<IPropertyHandle> InHandle) {
-		MaterialEditorInstance->LightmassSettings.ExportResolutionScale.ParameterValue = MaterialEditorInstance->Parent->GetDiffuseBoost();
+		if (MaterialEditorInstance->Parent != nullptr)
+		{
+			MaterialEditorInstance->LightmassSettings.ExportResolutionScale.ParameterValue = MaterialEditorInstance->Parent->GetDiffuseBoost();
+		}
 	});
 	FResetToDefaultOverride ResetExportResolutionScalePropertyOverride = FResetToDefaultOverride::Create(IsExportResolutionScalePropertyResetVisible, ResetExportResolutionScalePropertyHandler);
 	IDetailPropertyRow& ExportResolutionScalePropertyRow = LightmassSettingsGroup.AddPropertyRow(ExportResolutionScaleProperty->GetChildHandle(0).ToSharedRef());
@@ -860,10 +872,13 @@ void FMaterialInstanceParameterDetails::CreateBasePropertyOverrideWidgets(IDetai
 	TSharedPtr<IPropertyHandle> DitheredLODTransitionProperty = BasePropertyOverridePropery->GetChildHandle("DitheredLODTransition");
 
 	FIsResetToDefaultVisible IsOpacityClipMaskValuePropertyResetVisible = FIsResetToDefaultVisible::CreateLambda([this](TSharedPtr<IPropertyHandle> InHandle) {
-		return MaterialEditorInstance->BasePropertyOverrides.OpacityMaskClipValue != MaterialEditorInstance->Parent->GetOpacityMaskClipValue();
+		return MaterialEditorInstance->Parent != nullptr ? MaterialEditorInstance->BasePropertyOverrides.OpacityMaskClipValue != MaterialEditorInstance->Parent->GetOpacityMaskClipValue() : false;
 	});
 	FResetToDefaultHandler ResetOpacityClipMaskValuePropertyHandler = FResetToDefaultHandler::CreateLambda([this](TSharedPtr<IPropertyHandle> InHandle) {
-		MaterialEditorInstance->BasePropertyOverrides.OpacityMaskClipValue = MaterialEditorInstance->Parent->GetOpacityMaskClipValue();
+		if (MaterialEditorInstance->Parent != nullptr)
+		{
+			MaterialEditorInstance->BasePropertyOverrides.OpacityMaskClipValue = MaterialEditorInstance->Parent->GetOpacityMaskClipValue();
+		}
 	});
 	FResetToDefaultOverride ResetOpacityClipMaskValuePropertyOverride = FResetToDefaultOverride::Create(IsOpacityClipMaskValuePropertyResetVisible, ResetOpacityClipMaskValuePropertyHandler);
 	IDetailPropertyRow& OpacityClipMaskValuePropertyRow = BasePropertyOverrideGroup.AddPropertyRow(OpacityClipMaskValueProperty.ToSharedRef());
@@ -875,10 +890,13 @@ void FMaterialInstanceParameterDetails::CreateBasePropertyOverrideWidgets(IDetai
 		.OverrideResetToDefault(ResetOpacityClipMaskValuePropertyOverride);
 
 	FIsResetToDefaultVisible IsBlendModePropertyResetVisible = FIsResetToDefaultVisible::CreateLambda([this](TSharedPtr<IPropertyHandle> InHandle) {
-		return MaterialEditorInstance->BasePropertyOverrides.BlendMode != MaterialEditorInstance->Parent->GetBlendMode();
+		return MaterialEditorInstance->Parent != nullptr ? MaterialEditorInstance->BasePropertyOverrides.BlendMode != MaterialEditorInstance->Parent->GetBlendMode() : false;
 	});
 	FResetToDefaultHandler ResetBlendModePropertyHandler = FResetToDefaultHandler::CreateLambda([this](TSharedPtr<IPropertyHandle> InHandle) {
-		MaterialEditorInstance->BasePropertyOverrides.BlendMode = MaterialEditorInstance->Parent->GetBlendMode();
+		if (MaterialEditorInstance->Parent != nullptr)
+		{
+			MaterialEditorInstance->BasePropertyOverrides.BlendMode = MaterialEditorInstance->Parent->GetBlendMode();
+		}
 	}); 
 	FResetToDefaultOverride ResetBlendModePropertyOverride = FResetToDefaultOverride::Create(IsBlendModePropertyResetVisible, ResetBlendModePropertyHandler);
 	IDetailPropertyRow& BlendModePropertyRow = BasePropertyOverrideGroup.AddPropertyRow(BlendModeProperty.ToSharedRef());
@@ -890,10 +908,13 @@ void FMaterialInstanceParameterDetails::CreateBasePropertyOverrideWidgets(IDetai
 		.OverrideResetToDefault(ResetBlendModePropertyOverride);
 
 	FIsResetToDefaultVisible IsShadingModelPropertyResetVisible = FIsResetToDefaultVisible::CreateLambda([this](TSharedPtr<IPropertyHandle> InHandle) {
-		return MaterialEditorInstance->BasePropertyOverrides.ShadingModel != MaterialEditorInstance->Parent->GetShadingModel();
+		return MaterialEditorInstance->Parent != nullptr ? MaterialEditorInstance->BasePropertyOverrides.ShadingModel != MaterialEditorInstance->Parent->GetShadingModel() : false;
 	});
 	FResetToDefaultHandler ResetShadingModelPropertyHandler = FResetToDefaultHandler::CreateLambda([this](TSharedPtr<IPropertyHandle> InHandle) {
-		MaterialEditorInstance->BasePropertyOverrides.ShadingModel = MaterialEditorInstance->Parent->GetShadingModel();
+		if (MaterialEditorInstance->Parent != nullptr)
+		{
+			MaterialEditorInstance->BasePropertyOverrides.ShadingModel = MaterialEditorInstance->Parent->GetShadingModel();
+		}
 	});
 	FResetToDefaultOverride ResetShadingModelPropertyOverride = FResetToDefaultOverride::Create(IsShadingModelPropertyResetVisible, ResetShadingModelPropertyHandler);
 	IDetailPropertyRow& ShadingModelPropertyRow = BasePropertyOverrideGroup.AddPropertyRow(ShadingModelProperty.ToSharedRef());
@@ -905,10 +926,13 @@ void FMaterialInstanceParameterDetails::CreateBasePropertyOverrideWidgets(IDetai
 		.OverrideResetToDefault(ResetShadingModelPropertyOverride);
 
 	FIsResetToDefaultVisible IsTwoSidedPropertyResetVisible = FIsResetToDefaultVisible::CreateLambda([this](TSharedPtr<IPropertyHandle> InHandle) {
-		return MaterialEditorInstance->BasePropertyOverrides.TwoSided != MaterialEditorInstance->Parent->IsTwoSided();
+		return MaterialEditorInstance->Parent != nullptr ? MaterialEditorInstance->BasePropertyOverrides.TwoSided != MaterialEditorInstance->Parent->IsTwoSided() : false;
 	});
 	FResetToDefaultHandler ResetTwoSidedValuePropertyHandler = FResetToDefaultHandler::CreateLambda([this](TSharedPtr<IPropertyHandle> InHandle) {
-		MaterialEditorInstance->BasePropertyOverrides.TwoSided = MaterialEditorInstance->Parent->IsTwoSided();
+		if (MaterialEditorInstance->Parent != nullptr)
+		{
+			MaterialEditorInstance->BasePropertyOverrides.TwoSided = MaterialEditorInstance->Parent->IsTwoSided();
+		}
 	});
 	FResetToDefaultOverride ResetTwoSidedPropertyOverride = FResetToDefaultOverride::Create(IsTwoSidedPropertyResetVisible, ResetTwoSidedValuePropertyHandler);
 	IDetailPropertyRow& TwoSidedPropertyRow = BasePropertyOverrideGroup.AddPropertyRow(TwoSidedProperty.ToSharedRef());
@@ -920,10 +944,13 @@ void FMaterialInstanceParameterDetails::CreateBasePropertyOverrideWidgets(IDetai
 		.OverrideResetToDefault(ResetTwoSidedPropertyOverride);
 
 	FIsResetToDefaultVisible IsDitheredLODTransitionPropertyResetVisible = FIsResetToDefaultVisible::CreateLambda([this](TSharedPtr<IPropertyHandle> InHandle) {
-		return MaterialEditorInstance->BasePropertyOverrides.DitheredLODTransition != MaterialEditorInstance->Parent->IsDitheredLODTransition();
+		return MaterialEditorInstance->Parent != nullptr ? MaterialEditorInstance->BasePropertyOverrides.DitheredLODTransition != MaterialEditorInstance->Parent->IsDitheredLODTransition() : false;
 	});
 	FResetToDefaultHandler ResetDitheredLODTransitionPropertyHandler = FResetToDefaultHandler::CreateLambda([this](TSharedPtr<IPropertyHandle> InHandle) {
-		MaterialEditorInstance->BasePropertyOverrides.DitheredLODTransition = MaterialEditorInstance->Parent->IsDitheredLODTransition();
+		if (MaterialEditorInstance->Parent != nullptr)
+		{
+			MaterialEditorInstance->BasePropertyOverrides.DitheredLODTransition = MaterialEditorInstance->Parent->IsDitheredLODTransition();
+		}
 	});
 	FResetToDefaultOverride ResetDitheredLODTransitionPropertyOverride = FResetToDefaultOverride::Create(IsDitheredLODTransitionPropertyResetVisible, ResetDitheredLODTransitionPropertyHandler);
 	IDetailPropertyRow& DitheredLODTransitionPropertyRow = BasePropertyOverrideGroup.AddPropertyRow(DitheredLODTransitionProperty.ToSharedRef());
