Adding functionality to MaterialEditingLibrary. Get Static Switch Parameter Value, plus matching Get X Parameter Value functions for Materials (for scalar, vector, texture, and static switch). Also added a getter to check material usage flags.
#rnx

diff --git a/Engine/Source/Editor/MaterialEditor/Private/MaterialEditingLibrary.cpp b/Engine/Source/Editor/MaterialEditor/Private/MaterialEditingLibrary.cpp
--- a/Engine/Source/Editor/MaterialEditor/Private/MaterialEditingLibrary.cpp
+++ b/Engine/Source/Editor/MaterialEditor/Private/MaterialEditingLibrary.cpp
@@ -563,6 +563,16 @@ bool UMaterialEditingLibrary::SetMaterialUsage(UMaterial* Material, EMaterialUsa
 	return bResult;
 }
 
+bool UMaterialEditingLibrary::HasMaterialUsage(UMaterial* Material, EMaterialUsage Usage)
+{
+	bool bResult = false;
+	if (Material)
+	{
+		bResult = Material->GetUsageByFlag(Usage);
+	}
+	return bResult;
+}
+
 bool UMaterialEditingLibrary::ConnectMaterialProperty(UMaterialExpression* FromExpression, FString FromOutputName, EMaterialProperty Property)
 {
 	bool bResult = false;
@@ -648,6 +658,47 @@ void UMaterialEditingLibrary::LayoutMaterialExpressions(UMaterial* Material)
 	MaterialEditingLibraryImpl::LayoutMaterialExpressions( Material );
 }
 
+float UMaterialEditingLibrary::GetMaterialDefaultScalarParameterValue(UMaterial* Material, FName ParameterName)
+{
+	float Result = 0.f;
+	if (Material)
+	{
+		Material->GetScalarParameterDefaultValue(ParameterName, Result);
+	}
+	return Result;
+}
+
+UTexture* UMaterialEditingLibrary::GetMaterialDefaultTextureParameterValue(UMaterial* Material, FName ParameterName)
+{
+	UTexture* Result = nullptr;
+	if (Material)
+	{
+		Material->GetTextureParameterDefaultValue(ParameterName, Result);
+	}
+	return Result;
+}
+
+FLinearColor UMaterialEditingLibrary::GetMaterialDefaultVectorParameterValue(UMaterial* Material, FName ParameterName)
+{
+	FLinearColor Result = FLinearColor::Black;
+	if (Material)
+	{
+		Material->GetVectorParameterDefaultValue(ParameterName, Result);
+	}
+	return Result;
+}
+
+bool UMaterialEditingLibrary::GetMaterialDefaultStaticSwitchParameterValue(UMaterial* Material, FName ParameterName)
+{
+	bool bResult = false;
+	if (Material)
+	{
+		FGuid OutGuid;
+		Material->GetStaticSwitchParameterDefaultValue(ParameterName, bResult, OutGuid);
+	}
+	return bResult;
+}
+
 //////////////////////////////////////////////////////////////////////////
 
 int32 UMaterialEditingLibrary::GetNumMaterialExpressionsInFunction(const UMaterialFunction* MaterialFunction)
@@ -888,6 +939,17 @@ bool UMaterialEditingLibrary::SetMaterialInstanceVectorParameterValue(UMaterialI
 }
 
 
+bool UMaterialEditingLibrary::GetMaterialInstanceStaticSwitchParameterValue(UMaterialInstanceConstant* Instance, FName ParameterName)
+{
+	bool bResult = false;
+	if (Instance)
+	{
+		FGuid OutGuid;
+		Instance->GetStaticSwitchParameterValue(ParameterName, bResult, OutGuid);
+	}
+	return bResult;
+}
+
 void UMaterialEditingLibrary::UpdateMaterialInstance(UMaterialInstanceConstant* Instance)
 {
 	if (Instance)
diff --git a/Engine/Source/Editor/MaterialEditor/Public/MaterialEditingLibrary.h b/Engine/Source/Editor/MaterialEditor/Public/MaterialEditingLibrary.h
--- a/Engine/Source/Editor/MaterialEditor/Public/MaterialEditingLibrary.h
+++ b/Engine/Source/Editor/MaterialEditor/Public/MaterialEditingLibrary.h
@@ -78,6 +78,14 @@ public:
 	UFUNCTION(BlueprintCallable, Category = "MaterialEditing")
 	static bool SetMaterialUsage(UMaterial* Material, EMaterialUsage Usage, bool& bNeedsRecompile);
 
+	/**
+	 *	Check if a particular usage is enabled for the supplied material (e.g. SkeletalMesh, ParticleSprite etc)
+	 *	@param	Material			Material to check usage for
+	 *	@param	Usage				Usage type to check for this material
+	 */
+	UFUNCTION(BlueprintPure, Category = "MaterialEditing")
+	static bool HasMaterialUsage(UMaterial* Material, EMaterialUsage Usage);
+
 	/** 
 	 *	Connect a material expression output to one of the material property inputs (e.g. diffuse color, opacity etc)
 	 *	@param	FromExpression		Expression to make connection from
@@ -109,6 +117,23 @@ public:
 	UFUNCTION(BlueprintCallable, Category = "MaterialEditing")
 	static void LayoutMaterialExpressions(UMaterial* Material);
 
+	/** Get the default scalar (float) parameter value from a Material */
+	UFUNCTION(BlueprintPure, Category = "MaterialEditing")
+	static float GetMaterialDefaultScalarParameterValue(UMaterial* Material, FName ParameterName);
+
+
+	/** Get the default texture parameter value from a Material  */
+	UFUNCTION(BlueprintPure, Category = "MaterialEditing")
+	static UTexture* GetMaterialDefaultTextureParameterValue(UMaterial* Material, FName ParameterName);
+
+	/** Get the default vector parameter value from a Material */
+	UFUNCTION(BlueprintPure, Category = "MaterialEditing")
+	static FLinearColor GetMaterialDefaultVectorParameterValue(UMaterial* Material, FName ParameterName);
+
+	/** Get the default static switch parameter value from a Material */
+	UFUNCTION(BlueprintPure, Category = "MaterialEditing")
+	static bool GetMaterialDefaultStaticSwitchParameterValue(UMaterial* Material, FName ParameterName);
+
 	//////// MATERIAL FUNCTION EDITING
 
 	/** Returns number of material expressions in the supplied material */
@@ -183,6 +208,9 @@ public:
 	UFUNCTION(BlueprintCallable, Category = "MaterialEditing")
 	static bool SetMaterialInstanceVectorParameterValue(UMaterialInstanceConstant* Instance, FName ParameterName, FLinearColor Value);
 
+	/** Get the current static switch parameter value from a Material Instance */
+	UFUNCTION(BlueprintPure, Category = "MaterialEditing")
+	static bool GetMaterialInstanceStaticSwitchParameterValue(UMaterialInstanceConstant* Instance, FName ParameterName);
 
 	/** Called after making modifications to a Material Instance to recompile shaders etc. */
 	UFUNCTION(BlueprintCallable, Category = "MaterialEditing")
