Adding a function to the Material Editing Library to get children of any material or material instance.
#rnx

diff --git a/Engine/Source/Editor/MaterialEditor/Private/MaterialEditingLibrary.cpp b/Engine/Source/Editor/MaterialEditor/Private/MaterialEditingLibrary.cpp
--- a/Engine/Source/Editor/MaterialEditor/Private/MaterialEditingLibrary.cpp
+++ b/Engine/Source/Editor/MaterialEditor/Private/MaterialEditingLibrary.cpp
@@ -26,6 +26,7 @@
 #include "EditorSupportDelegates.h"
 #include "Misc/RuntimeErrors.h"
 #include "SceneTypes.h"
+#include "AssetRegistryModule.h"
 
 
 DEFINE_LOG_CATEGORY_STATIC(LogMaterialEditingLibrary, Warning, All);
@@ -965,4 +966,26 @@ void UMaterialEditingLibrary::UpdateMaterialInstance(UMaterialInstanceConstant*
 		FEditorDelegates::RefreshEditor.Broadcast();
 		FEditorSupportDelegates::RedrawAllViewports.Broadcast();
 	}
-}
\ No newline at end of file
+}
+
+void UMaterialEditingLibrary::GetChildInstances(UMaterialInterface* Parent, TArray< TSoftObjectPtr<UMaterialInstance> >& ChildInstances)
+{
+	FAssetRegistryModule& AssetRegistryModule = FModuleManager::LoadModuleChecked<FAssetRegistryModule>(TEXT("AssetRegistry"));
+	FAssetData ParentData = FAssetData(Parent);
+	TArray<FAssetData> AssetList;
+	TMultiMap<FName, FString> TagsAndValues;
+	FString ParentNameString = ParentData.GetExportTextName();
+	TagsAndValues.Add(GET_MEMBER_NAME_CHECKED(UMaterialInstance, Parent), ParentNameString);
+	AssetRegistryModule.Get().GetAssetsByTagValues(TagsAndValues, AssetList);
+	
+	for (const FAssetData& MatInstRef : AssetList)
+	{
+		if (UMaterialInstance* MaterialInstance = Cast<UMaterialInstance>(MatInstRef.GetAsset()))
+		{
+			if (MaterialInstance->Parent == Parent)
+			{
+				ChildInstances.Add(MaterialInstance);
+			}
+		}
+	}
+}
diff --git a/Engine/Source/Editor/MaterialEditor/Public/MaterialEditingLibrary.h b/Engine/Source/Editor/MaterialEditor/Public/MaterialEditingLibrary.h
--- a/Engine/Source/Editor/MaterialEditor/Public/MaterialEditingLibrary.h
+++ b/Engine/Source/Editor/MaterialEditor/Public/MaterialEditingLibrary.h
@@ -9,7 +9,8 @@
 #include "MaterialEditingLibrary.generated.h"
 
 class UMaterialFunction;
-class MaterialInstance;
+class UMaterialInstance;
+
 
 /** Blueprint library for creating/editing Materials */
 UCLASS()
@@ -215,4 +216,8 @@ public:
 	/** Called after making modifications to a Material Instance to recompile shaders etc. */
 	UFUNCTION(BlueprintCallable, Category = "MaterialEditing")
 	static void UpdateMaterialInstance(UMaterialInstanceConstant* Instance);
+
+	/** Gets all direct child mat instances */
+	UFUNCTION(BlueprintCallable, Category = "MaterialEditing")
+	static void GetChildInstances(UMaterialInterface* Parent, TArray<TSoftObjectPtr<UMaterialInstance>>& ChildInstances);
 };
\ No newline at end of file
