Adding hierarchy button to material editor toolbar to show material instances and added children to material instance editor hierarchy dropdown
#rnx

diff --git a/Engine/Source/Editor/MaterialEditor/MaterialEditor.Build.cs b/Engine/Source/Editor/MaterialEditor/MaterialEditor.Build.cs
--- a/Engine/Source/Editor/MaterialEditor/MaterialEditor.Build.cs
+++ b/Engine/Source/Editor/MaterialEditor/MaterialEditor.Build.cs
@@ -41,8 +41,8 @@ public class MaterialEditor : ModuleRules
 				"GraphEditor",
                 "AdvancedPreviewScene",
                 "Projects",
-
-			}
+                "AssetRegistry",
+            }
 		);
 
 		DynamicallyLoadedModuleNames.AddRange(
diff --git a/Engine/Source/Editor/MaterialEditor/Private/MaterialEditingLibrary.cpp b/Engine/Source/Editor/MaterialEditor/Private/MaterialEditingLibrary.cpp
--- a/Engine/Source/Editor/MaterialEditor/Private/MaterialEditingLibrary.cpp
+++ b/Engine/Source/Editor/MaterialEditor/Private/MaterialEditingLibrary.cpp
@@ -968,24 +968,17 @@ void UMaterialEditingLibrary::UpdateMaterialInstance(UMaterialInstanceConstant*
 	}
 }
 
-void UMaterialEditingLibrary::GetChildInstances(UMaterialInterface* Parent, TArray< TSoftObjectPtr<UMaterialInstance> >& ChildInstances)
+void UMaterialEditingLibrary::GetChildInstances(UMaterialInterface* Parent, TArray< FAssetData>& ChildInstances)
 {
 	FAssetRegistryModule& AssetRegistryModule = FModuleManager::LoadModuleChecked<FAssetRegistryModule>(TEXT("AssetRegistry"));
-	FAssetData ParentData = FAssetData(Parent);
 	TArray<FAssetData> AssetList;
 	TMultiMap<FName, FString> TagsAndValues;
-	FString ParentNameString = ParentData.GetExportTextName();
+	const FString ParentNameString = FAssetData(Parent).GetExportTextName();
 	TagsAndValues.Add(GET_MEMBER_NAME_CHECKED(UMaterialInstance, Parent), ParentNameString);
 	AssetRegistryModule.Get().GetAssetsByTagValues(TagsAndValues, AssetList);
 	
-	for (const FAssetData& MatInstRef : AssetList)
+	for (const FAssetData MatInstRef : AssetList)
 	{
-		if (UMaterialInstance* MaterialInstance = Cast<UMaterialInstance>(MatInstRef.GetAsset()))
-		{
-			if (MaterialInstance->Parent == Parent)
-			{
-				ChildInstances.Add(MaterialInstance);
-			}
-		}
+		ChildInstances.Add(MatInstRef);
 	}
 }
diff --git a/Engine/Source/Editor/MaterialEditor/Private/MaterialEditor.cpp b/Engine/Source/Editor/MaterialEditor/Private/MaterialEditor.cpp
--- a/Engine/Source/Editor/MaterialEditor/Private/MaterialEditor.cpp
+++ b/Engine/Source/Editor/MaterialEditor/Private/MaterialEditor.cpp
@@ -127,6 +127,7 @@
 
 #include "SMaterialParametersOverviewWidget.h"
 #include "IPropertyRowGenerator.h"
+#include "Widgets/Layout/SScrollBox.h"
 
 #define LOCTEXT_NAMESPACE "MaterialEditor"
 
@@ -1168,8 +1169,74 @@ void FMaterialEditor::FillToolbar(FToolBarBuilder& ToolbarBuilder)
 			false);
 	}
 	ToolbarBuilder.EndSection();
+	ToolbarBuilder.BeginSection("Hierarchy");
+	{
+		ToolbarBuilder.AddComboButton(
+			FUIAction(),
+			FOnGetContent::CreateSP(this, &FMaterialEditor::GenerateInheritanceMenu),
+			LOCTEXT("Hierarchy", "Hierarchy"),
+			FText::GetEmpty(),
+			FSlateIcon(FSlateIcon(FEditorStyle::GetStyleSetName(), "BTEditor.SwitchToBehaviorTreeMode")),
+			false
+		);
+	}
 };
 
+TSharedRef<SWidget> FMaterialEditor::GenerateInheritanceMenu()
+{
+	struct Local
+	{
+		static void MakeMaterialSubmenu(FMenuBuilder& MenuBuilder, FAssetData InMaterial)
+		{
+			MenuBuilder.AddMenuEntry(LOCTEXT("OpenInEditor", "Open In Editor"),
+				FText::GetEmpty(),
+				FSlateIcon(FEditorStyle::GetStyleSetName(), "ContentBrowser.AssetActions.OpenInExternalEditor"),
+				FUIAction(FExecuteAction::CreateStatic(&FMaterialEditorUtilities::OnOpenMaterial, InMaterial)));
+
+			MenuBuilder.AddMenuEntry(LOCTEXT("FindInContentBrowser", "Find In Content Browser"),
+				FText::GetEmpty(),
+				FSlateIcon(FEditorStyle::GetStyleSetName(), "SystemWideCommands.FindInContentBrowser"),
+				FUIAction(FExecuteAction::CreateStatic(&FMaterialEditorUtilities::OnShowMaterialInContentBrowser, InMaterial)));
+		}
+	};
+
+	RebuildInheritanceList();
+	TSharedPtr<FExtender> ToolbarExtender = MakeShareable(new FExtender);
+	const bool bShouldCloseWindowAfterMenuSelection = true;
+	TSharedPtr<FUICommandList> InCommandList = MakeShareable(new FUICommandList);
+	FMenuBuilder MenuBuilder(bShouldCloseWindowAfterMenuSelection, InCommandList, ToolbarExtender);
+
+	if (!MaterialFunction)
+	{
+		const FName MaterialInstances = TEXT("MaterialInstances");
+		MenuBuilder.BeginSection(MaterialInstances, LOCTEXT("MaterialInstances", "Material Instances"));
+		for (FAssetData MaterialChild : MaterialChildList)
+		{
+			FFormatNamedArguments Args;
+			Args.Add(TEXT("ChildName"), FText::FromName(MaterialChild.AssetName));
+			MenuBuilder.AddSubMenu(FText::Format(LOCTEXT("MaterialChildName", "{ChildName}"), Args),
+				FText::GetEmpty(),
+				FNewMenuDelegate::CreateStatic(&Local::MakeMaterialSubmenu, MaterialChild),
+				false,
+				FSlateIcon());
+		}
+		MenuBuilder.EndSection();
+	}
+
+	TSharedRef<SWidget> ConstrainedMenu = SNew(SVerticalBox)
+		+ SVerticalBox::Slot()
+		.MaxHeight(500.0f)
+		[
+			SNew(SScrollBox)
+			+ SScrollBox::Slot()
+		[
+			MenuBuilder.MakeWidget()
+		]
+		];
+
+	return ConstrainedMenu;
+}
+
 TSharedRef< SWidget > FMaterialEditor::GeneratePreviewMenuContent()
 {
 	bool bShouldCloseWindowAfterMenuSelection = true;
@@ -5153,4 +5220,14 @@ void FMaterialEditor::FocusDetailsPanel()
 	}
 }
 
+void FMaterialEditor::RebuildInheritanceList()
+{
+	if (!MaterialFunction)
+	{
+		MaterialChildList.Empty();
+
+		UMaterialEditingLibrary::GetChildInstances(OriginalMaterial, MaterialChildList);
+	}
+}
+
 #undef LOCTEXT_NAMESPACE
diff --git a/Engine/Source/Editor/MaterialEditor/Private/MaterialEditor.h b/Engine/Source/Editor/MaterialEditor/Private/MaterialEditor.h
--- a/Engine/Source/Editor/MaterialEditor/Private/MaterialEditor.h
+++ b/Engine/Source/Editor/MaterialEditor/Private/MaterialEditor.h
@@ -36,6 +36,7 @@ class UMaterialExpressionComment;
 class UMaterialInstance;
 class UMaterialGraphNode;
 struct FGraphAppearanceInfo;
+class UMaterialFunctionInstance;
 
 /**
  * Class for rendering previews of material expressions in the material editor's linked object viewport.
@@ -389,6 +390,9 @@ public:
 	/** Called to bring focus to the details panel */
 	void FocusDetailsPanel();
 
+	/** Rebuilds the inheritance list for this material. */
+	void RebuildInheritanceList();
+
 public:
 	/** Set to true when modifications have been made to the material */
 	bool bMaterialDirty;
@@ -524,6 +528,8 @@ private:
 	/** Creates the toolbar buttons. Bound by ExtendToolbar*/
 	void FillToolbar(FToolBarBuilder& ToolbarBuilder);
 
+	TSharedRef<SWidget> GenerateInheritanceMenu();
+
 	TSharedRef< SWidget > GeneratePreviewMenuContent();
 
 	/** Allows editor to veto the setting of a preview asset */
@@ -852,4 +858,9 @@ private:
 	/** True if the quality level or feature level to preview has been changed */
 	bool bPreviewFeaturesChanged;
 
+	/** List of children used to populate the inheritance list chain. */
+	TArray< FAssetData > MaterialChildList;
+
+	/** List of children used to populate the inheritance list chain. */
+	TArray< FAssetData > FunctionChildList;
 };
diff --git a/Engine/Source/Editor/MaterialEditor/Private/MaterialEditorUtilities.cpp b/Engine/Source/Editor/MaterialEditor/Private/MaterialEditorUtilities.cpp
--- a/Engine/Source/Editor/MaterialEditor/Private/MaterialEditorUtilities.cpp
+++ b/Engine/Source/Editor/MaterialEditor/Private/MaterialEditorUtilities.cpp
@@ -32,6 +32,7 @@
 #include "MaterialUtilities.h"
 #include "Misc/ScopedSlowTask.h"
 #include "Templates/UniquePtr.h"
+#include "Materials/MaterialFunctionInstance.h"
 
 #define LOCTEXT_NAMESPACE "MaterialEditorUtilities"
 
@@ -719,4 +720,70 @@ void FMaterialEditorUtilities::BuildTextureStreamingData(UMaterialInterface* Upd
 	}
 }
 
+
+void FMaterialEditorUtilities::OnOpenMaterial(FAssetData InMaterial)
+{
+	UMaterialInterface* MaterialInterface = Cast<UMaterialInterface>(InMaterial.GetAsset());
+	OpenSelectedParentEditor(MaterialInterface);
+}
+
+void FMaterialEditorUtilities::OnOpenFunction(FAssetData InFunction)
+{
+	UMaterialFunctionInterface* MaterialFunctionInterface = Cast<UMaterialFunctionInterface>(InFunction.GetAsset());
+	OpenSelectedParentEditor(MaterialFunctionInterface);
+}
+
+void FMaterialEditorUtilities::OnShowMaterialInContentBrowser(FAssetData InMaterial)
+{
+	TArray<UObject*> SyncedObject;
+	SyncedObject.Add(InMaterial.GetAsset());
+	GEditor->SyncBrowserToObjects(SyncedObject);
+}
+
+void FMaterialEditorUtilities::OnShowFunctionInContentBrowser(FAssetData InFunction)
+{
+	TArray<UObject*> SyncedObject;
+	SyncedObject.Add(InFunction.GetAsset());
+	GEditor->SyncBrowserToObjects(SyncedObject);
+}
+
+void FMaterialEditorUtilities::OpenSelectedParentEditor(UMaterialInterface* InMaterialInterface)
+{
+	// See if its a material or material instance constant. 
+	if (ensure(InMaterialInterface))
+	{
+		if (InMaterialInterface->IsA(UMaterial::StaticClass()))
+		{
+			// Show material editor
+			UMaterial* Material = Cast<UMaterial>(InMaterialInterface);
+			FAssetEditorManager::Get().OpenEditorForAsset(Material);
+		}
+		else if (InMaterialInterface->IsA(UMaterialInstance::StaticClass()))
+		{
+			// Show material instance editor
+			UMaterialInstance* MaterialInstance = Cast<UMaterialInstance>(InMaterialInterface);
+			FAssetEditorManager::Get().OpenEditorForAsset(MaterialInstance);
+		}
+	}
+}
+
+void FMaterialEditorUtilities::OpenSelectedParentEditor(UMaterialFunctionInterface* InMaterialFunction)
+{
+	// See if its a material or material instance constant.  
+	if (ensure(InMaterialFunction) )
+	{
+		if (InMaterialFunction->IsA(UMaterialFunctionInstance::StaticClass()))
+		{
+			// Show function instance editor
+			UMaterialFunctionInstance* FunctionInstance = Cast<UMaterialFunctionInstance>(InMaterialFunction);
+			FAssetEditorManager::Get().OpenEditorForAsset(FunctionInstance);
+		}
+		else
+		{
+			// Show function editor
+			FAssetEditorManager::Get().OpenEditorForAsset(InMaterialFunction);
+		}
+	}
+}
+
 #undef LOCTEXT_NAMESPACE
diff --git a/Engine/Source/Editor/MaterialEditor/Private/MaterialInstanceEditor.cpp b/Engine/Source/Editor/MaterialEditor/Private/MaterialInstanceEditor.cpp
--- a/Engine/Source/Editor/MaterialEditor/Private/MaterialInstanceEditor.cpp
+++ b/Engine/Source/Editor/MaterialEditor/Private/MaterialInstanceEditor.cpp
@@ -41,6 +41,8 @@
 #include "Framework/Commands/UICommandInfo.h"
 #include "EditorStyleSet.h"
 #include "MaterialStats.h"
+#include "MaterialEditingLibrary.h"
+#include "Widgets/Layout/SScrollBox.h"
 
 #define LOCTEXT_NAMESPACE "MaterialInstanceEditor"
 
@@ -634,13 +636,6 @@ void FMaterialInstanceEditor::AddReferencedObjects(FReferenceCollector& Collecto
 {
 	// Serialize our custom object instance
 	Collector.AddReferencedObject(MaterialEditorInstance);
-
-	// Serialize all parent material instances that are stored in the list.
-	for (int32 Index = 0; Index < MaterialParentList.Num(); Index++)
-	{
-		UMaterialInterface* Parent = MaterialParentList[Index].Get();
-		Collector.AddReferencedObject(Parent);
-	}
 }
 
 void FMaterialInstanceEditor::BindCommands()
@@ -728,30 +723,6 @@ bool FMaterialInstanceEditor::IsShowAllMaterialParametersChecked() const
 	return bShowAllMaterialParameters;
 }
 
-void FMaterialInstanceEditor::OnOpenMaterial(TWeakObjectPtr<UMaterialInterface> InMaterial)
-{
-	OpenSelectedParentEditor(InMaterial.Get());
-}
-
-void FMaterialInstanceEditor::OnOpenFunction(TWeakObjectPtr<UMaterialFunctionInterface> InFunction)
-{
-	OpenSelectedParentEditor(InFunction.Get());
-}
-
-void FMaterialInstanceEditor::OnShowMaterialInContentBrowser(TWeakObjectPtr<UMaterialInterface> InMaterial)
-{
-	TArray<UObject*> SyncedObject;
-	SyncedObject.Add(InMaterial.Get());
-	GEditor->SyncBrowserToObjects(SyncedObject);
-}
-
-void FMaterialInstanceEditor::OnShowFunctionInContentBrowser(TWeakObjectPtr<UMaterialFunctionInterface> InFunction)
-{
-	TArray<UObject*> SyncedObject;
-	SyncedObject.Add(InFunction.Get());
-	GEditor->SyncBrowserToObjects(SyncedObject);
-}
-
 void FMaterialInstanceEditor::CreateInternalWidgets()
 {
 	PreviewVC = SNew(SMaterialEditor3DPreviewViewport)
@@ -853,30 +824,30 @@ TSharedRef< SWidget > FMaterialInstanceEditor::GenerateInheritanceMenu()
 {
 	struct Local
 	{
-		static void MakeFunctionSubmenu(FMenuBuilder& MenuBuilder, FMaterialInstanceEditor* InEditor, TWeakObjectPtr<UMaterialFunctionInterface> InFunction)
+		static void MakeFunctionSubmenu(FMenuBuilder& MenuBuilder, FAssetData InFunction)
 		{
 			MenuBuilder.AddMenuEntry(LOCTEXT("OpenInEditor", "Open In Editor"),
 				FText::GetEmpty(),
 				FSlateIcon(FEditorStyle::GetStyleSetName(), "ContentBrowser.AssetActions.OpenInExternalEditor"),
-				FUIAction(FExecuteAction::CreateSP(InEditor, &FMaterialInstanceEditor::OnOpenFunction, InFunction)));
+				FUIAction(FExecuteAction::CreateStatic(&FMaterialEditorUtilities::OnOpenFunction, InFunction)));
 			
 			MenuBuilder.AddMenuEntry(LOCTEXT("FindInContentBrowser", "Find In Content Browser"),
 				FText::GetEmpty(),
 				FSlateIcon(FEditorStyle::GetStyleSetName(), "SystemWideCommands.FindInContentBrowser"),
-				FUIAction(FExecuteAction::CreateSP(InEditor, &FMaterialInstanceEditor::OnShowFunctionInContentBrowser, InFunction)));
+				FUIAction(FExecuteAction::CreateStatic(&FMaterialEditorUtilities::OnShowFunctionInContentBrowser, InFunction)));
 		}
 
-		static void MakeMaterialSubmenu(FMenuBuilder& MenuBuilder, FMaterialInstanceEditor* InEditor, TWeakObjectPtr<UMaterialInterface> InMaterial)
+		static void MakeMaterialSubmenu(FMenuBuilder& MenuBuilder, FAssetData InMaterial)
 		{
 			MenuBuilder.AddMenuEntry(LOCTEXT("OpenInEditor", "Open In Editor"),
 				FText::GetEmpty(),
 				FSlateIcon(FEditorStyle::GetStyleSetName(), "ContentBrowser.AssetActions.OpenInExternalEditor"),
-				FUIAction(FExecuteAction::CreateSP(InEditor, &FMaterialInstanceEditor::OnOpenMaterial, InMaterial)));
+				FUIAction(FExecuteAction::CreateStatic(&FMaterialEditorUtilities::OnOpenMaterial, InMaterial)));
 			
 			MenuBuilder.AddMenuEntry(LOCTEXT("FindInContentBrowser", "Find In Content Browser"),
 				FText::GetEmpty(),
 				FSlateIcon(FEditorStyle::GetStyleSetName(), "SystemWideCommands.FindInContentBrowser"),
-				FUIAction(FExecuteAction::CreateSP(InEditor, &FMaterialInstanceEditor::OnShowMaterialInContentBrowser, InMaterial)));
+				FUIAction(FExecuteAction::CreateStatic(&FMaterialEditorUtilities::OnShowMaterialInContentBrowser, InMaterial)));
 		}
 	};
 
@@ -885,41 +856,65 @@ TSharedRef< SWidget > FMaterialInstanceEditor::GenerateInheritanceMenu()
 	const bool bShouldCloseWindowAfterMenuSelection = true;
 	TSharedPtr<FUICommandList> InCommandList = MakeShareable(new FUICommandList);
 	FMenuBuilder MenuBuilder(bShouldCloseWindowAfterMenuSelection, InCommandList, ToolbarExtender);
-
+	const FName ParentName = TEXT("ParentChain");
+	MenuBuilder.BeginSection(ParentName, LOCTEXT("ParentChain", "Parent Chain"));
 	if (bIsFunctionPreviewMaterial)
 	{
-		for (TWeakObjectPtr<UMaterialFunctionInterface> FunctionParent : FunctionParentList)
+		for (FAssetData FunctionParent : FunctionParentList)
 		{
-			if (FunctionParent.IsValid())
-			{
-				FFormatNamedArguments Args;
-				Args.Add(TEXT("ParentName"), FText::FromString(FunctionParent->GetName()));
-				MenuBuilder.AddSubMenu(FText::Format(LOCTEXT("InstanceParentName", "{ParentName}"), Args),
-					FText::GetEmpty(),
-					FNewMenuDelegate::CreateStatic(&Local::MakeFunctionSubmenu, this, FunctionParent),
-					false,
-					FSlateIcon());
-			}
+			FFormatNamedArguments Args;
+			Args.Add(TEXT("ParentName"), FText::FromName(FunctionParent.AssetName));
+			MenuBuilder.AddSubMenu(FText::Format(LOCTEXT("InstanceParentName", "{ParentName}"), Args),
+				FText::GetEmpty(),
+				FNewMenuDelegate::CreateStatic(&Local::MakeFunctionSubmenu, FunctionParent),
+				false,
+				FSlateIcon());
 		}
 	}
 	else
 	{
-		for (TWeakObjectPtr<UMaterialInterface> MaterialParent : MaterialParentList)
+		for (FAssetData MaterialParent : MaterialParentList)
 		{
-			if (MaterialParent.IsValid())
-			{
-				FFormatNamedArguments Args;
-				Args.Add(TEXT("ParentName"), FText::FromString(MaterialParent->GetName()));
-				MenuBuilder.AddSubMenu(FText::Format(LOCTEXT("InstanceParentName", "{ParentName}"), Args),
-					FText::GetEmpty(),
-					FNewMenuDelegate::CreateStatic(&Local::MakeMaterialSubmenu, this, MaterialParent),
-					false,
-					FSlateIcon());
-			}
+			FFormatNamedArguments Args;
+			Args.Add(TEXT("ParentName"), FText::FromName(MaterialParent.AssetName));
+			MenuBuilder.AddSubMenu(FText::Format(LOCTEXT("InstanceParentName", "{ParentName}"), Args),
+				FText::GetEmpty(),
+				FNewMenuDelegate::CreateStatic(&Local::MakeMaterialSubmenu, MaterialParent),
+				false,
+				FSlateIcon());
 		}
 	}
+	MenuBuilder.EndSection();
 
-	return MenuBuilder.MakeWidget();
+	if (!bIsFunctionPreviewMaterial)
+	{
+		const FName MaterialInstances = TEXT("MaterialInstances");
+		MenuBuilder.BeginSection(MaterialInstances, LOCTEXT("MaterialInstances", "Material Instances"));
+		for (FAssetData MaterialChild : MaterialChildList)
+		{
+			FFormatNamedArguments Args;
+			Args.Add(TEXT("ChildName"), FText::FromName(MaterialChild.AssetName));
+			MenuBuilder.AddSubMenu(FText::Format(LOCTEXT("MaterialInstanceName", "{ChildName}"), Args),
+				FText::GetEmpty(),
+				FNewMenuDelegate::CreateStatic(&Local::MakeMaterialSubmenu, MaterialChild),
+				false,
+				FSlateIcon());
+		}
+		MenuBuilder.EndSection();
+	}
+
+	TSharedRef<SWidget> ConstrainedMenu = SNew(SVerticalBox)
+		+SVerticalBox::Slot()
+		.MaxHeight(500.0f)
+		[		
+			SNew(SScrollBox)
+			+ SScrollBox::Slot()
+			[
+				MenuBuilder.MakeWidget()
+			]
+		];
+
+	return ConstrainedMenu;
 }
 
 
@@ -1149,8 +1144,6 @@ void FMaterialInstanceEditor::RebuildInheritanceList()
 	if (bIsFunctionPreviewMaterial)
 	{
 		FunctionParentList.Empty();
-		
-		FunctionParentList.Add(MaterialFunctionOriginal);
 
 		// Append function instance parent chain
 		UMaterialFunctionInstance* Current = MaterialFunctionOriginal;
@@ -1165,6 +1158,7 @@ void FMaterialInstanceEditor::RebuildInheritanceList()
 	}
 	else
 	{
+		MaterialChildList.Empty();
 		MaterialParentList.Empty();
 
 		// Travel up the parent chain for this material instance until we reach the root material.
@@ -1172,7 +1166,7 @@ void FMaterialInstanceEditor::RebuildInheritanceList()
 
 		if(InstanceConstant)
 		{
-			MaterialParentList.Add(InstanceConstant);
+			UMaterialEditingLibrary::GetChildInstances(InstanceConstant, MaterialChildList);
 
 			// Add all parents
 			UMaterialInterface* Parent = InstanceConstant->Parent;
diff --git a/Engine/Source/Editor/MaterialEditor/Private/MaterialInstanceEditor.h b/Engine/Source/Editor/MaterialEditor/Private/MaterialInstanceEditor.h
--- a/Engine/Source/Editor/MaterialEditor/Private/MaterialInstanceEditor.h
+++ b/Engine/Source/Editor/MaterialEditor/Private/MaterialInstanceEditor.h
@@ -144,12 +144,6 @@ private:
 	void ToggleShowAllMaterialParameters();
 	bool IsShowAllMaterialParametersChecked() const;
 
-	/** Commands for the Parents menu */
-	void OnOpenMaterial(TWeakObjectPtr<UMaterialInterface> InMaterial);
-	void OnOpenFunction(TWeakObjectPtr<UMaterialFunctionInterface> InFunction);
-	void OnShowMaterialInContentBrowser(TWeakObjectPtr<UMaterialInterface> InMaterial);
-	void OnShowFunctionInContentBrowser(TWeakObjectPtr<UMaterialFunctionInterface> InFunction);
-
 	/** Creates all internal widgets for the tabs to point at */
 	void CreateInternalWidgets();
 
@@ -218,10 +212,16 @@ private:
 	TSharedPtr<class SMaterialLayersFunctionsInstanceWrapper> MaterialLayersFunctionsInstance;
 
 	/** List of parents used to populate the inheritance list chain. */
-	TArray< TWeakObjectPtr<UMaterialInterface> > MaterialParentList;
+	TArray< FAssetData > MaterialParentList;
 
 	/** List of parents used to populate the inheritance list chain. */
-	TArray< TWeakObjectPtr<UMaterialFunctionInterface> > FunctionParentList;
+	TArray< FAssetData > FunctionParentList;
+
+	/** List of children used to populate the inheritance list chain. */
+	TArray< FAssetData > MaterialChildList;
+
+	/** List of children used to populate the inheritance list chain. */
+	TArray< FAssetData > FunctionChildList;
 
 	/** Object that stores all of the possible parameters we can edit. */
 	UMaterialEditorInstanceConstant* MaterialEditorInstance;
diff --git a/Engine/Source/Editor/MaterialEditor/Public/MaterialEditingLibrary.h b/Engine/Source/Editor/MaterialEditor/Public/MaterialEditingLibrary.h
--- a/Engine/Source/Editor/MaterialEditor/Public/MaterialEditingLibrary.h
+++ b/Engine/Source/Editor/MaterialEditor/Public/MaterialEditingLibrary.h
@@ -219,5 +219,5 @@ public:
 
 	/** Gets all direct child mat instances */
 	UFUNCTION(BlueprintCallable, Category = "MaterialEditing")
-	static void GetChildInstances(UMaterialInterface* Parent, TArray<TSoftObjectPtr<UMaterialInstance>>& ChildInstances);
+	static void GetChildInstances(UMaterialInterface* Parent, TArray<FAssetData>& ChildInstances);
 };
\ No newline at end of file
diff --git a/Engine/Source/Editor/MaterialEditor/Public/MaterialEditorUtilities.h b/Engine/Source/Editor/MaterialEditor/Public/MaterialEditorUtilities.h
--- a/Engine/Source/Editor/MaterialEditor/Public/MaterialEditorUtilities.h
+++ b/Engine/Source/Editor/MaterialEditor/Public/MaterialEditorUtilities.h
@@ -194,8 +194,18 @@ public:
 	/** Get IMaterialEditor for given object, if it exists */
 	static TSharedPtr<class IMaterialEditor> GetIMaterialEditorForObject(const UObject* ObjectToFocusOn);
 
+
+	/** Commands for the Parents menu */
+	static void OnOpenMaterial(FAssetData InMaterial);
+	static void OnOpenFunction(FAssetData InFunction);
+	static void OnShowMaterialInContentBrowser(FAssetData InMaterial);
+	static void OnShowFunctionInContentBrowser(FAssetData InFunction);
+
 private:
 
+	static void OpenSelectedParentEditor(UMaterialFunctionInterface* InMaterialFunction);
+	static void OpenSelectedParentEditor(UMaterialInterface* InMaterialInterface);
+
 	/**
 	 * Recursively walks the expression tree and parses the visible expression branches.
 	 *
