Adding getters for the material parameter names to the Material Editing Library
#rnx

diff --git a/Engine/Source/Editor/MaterialEditor/Private/MaterialEditingLibrary.cpp b/Engine/Source/Editor/MaterialEditor/Private/MaterialEditingLibrary.cpp
--- a/Engine/Source/Editor/MaterialEditor/Private/MaterialEditingLibrary.cpp
+++ b/Engine/Source/Editor/MaterialEditor/Private/MaterialEditingLibrary.cpp
@@ -982,3 +982,67 @@ void UMaterialEditingLibrary::GetChildInstances(UMaterialInterface* Parent, TArr
 		ChildInstances.Add(MatInstRef);
 	}
 }
+
+void UMaterialEditingLibrary::GetScalarParameterNames(UMaterialInterface* Material, TArray<FName>& ParameterNames)
+{
+	ParameterNames.Empty();
+	if (Material)
+	{
+		TArray<FMaterialParameterInfo> MaterialInfo;
+		TArray<FGuid> MaterialGuids;
+		Material->GetAllScalarParameterInfo(MaterialInfo, MaterialGuids);
+
+		for (const FMaterialParameterInfo& Info : MaterialInfo)
+		{
+			ParameterNames.Add(Info.Name);
+		}
+	}
+}
+
+void UMaterialEditingLibrary::GetVectorParameterNames(UMaterialInterface* Material, TArray<FName>& ParameterNames)
+{
+	ParameterNames.Empty();
+	if (Material)
+	{
+		TArray<FMaterialParameterInfo> MaterialInfo;
+		TArray<FGuid> MaterialGuids;
+		Material->GetAllVectorParameterInfo(MaterialInfo, MaterialGuids);
+
+		for (const FMaterialParameterInfo& Info : MaterialInfo)
+		{
+			ParameterNames.Add(Info.Name);
+		}
+	}
+}
+
+void UMaterialEditingLibrary::GetTextureParameterNames(UMaterialInterface* Material, TArray<FName>& ParameterNames)
+{
+	ParameterNames.Empty();
+	if (Material)
+	{
+		TArray<FMaterialParameterInfo> MaterialInfo;
+		TArray<FGuid> MaterialGuids;
+		Material->GetAllTextureParameterInfo(MaterialInfo, MaterialGuids);
+
+		for (const FMaterialParameterInfo& Info : MaterialInfo)
+		{
+			ParameterNames.Add(Info.Name);
+		}
+	}
+}
+
+void UMaterialEditingLibrary::GetStaticSwitchParameterNames(UMaterialInterface* Material, TArray<FName>& ParameterNames)
+{
+	ParameterNames.Empty();
+	if (Material)
+	{
+		TArray<FMaterialParameterInfo> MaterialInfo;
+		TArray<FGuid> MaterialGuids;
+		Material->GetAllStaticSwitchParameterInfo(MaterialInfo, MaterialGuids);
+
+		for (const FMaterialParameterInfo& Info : MaterialInfo)
+		{
+			ParameterNames.Add(Info.Name);
+		}
+	}
+}
diff --git a/Engine/Source/Editor/MaterialEditor/Public/MaterialEditingLibrary.h b/Engine/Source/Editor/MaterialEditor/Public/MaterialEditingLibrary.h
--- a/Engine/Source/Editor/MaterialEditor/Public/MaterialEditingLibrary.h
+++ b/Engine/Source/Editor/MaterialEditor/Public/MaterialEditingLibrary.h
@@ -220,4 +220,20 @@ public:
 	/** Gets all direct child mat instances */
 	UFUNCTION(BlueprintCallable, Category = "MaterialEditing")
 	static void GetChildInstances(UMaterialInterface* Parent, TArray<FAssetData>& ChildInstances);
+
+	/** Gets all scalar parameter names */
+	UFUNCTION(BlueprintCallable, Category = "MaterialEditing")
+	static void GetScalarParameterNames(UMaterialInterface* Material, TArray<FName>& ParameterNames);
+
+	/** Gets all vector parameter names */
+	UFUNCTION(BlueprintCallable, Category = "MaterialEditing")
+	static void GetVectorParameterNames(UMaterialInterface* Material, TArray<FName>& ParameterNames);
+
+	/** Gets all texture parameter names */
+	UFUNCTION(BlueprintCallable, Category = "MaterialEditing")
+	static void GetTextureParameterNames(UMaterialInterface* Material, TArray<FName>& ParameterNames);
+
+	/** Gets all static switch parameter names */
+	UFUNCTION(BlueprintCallable, Category = "MaterialEditing")
+	static void GetStaticSwitchParameterNames(UMaterialInterface* Material, TArray<FName>& ParameterNames);
 };
\ No newline at end of file
