Minor fix to material editor alignment shortcut

diff --git a/Engine/Source/Editor/MaterialEditor/Private/MaterialEditor.cpp b/Engine/Source/Editor/MaterialEditor/Private/MaterialEditor.cpp
--- a/Engine/Source/Editor/MaterialEditor/Private/MaterialEditor.cpp
+++ b/Engine/Source/Editor/MaterialEditor/Private/MaterialEditor.cpp
@@ -4362,7 +4362,7 @@ void FMaterialEditor::OnAlignBottom()
 {
 	if (GraphEditor.IsValid())
 	{
-		GraphEditor->OnAlignMiddle();
+		GraphEditor->OnAlignBottom();
 	}
 }
 
