Hiding material hierarchy in the material editor if editing a function, adding a "no instances found" item if there are no material instances

diff --git a/Engine/Source/Editor/MaterialEditor/Private/MaterialEditor.cpp b/Engine/Source/Editor/MaterialEditor/Private/MaterialEditor.cpp
--- a/Engine/Source/Editor/MaterialEditor/Private/MaterialEditor.cpp
+++ b/Engine/Source/Editor/MaterialEditor/Private/MaterialEditor.cpp
@@ -1169,16 +1169,19 @@ void FMaterialEditor::FillToolbar(FToolBarBuilder& ToolbarBuilder)
 			false);
 	}
 	ToolbarBuilder.EndSection();
-	ToolbarBuilder.BeginSection("Hierarchy");
+	if (!MaterialFunction)
 	{
-		ToolbarBuilder.AddComboButton(
-			FUIAction(),
-			FOnGetContent::CreateSP(this, &FMaterialEditor::GenerateInheritanceMenu),
-			LOCTEXT("Hierarchy", "Hierarchy"),
-			FText::GetEmpty(),
-			FSlateIcon(FSlateIcon(FEditorStyle::GetStyleSetName(), "BTEditor.SwitchToBehaviorTreeMode")),
-			false
-		);
+		ToolbarBuilder.BeginSection("Hierarchy");
+		{
+			ToolbarBuilder.AddComboButton(
+				FUIAction(),
+				FOnGetContent::CreateSP(this, &FMaterialEditor::GenerateInheritanceMenu),
+				LOCTEXT("Hierarchy", "Hierarchy"),
+				FText::GetEmpty(),
+				FSlateIcon(FSlateIcon(FEditorStyle::GetStyleSetName(), "BTEditor.SwitchToBehaviorTreeMode")),
+				false
+			);
+		}
 	}
 };
 
@@ -1210,6 +1213,13 @@ TSharedRef<SWidget> FMaterialEditor::GenerateInheritanceMenu()
 	{
 		const FName MaterialInstances = TEXT("MaterialInstances");
 		MenuBuilder.BeginSection(MaterialInstances, LOCTEXT("MaterialInstances", "Material Instances"));
+		if (MaterialChildList.Num() == 0)
+		{
+			const FText NoChildText = LOCTEXT("NoInstancesFound", "No Instances Found");
+			TSharedRef<SWidget> NoChildWidget = SNew(STextBlock)
+				.Text(NoChildText);
+			MenuBuilder.AddWidget(NoChildWidget, FText::GetEmpty());
+		}
 		for (FAssetData MaterialChild : MaterialChildList)
 		{
 			FFormatNamedArguments Args;
