Fix for assert for missing EndSection call

diff --git a/Engine/Source/Editor/MaterialEditor/Private/MaterialEditor.cpp b/Engine/Source/Editor/MaterialEditor/Private/MaterialEditor.cpp
--- a/Engine/Source/Editor/MaterialEditor/Private/MaterialEditor.cpp
+++ b/Engine/Source/Editor/MaterialEditor/Private/MaterialEditor.cpp
@@ -1182,6 +1182,7 @@ void FMaterialEditor::FillToolbar(FToolBarBuilder& ToolbarBuilder)
 				false
 			);
 		}
+		ToolbarBuilder.EndSection("Hierarchy");
 	}
 };
 
