Ensure that the transient instance created by the material editor is marked pending-kill once the material editor is closed

This fixes an issue where a lingering material editor instance would re-apply old material constant values after a content hot-reload, as the content hot-reload would call PostEditChange on the lingering instance.

diff --git a/Engine/Source/Editor/MaterialEditor/Private/MaterialInstanceEditor.cpp b/Engine/Source/Editor/MaterialEditor/Private/MaterialInstanceEditor.cpp
--- a/Engine/Source/Editor/MaterialEditor/Private/MaterialInstanceEditor.cpp
+++ b/Engine/Source/Editor/MaterialEditor/Private/MaterialInstanceEditor.cpp
@@ -616,7 +616,12 @@ FMaterialInstanceEditor::~FMaterialInstanceEditor()
 		}
 	}
 
-	MaterialEditorInstance = NULL;
+	if (MaterialEditorInstance)
+	{
+		MaterialEditorInstance->MarkPendingKill();
+		MaterialEditorInstance = nullptr;
+	}
+
 	MaterialParentList.Empty();
 	FunctionParentList.Empty();
 
