UE4.21 - Fix for rendering thread potentially using deleted materials

diff --git a/Engine/Source/Editor/MaterialEditor/Private/SMaterialEditorViewport.cpp b/Engine/Source/Editor/MaterialEditor/Private/SMaterialEditorViewport.cpp
--- a/Engine/Source/Editor/MaterialEditor/Private/SMaterialEditorViewport.cpp
+++ b/Engine/Source/Editor/MaterialEditor/Private/SMaterialEditorViewport.cpp
@@ -427,6 +427,7 @@ void SMaterialEditor3DPreviewViewport::SetPreviewMaterial(UMaterialInterface* In
 		if (PreviewMeshComponent != nullptr)
 		{
 			PreviewMeshComponent->OverrideMaterials.Empty();
+			PreviewMeshComponent->MarkRenderStateDirty();
 		}
 	}
 	else
@@ -436,6 +437,7 @@ void SMaterialEditor3DPreviewViewport::SetPreviewMaterial(UMaterialInterface* In
 		{
 			PreviewMeshComponent->OverrideMaterials.Empty();
 			PreviewMeshComponent->OverrideMaterials.Add(PreviewMaterial);
+			PreviewMeshComponent->MarkRenderStateDirty();
 		}
 		
 		PostProcessVolumeActor = nullptr;
