Fix for not being able to reopen details panel after close

diff --git a/Engine/Source/Editor/MaterialEditor/Private/MaterialEditor.cpp b/Engine/Source/Editor/MaterialEditor/Private/MaterialEditor.cpp
--- a/Engine/Source/Editor/MaterialEditor/Private/MaterialEditor.cpp
+++ b/Engine/Source/Editor/MaterialEditor/Private/MaterialEditor.cpp
@@ -3234,7 +3234,7 @@ TSharedRef<SDockTab> FMaterialEditor::SpawnTab_GraphCanvas(const FSpawnTabArgs&
 
 TSharedRef<SDockTab> FMaterialEditor::SpawnTab_MaterialProperties(const FSpawnTabArgs& Args)
 {
-	SpawnedDetailsTab = SNew(SDockTab)
+	TSharedPtr<SDockTab> DetailsTab = SNew(SDockTab)
 		.Icon( FEditorStyle::GetBrush("LevelEditor.Tabs.Details") )
 		.Label( LOCTEXT("MaterialDetailsTitle", "Details") )
 		[
@@ -3246,8 +3246,8 @@ TSharedRef<SDockTab> FMaterialEditor::SpawnTab_MaterialProperties(const FSpawnTa
 		// Since we're initialising, make sure nothing is selected
 		GraphEditor->ClearSelectionSet();
 	}
-
-	return SpawnedDetailsTab.ToSharedRef();
+	SpawnedDetailsTab = DetailsTab;
+	return DetailsTab.ToSharedRef();
 }
 
 TSharedRef<SDockTab> FMaterialEditor::SpawnTab_Palette(const FSpawnTabArgs& Args)
@@ -4750,9 +4750,9 @@ void FMaterialEditor::NotifyExternalMaterialChange()
 
 void FMaterialEditor::FocusDetailsPanel()
 {
-	if (SpawnedDetailsTab && !SpawnedDetailsTab->IsForeground())
+	if (SpawnedDetailsTab.IsValid() && !SpawnedDetailsTab.Pin()->IsForeground())
 	{
-		SpawnedDetailsTab->DrawAttention();
+		SpawnedDetailsTab.Pin()->DrawAttention();
 	}
 }
 
diff --git a/Engine/Source/Editor/MaterialEditor/Private/MaterialEditor.h b/Engine/Source/Editor/MaterialEditor/Private/MaterialEditor.h
--- a/Engine/Source/Editor/MaterialEditor/Private/MaterialEditor.h
+++ b/Engine/Source/Editor/MaterialEditor/Private/MaterialEditor.h
@@ -797,7 +797,7 @@ private:
 	TSharedPtr<class FMaterialStats> MaterialStatsManager;
 
 	/** Tab that holds the details panel */
-	TSharedPtr<SDockTab> SpawnedDetailsTab;
+	TWeakPtr<SDockTab> SpawnedDetailsTab;
 
 	/** Stores the quality level used to preview the material graph */
 	EMaterialQualityLevel::Type NodeQualityLevel;
