Fix for compile error when generating example layer and blend graphs

diff --git a/Engine/Source/Editor/MaterialEditor/Private/MaterialEditor.cpp b/Engine/Source/Editor/MaterialEditor/Private/MaterialEditor.cpp
--- a/Engine/Source/Editor/MaterialEditor/Private/MaterialEditor.cpp
+++ b/Engine/Source/Editor/MaterialEditor/Private/MaterialEditor.cpp
@@ -612,6 +612,7 @@ void FMaterialEditor::InitMaterialEditor( const EToolkitMode::Type Mode, const T
 						{
 							UMaterialEditingLibrary::ConnectMaterialExpressions(Input, FString(), SetMaterialAttributes, FString());
 							UMaterialEditingLibrary::ConnectMaterialExpressions(SetMaterialAttributes, FString(), Expression, FString());
+							bMaterialDirty = true;
 						}
 					}
 				}
@@ -643,6 +644,7 @@ void FMaterialEditor::InitMaterialEditor( const EToolkitMode::Type Mode, const T
 							UMaterialEditingLibrary::ConnectMaterialExpressions(InputBottom, FString(), BlendMaterialAttributes, FString(TEXT("A")));
 							UMaterialEditingLibrary::ConnectMaterialExpressions(InputTop, FString(), BlendMaterialAttributes, FString(TEXT("B")));
 							UMaterialEditingLibrary::ConnectMaterialExpressions(BlendMaterialAttributes, FString(), Expression, FString());
+							bMaterialDirty = true;
 						}
 					}
 				}
@@ -705,6 +707,11 @@ void FMaterialEditor::InitMaterialEditor( const EToolkitMode::Type Mode, const T
 		FSuppressableWarningDialog EditingDefaultMaterial( Info );
 		EditingDefaultMaterial.ShowModal();	
 	}
+
+	if (GetDefault<UEditorExperimentalSettings>()->bExampleLayersAndBlends && bMaterialDirty)
+	{
+		SaveAsset_Execute();
+	}
 }
 
 FMaterialEditor::FMaterialEditor()
