PR #4391: UE-36841: Add ScopedTransation to CreateNewMaterialExpressionComment (Contributed by [REDACTED_PERSON])

diff --git a/Engine/Source/Editor/MaterialEditor/Private/MaterialEditor.cpp b/Engine/Source/Editor/MaterialEditor/Private/MaterialEditor.cpp
--- a/Engine/Source/Editor/MaterialEditor/Private/MaterialEditor.cpp
+++ b/Engine/Source/Editor/MaterialEditor/Private/MaterialEditor.cpp
@@ -3327,6 +3327,7 @@ UMaterialExpressionComment* FMaterialEditor::CreateNewMaterialExpressionComment(
 {
 	UMaterialExpressionComment* NewComment = NULL;
 	{
+		const FScopedTransaction Transaction(NSLOCTEXT("UnrealEd", "MaterialEditorCreateComment", "Material Editor: Create comment"));
 		Material->Modify();
 
 		UObject* ExpressionOuter = Material;
@@ -3360,10 +3361,7 @@ UMaterialExpressionComment* FMaterialEditor::CreateNewMaterialExpressionComment(
 		}
 
 		NewComment->Text = NSLOCTEXT("K2Node", "CommentBlock_NewEmptyComment", "Comment").ToString();
-	}
 
-	if (NewComment)
-	{
 		Material->MaterialGraph->AddComment(NewComment, true);
 
 		// Select the new comment.
