Only create groups with visible parameters

diff --git a/Engine/Source/Editor/MaterialEditor/Private/MaterialEditorInstanceDetailCustomization.cpp b/Engine/Source/Editor/MaterialEditor/Private/MaterialEditorInstanceDetailCustomization.cpp
--- a/Engine/Source/Editor/MaterialEditor/Private/MaterialEditorInstanceDetailCustomization.cpp
+++ b/Engine/Source/Editor/MaterialEditor/Private/MaterialEditorInstanceDetailCustomization.cpp
@@ -253,8 +253,20 @@ void FMaterialInstanceParameterDetails::CreateGroupsWidget(TSharedRef<IPropertyH
 			&& ParameterGroup.GroupName != FMaterialPropertyHelpers::LayerParamName)
 		{
 			bShowSaveButtons = true;
-			IDetailGroup& DetailGroup = GroupsCategory.AddGroup(ParameterGroup.GroupName, FText::FromName(ParameterGroup.GroupName), false, true);
-			CreateSingleGroupWidget(ParameterGroup, ParameterGroupsProperty->GetChildHandle(GroupIdx), DetailGroup);
+			bool bCreateGroup = false;
+			for (int32 ParamIdx = 0; ParamIdx < ParameterGroup.Parameters.Num(); ++ParamIdx)
+			{
+				UDEditorParameterValue* Parameter = ParameterGroup.Parameters[ParamIdx];
+				if (MaterialEditorInstance->VisibleExpressions.Contains(Parameter->ParameterInfo))
+				{
+					bCreateGroup = true;
+				}
+			}
+			if (bCreateGroup)
+			{
+				IDetailGroup& DetailGroup = GroupsCategory.AddGroup(ParameterGroup.GroupName, FText::FromName(ParameterGroup.GroupName), false, true);
+				CreateSingleGroupWidget(ParameterGroup, ParameterGroupsProperty->GetChildHandle(GroupIdx), DetailGroup);
+			}
 		}
 	}
 	if (bShowSaveButtons)
