Hiding Show Modified Properties and adding custom Show Overridden Parameters in Material Instance editor. This works for all parameters and needs additional work for general properties (Physical Material, etc.).

diff --git a/Engine/Source/Editor/MaterialEditor/Private/MaterialEditorInstanceDetailCustomization.cpp b/Engine/Source/Editor/MaterialEditor/Private/MaterialEditorInstanceDetailCustomization.cpp
--- a/Engine/Source/Editor/MaterialEditor/Private/MaterialEditorInstanceDetailCustomization.cpp
+++ b/Engine/Source/Editor/MaterialEditor/Private/MaterialEditorInstanceDetailCustomization.cpp
@@ -257,9 +257,19 @@ void FMaterialInstanceParameterDetails::CreateGroupsWidget(TSharedRef<IPropertyH
 			for (int32 ParamIdx = 0; ParamIdx < ParameterGroup.Parameters.Num(); ++ParamIdx)
 			{
 				UDEditorParameterValue* Parameter = ParameterGroup.Parameters[ParamIdx];
-				if (MaterialEditorInstance->VisibleExpressions.Contains(Parameter->ParameterInfo))
+				if (MaterialEditorInstance->bShowOnlyOverrides)
 				{
-					bCreateGroup = true;
+					if (MaterialEditorInstance->VisibleExpressions.Contains(Parameter->ParameterInfo) && FMaterialPropertyHelpers::IsOverriddenExpression(Parameter))
+					{
+						bCreateGroup = true;
+					}
+				}
+				else
+				{
+					if (MaterialEditorInstance->VisibleExpressions.Contains(Parameter->ParameterInfo))
+					{
+						bCreateGroup = true;
+					}
 				}
 			}
 			if (bCreateGroup)
@@ -743,31 +753,46 @@ void FMaterialInstanceParameterDetails::CreateBasePropertyOverrideWidgets(IDetai
 	OpacityClipMaskValuePropertyRow
 		.DisplayName(OpacityClipMaskValueProperty->GetPropertyDisplayName())
 		.ToolTip(OpacityClipMaskValueProperty->GetToolTipText())
-		.EditCondition(IsOverrideOpacityClipMaskValueEnabled, FOnBooleanValueChanged::CreateSP(this, &FMaterialInstanceParameterDetails::OnOverrideOpacityClipMaskValueChanged));
+		.EditCondition(IsOverrideOpacityClipMaskValueEnabled, FOnBooleanValueChanged::CreateSP(this, &FMaterialInstanceParameterDetails::OnOverrideOpacityClipMaskValueChanged))
+		.Visibility(TAttribute<EVisibility>::Create(TAttribute<EVisibility>::FGetter::CreateSP(this, &FMaterialInstanceParameterDetails::IsOverriddenAndVisible, IsOverrideOpacityClipMaskValueEnabled)));
 
 	IDetailPropertyRow& BlendModePropertyRow = BasePropertyOverrideGroup.AddPropertyRow(BlendModeProperty.ToSharedRef());
 	BlendModePropertyRow
 		.DisplayName(BlendModeProperty->GetPropertyDisplayName())
 		.ToolTip(BlendModeProperty->GetToolTipText())
-		.EditCondition(IsOverrideBlendModeEnabled, FOnBooleanValueChanged::CreateSP(this, &FMaterialInstanceParameterDetails::OnOverrideBlendModeChanged));
+		.EditCondition(IsOverrideBlendModeEnabled, FOnBooleanValueChanged::CreateSP(this, &FMaterialInstanceParameterDetails::OnOverrideBlendModeChanged))
+		.Visibility(TAttribute<EVisibility>::Create(TAttribute<EVisibility>::FGetter::CreateSP(this, &FMaterialInstanceParameterDetails::IsOverriddenAndVisible, IsOverrideBlendModeEnabled)));
 
 	IDetailPropertyRow& ShadingModelPropertyRow = BasePropertyOverrideGroup.AddPropertyRow(ShadingModelProperty.ToSharedRef());
 	ShadingModelPropertyRow
 		.DisplayName(ShadingModelProperty->GetPropertyDisplayName())
 		.ToolTip(ShadingModelProperty->GetToolTipText())
-		.EditCondition(IsOverrideShadingModelEnabled, FOnBooleanValueChanged::CreateSP(this, &FMaterialInstanceParameterDetails::OnOverrideShadingModelChanged));
+		.EditCondition(IsOverrideShadingModelEnabled, FOnBooleanValueChanged::CreateSP(this, &FMaterialInstanceParameterDetails::OnOverrideShadingModelChanged))
+		.Visibility(TAttribute<EVisibility>::Create(TAttribute<EVisibility>::FGetter::CreateSP(this, &FMaterialInstanceParameterDetails::IsOverriddenAndVisible, IsOverrideShadingModelEnabled)));
 
 	IDetailPropertyRow& TwoSidedPropertyRow = BasePropertyOverrideGroup.AddPropertyRow(TwoSidedProperty.ToSharedRef());
 	TwoSidedPropertyRow
 		.DisplayName(TwoSidedProperty->GetPropertyDisplayName())
 		.ToolTip(TwoSidedProperty->GetToolTipText())
-		.EditCondition(IsOverrideTwoSidedEnabled, FOnBooleanValueChanged::CreateSP(this, &FMaterialInstanceParameterDetails::OnOverrideTwoSidedChanged));
+		.EditCondition(IsOverrideTwoSidedEnabled, FOnBooleanValueChanged::CreateSP(this, &FMaterialInstanceParameterDetails::OnOverrideTwoSidedChanged))
+		.Visibility(TAttribute<EVisibility>::Create(TAttribute<EVisibility>::FGetter::CreateSP(this, &FMaterialInstanceParameterDetails::IsOverriddenAndVisible, IsOverrideTwoSidedEnabled)));
 
 	IDetailPropertyRow& DitheredLODTransitionPropertyRow = BasePropertyOverrideGroup.AddPropertyRow(DitheredLODTransitionProperty.ToSharedRef());
 	DitheredLODTransitionPropertyRow
 		.DisplayName(DitheredLODTransitionProperty->GetPropertyDisplayName())
 		.ToolTip(DitheredLODTransitionProperty->GetToolTipText())
-		.EditCondition(IsOverrideDitheredLODTransitionEnabled, FOnBooleanValueChanged::CreateSP(this, &FMaterialInstanceParameterDetails::OnOverrideDitheredLODTransitionChanged));
+		.EditCondition(IsOverrideDitheredLODTransitionEnabled, FOnBooleanValueChanged::CreateSP(this, &FMaterialInstanceParameterDetails::OnOverrideDitheredLODTransitionChanged))
+		.Visibility(TAttribute<EVisibility>::Create(TAttribute<EVisibility>::FGetter::CreateSP(this, &FMaterialInstanceParameterDetails::IsOverriddenAndVisible, IsOverrideDitheredLODTransitionEnabled)));
+}
+
+EVisibility FMaterialInstanceParameterDetails::IsOverriddenAndVisible(TAttribute<bool> IsOverridden) const
+{
+	bool bShouldBeVisible = true;
+	if (MaterialEditorInstance->bShowOnlyOverrides)
+	{
+		bShouldBeVisible = IsOverridden.Get();
+	}
+	return bShouldBeVisible ? EVisibility::Visible : EVisibility::Collapsed;
 }
 
 bool FMaterialInstanceParameterDetails::OverrideOpacityClipMaskValueEnabled() const
diff --git a/Engine/Source/Editor/MaterialEditor/Private/MaterialEditorInstanceDetailCustomization.h b/Engine/Source/Editor/MaterialEditor/Private/MaterialEditorInstanceDetailCustomization.h
--- a/Engine/Source/Editor/MaterialEditor/Private/MaterialEditorInstanceDetailCustomization.h
+++ b/Engine/Source/Editor/MaterialEditor/Private/MaterialEditorInstanceDetailCustomization.h
@@ -79,6 +79,8 @@ private:
 	/** Creates all the base property override widgets. */
 	void CreateBasePropertyOverrideWidgets(IDetailLayoutBuilder& DetailLayout);
 
+	EVisibility IsOverriddenAndVisible(TAttribute<bool> IsOverridden) const;
+
 	bool OverrideOpacityClipMaskValueEnabled() const;
 	bool OverrideBlendModeEnabled() const;
 	bool OverrideShadingModelEnabled() const;
diff --git a/Engine/Source/Editor/MaterialEditor/Private/MaterialInstanceEditor.cpp b/Engine/Source/Editor/MaterialEditor/Private/MaterialInstanceEditor.cpp
--- a/Engine/Source/Editor/MaterialEditor/Private/MaterialInstanceEditor.cpp
+++ b/Engine/Source/Editor/MaterialEditor/Private/MaterialInstanceEditor.cpp
@@ -755,12 +755,15 @@ void FMaterialInstanceEditor::CreateInternalWidgets()
 	PreviewUIViewport = SNew(SMaterialEditorUIPreviewViewport, GetMaterialInterface());
 
 	FPropertyEditorModule& PropertyEditorModule = FModuleManager::GetModuleChecked<FPropertyEditorModule>( "PropertyEditor" );
-	const FDetailsViewArgs DetailsViewArgs( false, false, true, FDetailsViewArgs::HideNameArea, true, this );
+	FDetailsViewArgs DetailsViewArgs( false, false, true, FDetailsViewArgs::HideNameArea, true, this );
+	DetailsViewArgs.bShowModifiedPropertiesOption = false;
+	DetailsViewArgs.bShowCustomFilterOption = true;
 	MaterialInstanceDetails = PropertyEditorModule.CreateDetailView( DetailsViewArgs );
 	FOnGetDetailCustomizationInstance LayoutMICDetails = FOnGetDetailCustomizationInstance::CreateStatic( 
 		&FMaterialInstanceParameterDetails::MakeInstance, MaterialEditorInstance, FGetShowHiddenParameters::CreateSP(this, &FMaterialInstanceEditor::GetShowHiddenParameters) );
 	MaterialInstanceDetails->RegisterInstancedCustomPropertyLayout( UMaterialEditorInstanceConstant::StaticClass(), LayoutMICDetails );
-
+	MaterialInstanceDetails->SetCustomFilterLabel(LOCTEXT("ShowOverriddenOnly", "Show Only Overridden Parameters"));
+	MaterialInstanceDetails->SetCustomFilterDelegate(FSimpleDelegate::CreateSP(this, &FMaterialInstanceEditor::FilterOverriddenProperties));
 	MaterialEditorInstance->DetailsView = MaterialInstanceDetails;
 
 	IMaterialEditorModule* MaterialEditorModule = &FModuleManager::LoadModuleChecked<IMaterialEditorModule>("MaterialEditor");
@@ -771,6 +774,12 @@ void FMaterialInstanceEditor::CreateInternalWidgets()
 	}
 }
 
+void FMaterialInstanceEditor::FilterOverriddenProperties()
+{
+	MaterialEditorInstance->bShowOnlyOverrides = !MaterialEditorInstance->bShowOnlyOverrides;
+	MaterialInstanceDetails->ForceRefresh();
+}
+
 void FMaterialInstanceEditor::UpdatePreviewViewportsVisibility()
 {
 	UMaterial* PreviewMaterial = GetMaterialInterface()->GetBaseMaterial();
diff --git a/Engine/Source/Editor/MaterialEditor/Private/MaterialInstanceEditor.h b/Engine/Source/Editor/MaterialEditor/Private/MaterialInstanceEditor.h
--- a/Engine/Source/Editor/MaterialEditor/Private/MaterialInstanceEditor.h
+++ b/Engine/Source/Editor/MaterialEditor/Private/MaterialInstanceEditor.h
@@ -153,6 +153,9 @@ private:
 	/** Creates all internal widgets for the tabs to point at */
 	void CreateInternalWidgets();
 
+	/** Delegate for overriding Show Modified to Show Overridden instead */
+	void FilterOverriddenProperties();
+
 	/** Updates the 3D and UI preview viewport visibility based on material domain */
 	void UpdatePreviewViewportsVisibility();
 
diff --git a/Engine/Source/Editor/MaterialEditor/Private/MaterialPropertyHelpers.cpp b/Engine/Source/Editor/MaterialEditor/Private/MaterialPropertyHelpers.cpp
--- a/Engine/Source/Editor/MaterialEditor/Private/MaterialPropertyHelpers.cpp
+++ b/Engine/Source/Editor/MaterialEditor/Private/MaterialPropertyHelpers.cpp
@@ -101,7 +101,14 @@ EVisibility FMaterialPropertyHelpers::ShouldShowExpression(UDEditorParameterValu
 
 	ShowHiddenDelegate.ExecuteIfBound(bShowHidden);
 
-	return (bShowHidden || MaterialEditorInstance->VisibleExpressions.Contains(Parameter->ParameterInfo))? EVisibility::Visible: EVisibility::Collapsed;
+	const bool bShouldShowExpression = (bShowHidden || MaterialEditorInstance->VisibleExpressions.Contains(Parameter->ParameterInfo));
+
+	if (MaterialEditorInstance->bShowOnlyOverrides)
+	{
+		return (IsOverriddenExpression(Parameter) && bShouldShowExpression) ? EVisibility::Visible : EVisibility::Collapsed;
+	}
+
+	return bShouldShowExpression ? EVisibility::Visible: EVisibility::Collapsed;
 }
 
 void FMaterialPropertyHelpers::OnMaterialLayerAssetChanged(const struct FAssetData& InAssetData, int32 Index, EMaterialParameterAssociation MaterialType, TSharedPtr<class IPropertyHandle> InHandle, FMaterialLayersFunctions* InMaterialFunction)
diff --git a/Engine/Source/Editor/PropertyEditor/Private/SDetailsView.cpp b/Engine/Source/Editor/PropertyEditor/Private/SDetailsView.cpp
--- a/Engine/Source/Editor/PropertyEditor/Private/SDetailsView.cpp
+++ b/Engine/Source/Editor/PropertyEditor/Private/SDetailsView.cpp
@@ -75,6 +75,24 @@ void SDetailsView::Construct(const FArguments& InArgs)
 			);
 		}
 
+		if (DetailsViewArgs.bShowCustomFilterOption)
+		{
+			TAttribute<FText> CustomFilterLabelDelegate;
+			CustomFilterLabelDelegate.BindRaw(this, &SDetailsView::GetCustomFilterLabel);
+			DetailViewOptions.AddMenuEntry(
+				CustomFilterLabelDelegate,
+				FText(),
+				FSlateIcon(),
+				FUIAction(
+					FExecuteAction::CreateSP(this, &SDetailsView::OnCustomFilterClicked),
+					FCanExecuteAction(),
+					FIsActionChecked::CreateSP(this, &SDetailsView::IsCustomFilterChecked)
+				),
+				NAME_None,
+				EUserInterfaceActionType::ToggleButton
+			);
+		}
+
 		if( DetailsViewArgs.bShowDifferingPropertiesOption )
 		{
 			DetailViewOptions.AddMenuEntry(
diff --git a/Engine/Source/Editor/PropertyEditor/Private/SDetailsViewBase.cpp b/Engine/Source/Editor/PropertyEditor/Private/SDetailsViewBase.cpp
--- a/Engine/Source/Editor/PropertyEditor/Private/SDetailsViewBase.cpp
+++ b/Engine/Source/Editor/PropertyEditor/Private/SDetailsViewBase.cpp
@@ -525,6 +525,15 @@ void SDetailsViewBase::OnShowOnlyModifiedClicked()
 	UpdateFilteredDetails();
 }
 
+void SDetailsViewBase::OnCustomFilterClicked()
+{
+	if (CustomFilterDelegate.IsBound())
+	{
+		CustomFilterDelegate.Execute();
+		bCustomFilterActive = !bCustomFilterActive;
+	}
+}
+
 void SDetailsViewBase::OnShowAllAdvancedClicked()
 {
 	CurrentFilter.bShowAllAdvanced = !CurrentFilter.bShowAllAdvanced;
@@ -609,7 +618,6 @@ void SDetailsViewBase::SetHostTabManager(TSharedPtr<FTabManager> InTabManager)
 	DetailsViewArgs.HostTabManager = InTabManager;
 }
 
-
 /** 
  * Hides or shows properties based on the passed in filter text
  * 
diff --git a/Engine/Source/Editor/PropertyEditor/Private/SDetailsViewBase.h b/Engine/Source/Editor/PropertyEditor/Private/SDetailsViewBase.h
--- a/Engine/Source/Editor/PropertyEditor/Private/SDetailsViewBase.h
+++ b/Engine/Source/Editor/PropertyEditor/Private/SDetailsViewBase.h
@@ -269,6 +269,9 @@ protected:
 	/** @return true if show only modified is checked */
 	bool IsShowOnlyModifiedChecked() const { return CurrentFilter.bShowOnlyModifiedProperties; }
 
+	/** @return true if custom filter is checked */
+	bool IsCustomFilterChecked() const { return bCustomFilterActive; }
+
 	/** @return true if show all advanced is checked */
 	bool IsShowAllAdvancedChecked() const { return CurrentFilter.bShowAllAdvanced; }
 
@@ -287,6 +290,9 @@ protected:
 	/** Called when show only modified is clicked */
 	void OnShowOnlyModifiedClicked();
 
+	/** Called when custom filter is clicked */
+	void OnCustomFilterClicked();
+
 	/** Called when show all advanced is clicked */
 	void OnShowAllAdvancedClicked();
 
@@ -332,6 +338,22 @@ protected:
 	/** Called to get the visibility of the filter box */
 	EVisibility GetFilterBoxVisibility() const;
 
+
+	virtual void SetCustomFilterDelegate(FSimpleDelegate InOverride) override
+	{
+		CustomFilterDelegate = InOverride;
+	}
+
+	virtual void SetCustomFilterLabel(FText InText) override
+	{
+		CustomFilterLabel = InText;
+	}
+
+	FText GetCustomFilterLabel() const
+	{
+		return CustomFilterLabel;
+	}
+
 protected:
 	/** The user defined args for the details view */
 	FDetailsViewArgs DetailsViewArgs;
@@ -405,4 +427,12 @@ protected:
 	bool bDisableCustomDetailLayouts;
 
 	int32 NumVisbleTopLevelObjectNodes;
+
+	/** Delegate for overriding the show modified filter */
+	FSimpleDelegate CustomFilterDelegate;
+
+	/** When overriding show modified, you can't use the filter state to determine what is overridden anymore. Use this variable instead. */
+	bool bCustomFilterActive;
+
+	FText CustomFilterLabel;
 };
diff --git a/Engine/Source/Editor/PropertyEditor/Public/IDetailsView.h b/Engine/Source/Editor/PropertyEditor/Public/IDetailsView.h
--- a/Engine/Source/Editor/PropertyEditor/Public/IDetailsView.h
+++ b/Engine/Source/Editor/PropertyEditor/Public/IDetailsView.h
@@ -94,6 +94,8 @@ struct FDetailsViewArgs
 	uint32 bShowKeyablePropertiesOption : 1;
 	/** True if you want to show the 'Show Only Animated Properties'. Only valid in conjunction with bShowOptions */
 	uint32 bShowAnimatedPropertiesOption: 1;
+	/** True if you want to show a custom filter. */
+	uint32 bShowCustomFilterOption : 1;
 	/** The default column width */
 	float ColumnWidth;
 
@@ -334,4 +336,9 @@ public:
 
 	/** Force refresh */
 	virtual void ForceRefresh() = 0;
+	
+	/** Allows other systems to add a custom filter in the details panel */
+	virtual void SetCustomFilterDelegate(FSimpleDelegate InDelegate) = 0;
+
+	virtual void SetCustomFilterLabel(FText InText) = 0;
 };
diff --git a/Engine/Source/Editor/UnrealEd/Classes/MaterialEditor/MaterialEditorInstanceConstant.h b/Engine/Source/Editor/UnrealEd/Classes/MaterialEditor/MaterialEditorInstanceConstant.h
--- a/Engine/Source/Editor/UnrealEd/Classes/MaterialEditor/MaterialEditorInstanceConstant.h
+++ b/Engine/Source/Editor/UnrealEd/Classes/MaterialEditor/MaterialEditorInstanceConstant.h
@@ -331,5 +331,8 @@ class UNREALED_API UMaterialEditorInstanceConstant : public UObject
 	UPROPERTY()
 	TArray<class UMaterialInstanceConstant*> StoredBlendPreviews;
 #endif
+
+	/** Whether or not we should show only overridden properties*/
+	bool bShowOnlyOverrides;
 };
 
diff --git a/Engine/Source/Editor/UnrealEd/Private/PreviewMaterial.cpp b/Engine/Source/Editor/UnrealEd/Private/PreviewMaterial.cpp
--- a/Engine/Source/Editor/UnrealEd/Private/PreviewMaterial.cpp
+++ b/Engine/Source/Editor/UnrealEd/Private/PreviewMaterial.cpp
@@ -789,6 +789,7 @@ UMaterialEditorInstanceConstant::UMaterialEditorInstanceConstant(const FObjectIn
 	: Super(ObjectInitializer)
 {
 	bIsFunctionPreviewMaterial = false;
+	bShowOnlyOverrides = false;
 }
 
 void UMaterialEditorInstanceConstant::PostEditChangeProperty(FPropertyChangedEvent& PropertyChangedEvent)
diff --git a/Engine/Source/Runtime/Engine/Private/Materials/Material.cpp b/Engine/Source/Runtime/Engine/Private/Materials/Material.cpp
--- a/Engine/Source/Runtime/Engine/Private/Materials/Material.cpp
+++ b/Engine/Source/Runtime/Engine/Private/Materials/Material.cpp
@@ -2843,7 +2843,11 @@ FMaterialRenderProxy* UMaterial::GetRenderProxy(bool Selected,bool Hovered) cons
 
 UPhysicalMaterial* UMaterial::GetPhysicalMaterial() const
 {
-	return (PhysMaterial != NULL) ? PhysMaterial : GEngine->DefaultPhysMaterial;
+	if (GEngine)
+	{
+		return (PhysMaterial != nullptr) ? PhysMaterial : GEngine->DefaultPhysMaterial;
+	}
+	return nullptr;
 }
 
 /** Helper functions for text output of properties... */
