Lightmass settings now display in the same override style as the rest of the material instance properties. Base properties are read from parent on load if they are not overwritten.

diff --git a/Engine/Source/Editor/MaterialEditor/Private/MaterialEditorInstanceDetailCustomization.cpp b/Engine/Source/Editor/MaterialEditor/Private/MaterialEditorInstanceDetailCustomization.cpp
--- a/Engine/Source/Editor/MaterialEditor/Private/MaterialEditorInstanceDetailCustomization.cpp
+++ b/Engine/Source/Editor/MaterialEditor/Private/MaterialEditorInstanceDetailCustomization.cpp
@@ -208,7 +208,8 @@ void FMaterialInstanceParameterDetails::CustomizeDetails(IDetailLayoutBuilder& D
 
 
 		// Add/hide other properties
-		DefaultCategory.AddProperty("LightmassSettings");
+		DetailLayout.HideProperty("LightmassSettings");
+		CreateLightmassOverrideWidgets(DetailLayout);
 		DetailLayout.HideProperty("bUseOldStyleMICEditorGroups");
 		DetailLayout.HideProperty("ParameterGroups");
 
@@ -729,6 +730,70 @@ EVisibility FMaterialInstanceParameterDetails::ShouldShowSubsurfaceProfile() con
 }
 
 
+void FMaterialInstanceParameterDetails::CreateLightmassOverrideWidgets(IDetailLayoutBuilder& DetailLayout)
+{
+	IDetailCategoryBuilder& DetailCategory = DetailLayout.EditCategory(NAME_None);
+
+	static FName GroupName(TEXT("LightmassSettings"));
+	IDetailGroup& LightmassSettingsGroup = DetailCategory.AddGroup(GroupName, LOCTEXT("LightmassSettingsGroup", "Lightmass Settings"), false, false);
+
+	TAttribute<bool> IsOverrideCastShadowAsMaskedEnabled = TAttribute<bool>::Create(TAttribute<bool>::FGetter::CreateLambda([this] { return MaterialEditorInstance->LightmassSettings.CastShadowAsMasked.bOverride; }));
+	TAttribute<bool> IsOverrideEmissiveBoostEnabled = TAttribute<bool>::Create(TAttribute<bool>::FGetter::CreateLambda([this] { return MaterialEditorInstance->LightmassSettings.EmissiveBoost.bOverride; }));
+	TAttribute<bool> IsOverrideDiffuseBoostEnabled = TAttribute<bool>::Create(TAttribute<bool>::FGetter::CreateLambda([this] { return MaterialEditorInstance->LightmassSettings.DiffuseBoost.bOverride; }));
+	TAttribute<bool> IsOverrideExportResolutionScaleEnabled = TAttribute<bool>::Create(TAttribute<bool>::FGetter::CreateLambda([this] { return MaterialEditorInstance->LightmassSettings.ExportResolutionScale.bOverride; }));
+	
+	TSharedRef<IPropertyHandle> LightmassSettings = DetailLayout.GetProperty("LightmassSettings");
+	TSharedPtr<IPropertyHandle> CastShadowAsMaskedProperty = LightmassSettings->GetChildHandle("CastShadowAsMasked");
+	TSharedPtr<IPropertyHandle> EmissiveBoostProperty = LightmassSettings->GetChildHandle("EmissiveBoost");
+	TSharedPtr<IPropertyHandle> DiffuseBoostProperty = LightmassSettings->GetChildHandle("DiffuseBoost");
+	TSharedPtr<IPropertyHandle> ExportResolutionScaleProperty = LightmassSettings->GetChildHandle("ExportResolutionScale");
+
+
+	IDetailPropertyRow& CastShadowAsMaskedPropertyRow = LightmassSettingsGroup.AddPropertyRow(CastShadowAsMaskedProperty->GetChildHandle(0).ToSharedRef());
+	CastShadowAsMaskedPropertyRow
+		.DisplayName(CastShadowAsMaskedProperty->GetPropertyDisplayName())
+		.ToolTip(CastShadowAsMaskedProperty->GetToolTipText())
+		.EditCondition(IsOverrideCastShadowAsMaskedEnabled, FOnBooleanValueChanged::CreateLambda([this](bool NewValue) {
+			MaterialEditorInstance->LightmassSettings.CastShadowAsMasked.bOverride = NewValue;
+			MaterialEditorInstance->PostEditChange();
+			FEditorSupportDelegates::RedrawAllViewports.Broadcast();
+		}))
+		.Visibility(TAttribute<EVisibility>::Create(TAttribute<EVisibility>::FGetter::CreateSP(this, &FMaterialInstanceParameterDetails::IsOverriddenAndVisible, IsOverrideCastShadowAsMaskedEnabled)));
+
+	IDetailPropertyRow& EmissiveBoostPropertyRow = LightmassSettingsGroup.AddPropertyRow(EmissiveBoostProperty->GetChildHandle(0).ToSharedRef());
+	EmissiveBoostPropertyRow
+		.DisplayName(EmissiveBoostProperty->GetPropertyDisplayName())
+		.ToolTip(EmissiveBoostProperty->GetToolTipText())
+		.EditCondition(IsOverrideEmissiveBoostEnabled, FOnBooleanValueChanged::CreateLambda([this](bool NewValue) {
+			MaterialEditorInstance->LightmassSettings.EmissiveBoost.bOverride = NewValue;
+			MaterialEditorInstance->PostEditChange();
+			FEditorSupportDelegates::RedrawAllViewports.Broadcast();
+		}))
+		.Visibility(TAttribute<EVisibility>::Create(TAttribute<EVisibility>::FGetter::CreateSP(this, &FMaterialInstanceParameterDetails::IsOverriddenAndVisible, IsOverrideEmissiveBoostEnabled)));
+	
+	IDetailPropertyRow& DiffuseBoostPropertyRow = LightmassSettingsGroup.AddPropertyRow(DiffuseBoostProperty->GetChildHandle(0).ToSharedRef());
+	DiffuseBoostPropertyRow
+		.DisplayName(DiffuseBoostProperty->GetPropertyDisplayName())
+		.ToolTip(DiffuseBoostProperty->GetToolTipText())
+		.EditCondition(IsOverrideDiffuseBoostEnabled, FOnBooleanValueChanged::CreateLambda([this](bool NewValue) {
+			MaterialEditorInstance->LightmassSettings.DiffuseBoost.bOverride = NewValue;
+			MaterialEditorInstance->PostEditChange();
+			FEditorSupportDelegates::RedrawAllViewports.Broadcast();
+		}))
+		.Visibility(TAttribute<EVisibility>::Create(TAttribute<EVisibility>::FGetter::CreateSP(this, &FMaterialInstanceParameterDetails::IsOverriddenAndVisible, IsOverrideDiffuseBoostEnabled)));
+
+	IDetailPropertyRow& ExportResolutionScalePropertyRow = LightmassSettingsGroup.AddPropertyRow(ExportResolutionScaleProperty->GetChildHandle(0).ToSharedRef());
+	ExportResolutionScalePropertyRow
+		.DisplayName(ExportResolutionScaleProperty->GetPropertyDisplayName())
+		.ToolTip(ExportResolutionScaleProperty->GetToolTipText())
+		.EditCondition(IsOverrideExportResolutionScaleEnabled, FOnBooleanValueChanged::CreateLambda([this](bool NewValue) {
+			MaterialEditorInstance->LightmassSettings.ExportResolutionScale.bOverride = NewValue;
+			MaterialEditorInstance->PostEditChange();
+			FEditorSupportDelegates::RedrawAllViewports.Broadcast();
+		}))
+		.Visibility(TAttribute<EVisibility>::Create(TAttribute<EVisibility>::FGetter::CreateSP(this, &FMaterialInstanceParameterDetails::IsOverriddenAndVisible, IsOverrideExportResolutionScaleEnabled)));
+}
+
 void FMaterialInstanceParameterDetails::CreateBasePropertyOverrideWidgets(IDetailLayoutBuilder& DetailLayout)
 {
 	IDetailCategoryBuilder& DetailCategory = DetailLayout.EditCategory(NAME_None);
diff --git a/Engine/Source/Editor/MaterialEditor/Private/MaterialEditorInstanceDetailCustomization.h b/Engine/Source/Editor/MaterialEditor/Private/MaterialEditorInstanceDetailCustomization.h
--- a/Engine/Source/Editor/MaterialEditor/Private/MaterialEditorInstanceDetailCustomization.h
+++ b/Engine/Source/Editor/MaterialEditor/Private/MaterialEditorInstanceDetailCustomization.h
@@ -73,7 +73,9 @@ private:
 	/** Returns true if the refraction options should be displayed */
 	EVisibility ShouldShowSubsurfaceProfile() const;
 
-	
+	/** Creates all the lightmass property override widgets. */
+	void CreateLightmassOverrideWidgets(IDetailLayoutBuilder& DetailLayout);
+
 	//Functions supporting BasePropertyOverrides
 
 	/** Creates all the base property override widgets. */
diff --git a/Engine/Source/Runtime/Engine/Private/Materials/MaterialInstance.cpp b/Engine/Source/Runtime/Engine/Private/Materials/MaterialInstance.cpp
--- a/Engine/Source/Runtime/Engine/Private/Materials/MaterialInstance.cpp
+++ b/Engine/Source/Runtime/Engine/Private/Materials/MaterialInstance.cpp
@@ -2395,7 +2395,7 @@ void UMaterialInstance::InitStaticPermutation()
 
 void UMaterialInstance::UpdateOverridableBaseProperties()
 {
-	//Parents base property overrides have to be cashed by now.
+	//Parents base property overrides have to be cached by now.
 	//This should be done on PostLoad()
 	//Or via an FMaterialUpdateContext when editing.
 
@@ -2416,6 +2416,7 @@ void UMaterialInstance::UpdateOverridableBaseProperties()
 	else
 	{
 		OpacityMaskClipValue = Parent->GetOpacityMaskClipValue();
+		BasePropertyOverrides.OpacityMaskClipValue = OpacityMaskClipValue;
 	}
 
 	if ( BasePropertyOverrides.bOverride_CastDynamicShadowAsMasked )
@@ -2425,6 +2426,7 @@ void UMaterialInstance::UpdateOverridableBaseProperties()
 	else
 	{
 		bCastDynamicShadowAsMasked = Parent->GetCastDynamicShadowAsMasked();
+		BasePropertyOverrides.bCastDynamicShadowAsMasked = bCastDynamicShadowAsMasked;
 	}
 
 	if (BasePropertyOverrides.bOverride_BlendMode)
@@ -2434,6 +2436,7 @@ void UMaterialInstance::UpdateOverridableBaseProperties()
 	else
 	{
 		BlendMode = Parent->GetBlendMode();
+		BasePropertyOverrides.BlendMode = BlendMode;
 	}
 
 	if (BasePropertyOverrides.bOverride_ShadingModel)
@@ -2443,6 +2446,7 @@ void UMaterialInstance::UpdateOverridableBaseProperties()
 	else
 	{
 		ShadingModel = Parent->GetShadingModel();
+		BasePropertyOverrides.ShadingModel = ShadingModel;
 	}
 
 	if (BasePropertyOverrides.bOverride_TwoSided)
@@ -2452,6 +2456,7 @@ void UMaterialInstance::UpdateOverridableBaseProperties()
 	else
 	{
 		TwoSided = Parent->IsTwoSided();
+		BasePropertyOverrides.TwoSided = TwoSided;
 	}
 
 	if (BasePropertyOverrides.bOverride_DitheredLODTransition)
@@ -2461,6 +2466,7 @@ void UMaterialInstance::UpdateOverridableBaseProperties()
 	else
 	{
 		DitheredLODTransition = Parent->IsDitheredLODTransition();
+		BasePropertyOverrides.DitheredLODTransition = DitheredLODTransition;
 	}
 }
 
