Adding uint32<->bool casts

diff --git a/Engine/Source/Editor/MaterialEditor/Private/MaterialEditorInstanceDetailCustomization.cpp b/Engine/Source/Editor/MaterialEditor/Private/MaterialEditorInstanceDetailCustomization.cpp
--- a/Engine/Source/Editor/MaterialEditor/Private/MaterialEditorInstanceDetailCustomization.cpp
+++ b/Engine/Source/Editor/MaterialEditor/Private/MaterialEditorInstanceDetailCustomization.cpp
@@ -737,10 +737,10 @@ void FMaterialInstanceParameterDetails::CreateLightmassOverrideWidgets(IDetailLa
 	static FName GroupName(TEXT("LightmassSettings"));
 	IDetailGroup& LightmassSettingsGroup = DetailCategory.AddGroup(GroupName, LOCTEXT("LightmassSettingsGroup", "Lightmass Settings"), false, false);
 
-	TAttribute<bool> IsOverrideCastShadowAsMaskedEnabled = TAttribute<bool>::Create(TAttribute<bool>::FGetter::CreateLambda([this] { return MaterialEditorInstance->LightmassSettings.CastShadowAsMasked.bOverride; }));
-	TAttribute<bool> IsOverrideEmissiveBoostEnabled = TAttribute<bool>::Create(TAttribute<bool>::FGetter::CreateLambda([this] { return MaterialEditorInstance->LightmassSettings.EmissiveBoost.bOverride; }));
-	TAttribute<bool> IsOverrideDiffuseBoostEnabled = TAttribute<bool>::Create(TAttribute<bool>::FGetter::CreateLambda([this] { return MaterialEditorInstance->LightmassSettings.DiffuseBoost.bOverride; }));
-	TAttribute<bool> IsOverrideExportResolutionScaleEnabled = TAttribute<bool>::Create(TAttribute<bool>::FGetter::CreateLambda([this] { return MaterialEditorInstance->LightmassSettings.ExportResolutionScale.bOverride; }));
+	TAttribute<bool> IsOverrideCastShadowAsMaskedEnabled = TAttribute<bool>::Create(TAttribute<bool>::FGetter::CreateLambda([this] { return (bool)MaterialEditorInstance->LightmassSettings.CastShadowAsMasked.bOverride; }));
+	TAttribute<bool> IsOverrideEmissiveBoostEnabled = TAttribute<bool>::Create(TAttribute<bool>::FGetter::CreateLambda([this] { return (bool)MaterialEditorInstance->LightmassSettings.EmissiveBoost.bOverride; }));
+	TAttribute<bool> IsOverrideDiffuseBoostEnabled = TAttribute<bool>::Create(TAttribute<bool>::FGetter::CreateLambda([this] { return (bool)MaterialEditorInstance->LightmassSettings.DiffuseBoost.bOverride; }));
+	TAttribute<bool> IsOverrideExportResolutionScaleEnabled = TAttribute<bool>::Create(TAttribute<bool>::FGetter::CreateLambda([this] { return (bool)MaterialEditorInstance->LightmassSettings.ExportResolutionScale.bOverride; }));
 	
 	TSharedRef<IPropertyHandle> LightmassSettings = DetailLayout.GetProperty("LightmassSettings");
 	TSharedPtr<IPropertyHandle> CastShadowAsMaskedProperty = LightmassSettings->GetChildHandle("CastShadowAsMasked");
@@ -754,7 +754,7 @@ void FMaterialInstanceParameterDetails::CreateLightmassOverrideWidgets(IDetailLa
 		.DisplayName(CastShadowAsMaskedProperty->GetPropertyDisplayName())
 		.ToolTip(CastShadowAsMaskedProperty->GetToolTipText())
 		.EditCondition(IsOverrideCastShadowAsMaskedEnabled, FOnBooleanValueChanged::CreateLambda([this](bool NewValue) {
-			MaterialEditorInstance->LightmassSettings.CastShadowAsMasked.bOverride = NewValue;
+			MaterialEditorInstance->LightmassSettings.CastShadowAsMasked.bOverride = (uint32)NewValue;
 			MaterialEditorInstance->PostEditChange();
 			FEditorSupportDelegates::RedrawAllViewports.Broadcast();
 		}))
@@ -765,7 +765,7 @@ void FMaterialInstanceParameterDetails::CreateLightmassOverrideWidgets(IDetailLa
 		.DisplayName(EmissiveBoostProperty->GetPropertyDisplayName())
 		.ToolTip(EmissiveBoostProperty->GetToolTipText())
 		.EditCondition(IsOverrideEmissiveBoostEnabled, FOnBooleanValueChanged::CreateLambda([this](bool NewValue) {
-			MaterialEditorInstance->LightmassSettings.EmissiveBoost.bOverride = NewValue;
+			MaterialEditorInstance->LightmassSettings.EmissiveBoost.bOverride = (uint32)NewValue;
 			MaterialEditorInstance->PostEditChange();
 			FEditorSupportDelegates::RedrawAllViewports.Broadcast();
 		}))
@@ -776,7 +776,7 @@ void FMaterialInstanceParameterDetails::CreateLightmassOverrideWidgets(IDetailLa
 		.DisplayName(DiffuseBoostProperty->GetPropertyDisplayName())
 		.ToolTip(DiffuseBoostProperty->GetToolTipText())
 		.EditCondition(IsOverrideDiffuseBoostEnabled, FOnBooleanValueChanged::CreateLambda([this](bool NewValue) {
-			MaterialEditorInstance->LightmassSettings.DiffuseBoost.bOverride = NewValue;
+			MaterialEditorInstance->LightmassSettings.DiffuseBoost.bOverride = (uint32)NewValue;
 			MaterialEditorInstance->PostEditChange();
 			FEditorSupportDelegates::RedrawAllViewports.Broadcast();
 		}))
@@ -787,7 +787,7 @@ void FMaterialInstanceParameterDetails::CreateLightmassOverrideWidgets(IDetailLa
 		.DisplayName(ExportResolutionScaleProperty->GetPropertyDisplayName())
 		.ToolTip(ExportResolutionScaleProperty->GetToolTipText())
 		.EditCondition(IsOverrideExportResolutionScaleEnabled, FOnBooleanValueChanged::CreateLambda([this](bool NewValue) {
-			MaterialEditorInstance->LightmassSettings.ExportResolutionScale.bOverride = NewValue;
+			MaterialEditorInstance->LightmassSettings.ExportResolutionScale.bOverride = (uint32)NewValue;
 			MaterialEditorInstance->PostEditChange();
 			FEditorSupportDelegates::RedrawAllViewports.Broadcast();
 		}))
