Fixing missing conversions to/from texture object param, hiding convert to param on texture sample param

diff --git a/Engine/Source/Editor/MaterialEditor/Private/MaterialEditor.cpp b/Engine/Source/Editor/MaterialEditor/Private/MaterialEditor.cpp
--- a/Engine/Source/Editor/MaterialEditor/Private/MaterialEditor.cpp
+++ b/Engine/Source/Editor/MaterialEditor/Private/MaterialEditor.cpp
@@ -2469,10 +2469,12 @@ void FMaterialEditor::OnConvertObjects()
 				UMaterialExpressionConstant3Vector* Constant3Expression = Cast<UMaterialExpressionConstant3Vector>(CurrentSelectedExpression);
 				UMaterialExpressionConstant4Vector* Constant4Expression = Cast<UMaterialExpressionConstant4Vector>(CurrentSelectedExpression);
 				UMaterialExpressionTextureSample* TextureSampleExpression = Cast<UMaterialExpressionTextureSample>(CurrentSelectedExpression);
+				UMaterialExpressionTextureObject* TextureObjectExpression = Cast<UMaterialExpressionTextureObject>(CurrentSelectedExpression);
 				UMaterialExpressionComponentMask* ComponentMaskExpression = Cast<UMaterialExpressionComponentMask>(CurrentSelectedExpression);
 				UMaterialExpressionParticleSubUV* ParticleSubUVExpression = Cast<UMaterialExpressionParticleSubUV>(CurrentSelectedExpression);
 				UMaterialExpressionScalarParameter* ScalarParameterExpression = Cast<UMaterialExpressionScalarParameter>(CurrentSelectedExpression);
 				UMaterialExpressionVectorParameter* VectorParameterExpression = Cast<UMaterialExpressionVectorParameter>(CurrentSelectedExpression);
+				UMaterialExpressionTextureObjectParameter* TextureObjectParameterExpression = Cast<UMaterialExpressionTextureObjectParameter>(CurrentSelectedExpression);
 
 				// Setup the class to convert to
 				UClass* ClassToCreate = NULL;
@@ -2491,11 +2493,19 @@ void FMaterialEditor::OnConvertObjects()
 				else if (TextureSampleExpression && TextureSampleExpression->Texture && TextureSampleExpression->Texture->IsA(UTextureCube::StaticClass()))
 				{
 					ClassToCreate = UMaterialExpressionTextureSampleParameterCube::StaticClass();
-				}	
+				}
+				else if (TextureObjectExpression)
+				{
+					ClassToCreate = UMaterialExpressionTextureObjectParameter::StaticClass();
+				}
+				else if (TextureObjectParameterExpression) // Has to come before TextureSample comparison 
+				{
+					ClassToCreate = UMaterialExpressionTextureObject::StaticClass();
+				}
 				else if (TextureSampleExpression)
 				{
 					ClassToCreate = UMaterialExpressionTextureSampleParameter2D::StaticClass();
-				}	
+				}
 				else if (ComponentMaskExpression)
 				{
 					ClassToCreate = UMaterialExpressionStaticComponentMaskParameter::StaticClass();
@@ -2509,6 +2519,7 @@ void FMaterialEditor::OnConvertObjects()
 					// Technically should be a constant 4 but UMaterialExpressionVectorParameter has an rgb pin, so using Constant3 to avoid a compile error
 					ClassToCreate = UMaterialExpressionConstant3Vector::StaticClass();
 				}
+			
 
 				if (ClassToCreate)
 				{
@@ -2549,7 +2560,7 @@ void FMaterialEditor::OnConvertObjects()
 							bNeedsRefresh = true;
 							CastChecked<UMaterialExpressionVectorParameter>(NewExpression)->DefaultValue = Constant4Expression->Constant;
 						}
-						else if (TextureSampleExpression)
+						else if (TextureSampleExpression && !TextureObjectParameterExpression)
 						{
 							bNeedsRefresh = true;
 							UMaterialExpressionTextureSampleParameter* NewTextureExpr = CastChecked<UMaterialExpressionTextureSampleParameter>(NewExpression);
@@ -3953,6 +3964,79 @@ void FMaterialEditor::RedoGraphAction()
 
 }
 
+void FMaterialEditor::OnAlignTop()
+{
+	if (GraphEditor.IsValid())
+	{
+		GraphEditor->OnAlignTop();
+	}
+}
+
+void FMaterialEditor::OnAlignMiddle()
+{
+	if (GraphEditor.IsValid())
+	{
+		GraphEditor->OnAlignMiddle();
+	}
+}
+
+void FMaterialEditor::OnAlignBottom()
+{
+	if (GraphEditor.IsValid())
+	{
+		GraphEditor->OnAlignMiddle();
+	}
+}
+
+void FMaterialEditor::OnAlignLeft()
+{
+	if (GraphEditor.IsValid())
+	{
+		GraphEditor->OnAlignLeft();
+	}
+}
+
+void FMaterialEditor::OnAlignCenter()
+{
+	if (GraphEditor.IsValid())
+	{
+		GraphEditor->OnAlignCenter();
+	}
+}
+
+void FMaterialEditor::OnAlignRight()
+{
+	if (GraphEditor.IsValid())
+	{
+		GraphEditor->OnAlignRight();
+	}
+}
+
+void FMaterialEditor::OnStraightenConnections()
+{
+	if (GraphEditor.IsValid())
+	{
+		GraphEditor->OnStraightenConnections();
+	}
+}
+
+void FMaterialEditor::OnDistributeNodesH()
+{
+	if (GraphEditor.IsValid())
+	{
+		GraphEditor->OnDistributeNodesH();
+	}
+}
+
+void FMaterialEditor::OnDistributeNodesV()
+{
+	if (GraphEditor.IsValid())
+	{
+		GraphEditor->OnDistributeNodesV();
+	}
+}
+
+
 void FMaterialEditor::PostUndo(bool bSuccess)
 {
 	if (bSuccess)
@@ -4353,6 +4437,44 @@ TSharedRef<SGraphEditor> FMaterialEditor::CreateGraphEditorWidget()
 			FCanExecuteAction::CreateSP(this, &FMaterialEditor::OnCanPromoteToParameter)
 			);
 
+		// Alignment Commands
+		GraphEditorCommands->MapAction(FGraphEditorCommands::Get().AlignNodesTop,
+			FExecuteAction::CreateSP(this, &FMaterialEditor::OnAlignTop)
+		);
+
+		GraphEditorCommands->MapAction(FGraphEditorCommands::Get().AlignNodesMiddle,
+			FExecuteAction::CreateSP(this, &FMaterialEditor::OnAlignMiddle)
+		);
+
+		GraphEditorCommands->MapAction(FGraphEditorCommands::Get().AlignNodesBottom,
+			FExecuteAction::CreateSP(this, &FMaterialEditor::OnAlignBottom)
+		);
+
+		GraphEditorCommands->MapAction(FGraphEditorCommands::Get().AlignNodesLeft,
+			FExecuteAction::CreateSP(this, &FMaterialEditor::OnAlignLeft)
+		);
+
+		GraphEditorCommands->MapAction(FGraphEditorCommands::Get().AlignNodesCenter,
+			FExecuteAction::CreateSP(this, &FMaterialEditor::OnAlignCenter)
+		);
+
+		GraphEditorCommands->MapAction(FGraphEditorCommands::Get().AlignNodesRight,
+			FExecuteAction::CreateSP(this, &FMaterialEditor::OnAlignRight)
+		);
+
+		GraphEditorCommands->MapAction(FGraphEditorCommands::Get().StraightenConnections,
+			FExecuteAction::CreateSP(this, &FMaterialEditor::OnStraightenConnections)
+		);
+
+		// Distribution Commands
+		GraphEditorCommands->MapAction(FGraphEditorCommands::Get().DistributeNodesHorizontally,
+			FExecuteAction::CreateSP(this, &FMaterialEditor::OnDistributeNodesH)
+		);
+
+		GraphEditorCommands->MapAction(FGraphEditorCommands::Get().DistributeNodesVertically,
+			FExecuteAction::CreateSP(this, &FMaterialEditor::OnDistributeNodesV)
+		);
+
 	}
 
 	SGraphEditor::FGraphEditorEvents InEvents;
diff --git a/Engine/Source/Editor/UnrealEd/Private/MaterialGraphNode.cpp b/Engine/Source/Editor/UnrealEd/Private/MaterialGraphNode.cpp
--- a/Engine/Source/Editor/UnrealEd/Private/MaterialGraphNode.cpp
+++ b/Engine/Source/Editor/UnrealEd/Private/MaterialGraphNode.cpp
@@ -29,6 +29,7 @@
 #include "Materials/MaterialExpressionVectorParameter.h"
 #include "Materials/MaterialExpressionViewProperty.h"
 #include "Materials/MaterialExpressionMaterialLayerOutput.h"
+#include "Materials/MaterialExpressionTextureObjectParameter.h"
 
 #include "MaterialEditorUtilities.h"
 #include "MaterialEditorActions.h"
@@ -38,6 +39,7 @@
 #include "ScopedTransaction.h"
 
 
+
 #define LOCTEXT_NAMESPACE "MaterialGraphNode"
 
 /////////////////////////////////////////////////////
@@ -330,7 +332,7 @@ void UMaterialGraphNode::GetContextMenuActions(const FGraphNodeContextMenuBuilde
 				// Add a 'Convert To Texture' option for convertible types
 				Context.MenuBuilder->BeginSection("MaterialEditorMenu0");
 				{
-					if ( MaterialExpression->IsA(UMaterialExpressionTextureSample::StaticClass()))
+					if ( MaterialExpression->IsA(UMaterialExpressionTextureSample::StaticClass()) && !MaterialExpression->HasAParameterName())
 					{
 						Context.MenuBuilder->AddMenuEntry(FMaterialEditorCommands::Get().ConvertToTextureObjects);
 					}
@@ -347,7 +349,8 @@ void UMaterialGraphNode::GetContextMenuActions(const FGraphNodeContextMenuBuilde
 				|| MaterialExpression->IsA(UMaterialExpressionConstant2Vector::StaticClass())
 				|| MaterialExpression->IsA(UMaterialExpressionConstant3Vector::StaticClass())
 				|| MaterialExpression->IsA(UMaterialExpressionConstant4Vector::StaticClass())
-				|| MaterialExpression->IsA(UMaterialExpressionTextureSample::StaticClass())
+				|| (MaterialExpression->IsA(UMaterialExpressionTextureSample::StaticClass()) && !MaterialExpression->HasAParameterName())
+				|| MaterialExpression->IsA(UMaterialExpressionTextureObject::StaticClass())
 				|| MaterialExpression->IsA(UMaterialExpressionComponentMask::StaticClass())
 				|| MaterialExpression->IsA(UMaterialExpressionMaterialFunctionCall::StaticClass()))
 			{
@@ -360,7 +363,8 @@ void UMaterialGraphNode::GetContextMenuActions(const FGraphNodeContextMenuBuilde
 
 			// Add a 'Convert To Constant' option for convertible types
 			if (MaterialExpression->IsA(UMaterialExpressionScalarParameter::StaticClass())
-				|| MaterialExpression->IsA(UMaterialExpressionVectorParameter::StaticClass()))
+				|| MaterialExpression->IsA(UMaterialExpressionVectorParameter::StaticClass())
+				|| MaterialExpression->IsA(UMaterialExpressionTextureObjectParameter::StaticClass()))
 			{
 				Context.MenuBuilder->BeginSection("MaterialEditorMenu1");
 				{
@@ -417,6 +421,24 @@ void UMaterialGraphNode::GetContextMenuActions(const FGraphNodeContextMenuBuilde
 		}
 		Context.MenuBuilder->EndSection();
 
+		Context.MenuBuilder->AddSubMenu(LOCTEXT("AlignmentHeader", "Alignment"), FText(), FNewMenuDelegate::CreateLambda([](FMenuBuilder& InMenuBuilder) {
+
+			InMenuBuilder.BeginSection("EdGraphSchemaAlignment", LOCTEXT("AlignHeader", "Align"));
+			InMenuBuilder.AddMenuEntry(FGraphEditorCommands::Get().AlignNodesTop);
+			InMenuBuilder.AddMenuEntry(FGraphEditorCommands::Get().AlignNodesMiddle);
+			InMenuBuilder.AddMenuEntry(FGraphEditorCommands::Get().AlignNodesBottom);
+			InMenuBuilder.AddMenuEntry(FGraphEditorCommands::Get().AlignNodesLeft);
+			InMenuBuilder.AddMenuEntry(FGraphEditorCommands::Get().AlignNodesCenter);
+			InMenuBuilder.AddMenuEntry(FGraphEditorCommands::Get().AlignNodesRight);
+			InMenuBuilder.AddMenuEntry(FGraphEditorCommands::Get().StraightenConnections);
+			InMenuBuilder.EndSection();
+
+			InMenuBuilder.BeginSection("EdGraphSchemaDistribution", LOCTEXT("DistributionHeader", "Distribution"));
+			InMenuBuilder.AddMenuEntry(FGraphEditorCommands::Get().DistributeNodesHorizontally);
+			InMenuBuilder.AddMenuEntry(FGraphEditorCommands::Get().DistributeNodesVertically);
+			InMenuBuilder.EndSection();
+		}));
+
 		Context.MenuBuilder->BeginSection("MaterialEditorMenuDocumentation");
 		{
 			Context.MenuBuilder->AddMenuEntry(FGraphEditorCommands::Get().GoToDocumentation);
