Fixing cis

diff --git a/Engine/Source/Editor/MaterialEditor/Private/MaterialEditor.cpp b/Engine/Source/Editor/MaterialEditor/Private/MaterialEditor.cpp
--- a/Engine/Source/Editor/MaterialEditor/Private/MaterialEditor.cpp
+++ b/Engine/Source/Editor/MaterialEditor/Private/MaterialEditor.cpp
@@ -3964,79 +3964,6 @@ void FMaterialEditor::RedoGraphAction()
 
 }
 
-void FMaterialEditor::OnAlignTop()
-{
-	if (GraphEditor.IsValid())
-	{
-		GraphEditor->OnAlignTop();
-	}
-}
-
-void FMaterialEditor::OnAlignMiddle()
-{
-	if (GraphEditor.IsValid())
-	{
-		GraphEditor->OnAlignMiddle();
-	}
-}
-
-void FMaterialEditor::OnAlignBottom()
-{
-	if (GraphEditor.IsValid())
-	{
-		GraphEditor->OnAlignMiddle();
-	}
-}
-
-void FMaterialEditor::OnAlignLeft()
-{
-	if (GraphEditor.IsValid())
-	{
-		GraphEditor->OnAlignLeft();
-	}
-}
-
-void FMaterialEditor::OnAlignCenter()
-{
-	if (GraphEditor.IsValid())
-	{
-		GraphEditor->OnAlignCenter();
-	}
-}
-
-void FMaterialEditor::OnAlignRight()
-{
-	if (GraphEditor.IsValid())
-	{
-		GraphEditor->OnAlignRight();
-	}
-}
-
-void FMaterialEditor::OnStraightenConnections()
-{
-	if (GraphEditor.IsValid())
-	{
-		GraphEditor->OnStraightenConnections();
-	}
-}
-
-void FMaterialEditor::OnDistributeNodesH()
-{
-	if (GraphEditor.IsValid())
-	{
-		GraphEditor->OnDistributeNodesH();
-	}
-}
-
-void FMaterialEditor::OnDistributeNodesV()
-{
-	if (GraphEditor.IsValid())
-	{
-		GraphEditor->OnDistributeNodesV();
-	}
-}
-
-
 void FMaterialEditor::PostUndo(bool bSuccess)
 {
 	if (bSuccess)
@@ -4437,44 +4364,6 @@ TSharedRef<SGraphEditor> FMaterialEditor::CreateGraphEditorWidget()
 			FCanExecuteAction::CreateSP(this, &FMaterialEditor::OnCanPromoteToParameter)
 			);
 
-		// Alignment Commands
-		GraphEditorCommands->MapAction(FGraphEditorCommands::Get().AlignNodesTop,
-			FExecuteAction::CreateSP(this, &FMaterialEditor::OnAlignTop)
-		);
-
-		GraphEditorCommands->MapAction(FGraphEditorCommands::Get().AlignNodesMiddle,
-			FExecuteAction::CreateSP(this, &FMaterialEditor::OnAlignMiddle)
-		);
-
-		GraphEditorCommands->MapAction(FGraphEditorCommands::Get().AlignNodesBottom,
-			FExecuteAction::CreateSP(this, &FMaterialEditor::OnAlignBottom)
-		);
-
-		GraphEditorCommands->MapAction(FGraphEditorCommands::Get().AlignNodesLeft,
-			FExecuteAction::CreateSP(this, &FMaterialEditor::OnAlignLeft)
-		);
-
-		GraphEditorCommands->MapAction(FGraphEditorCommands::Get().AlignNodesCenter,
-			FExecuteAction::CreateSP(this, &FMaterialEditor::OnAlignCenter)
-		);
-
-		GraphEditorCommands->MapAction(FGraphEditorCommands::Get().AlignNodesRight,
-			FExecuteAction::CreateSP(this, &FMaterialEditor::OnAlignRight)
-		);
-
-		GraphEditorCommands->MapAction(FGraphEditorCommands::Get().StraightenConnections,
-			FExecuteAction::CreateSP(this, &FMaterialEditor::OnStraightenConnections)
-		);
-
-		// Distribution Commands
-		GraphEditorCommands->MapAction(FGraphEditorCommands::Get().DistributeNodesHorizontally,
-			FExecuteAction::CreateSP(this, &FMaterialEditor::OnDistributeNodesH)
-		);
-
-		GraphEditorCommands->MapAction(FGraphEditorCommands::Get().DistributeNodesVertically,
-			FExecuteAction::CreateSP(this, &FMaterialEditor::OnDistributeNodesV)
-		);
-
 	}
 
 	SGraphEditor::FGraphEditorEvents InEvents;
