Redoing material parameters panel with property row generator

diff --git a/Engine/Source/Editor/MaterialEditor/Private/MaterialEditor.cpp b/Engine/Source/Editor/MaterialEditor/Private/MaterialEditor.cpp
--- a/Engine/Source/Editor/MaterialEditor/Private/MaterialEditor.cpp
+++ b/Engine/Source/Editor/MaterialEditor/Private/MaterialEditor.cpp
@@ -121,6 +121,8 @@
 #include "Materials/MaterialExpressionMaterialLayerOutput.h"
 
 #include "MaterialStats.h"
+#include "SMaterialParametersOverviewWidget.h"
+#include "IPropertyRowGenerator.h"
 
 #define LOCTEXT_NAMESPACE "MaterialEditor"
 
@@ -691,6 +693,8 @@ void FMaterialEditor::InitMaterialEditor( const EToolkitMode::Type Mode, const T
 	Material->MaterialGraph->RebuildGraph();
 	RecenterEditor();
 
+	UMaterialEditingLibrary::RecompileMaterial(OriginalMaterial);
+
 	//Make sure the preview material is initialized.
 	UpdatePreviewMaterial(true);
 	RegenerateCodeView(true);
@@ -764,7 +768,6 @@ FMaterialEditor::~FMaterialEditor()
 	SaveEditorSettings();
 
 	MaterialDetailsView.Reset();
-	MaterialParametersView.Reset();
 
 	{
 		SCOPED_SUSPEND_RENDERING_THREAD(true);
@@ -933,8 +936,6 @@ void FMaterialEditor::CreateInternalWidgets()
 
 	PropertyEditorModule.RegisterCustomClassLayout( UMaterial::StaticClass()->GetFName(), FOnGetDetailCustomizationInstance::CreateStatic(&FMaterialDetailCustomization::MakeInstance ) );
 
-	const FDetailsViewArgs ParametersViewArgs(false, false, true, FDetailsViewArgs::HideNameArea, true, this);
-	MaterialParametersView = PropertyEditorModule.CreateDetailView(ParametersViewArgs);
 	MaterialEditorInstance = NewObject<UMaterialEditorPreviewParameters>(GetTransientPackage(), NAME_None, RF_Transactional);
 	MaterialEditorInstance->PreviewMaterial = Material;
 	MaterialEditorInstance->OriginalMaterial = OriginalMaterial;
@@ -942,14 +943,9 @@ void FMaterialEditor::CreateInternalWidgets()
 	{
 		MaterialEditorInstance->OriginalFunction = MaterialFunction->ParentFunction;
 	}
-	MaterialEditorInstance->DetailsView = MaterialParametersView;
-	MaterialEditorInstance->RegenerateArrays();
-
-	FOnGetDetailCustomizationInstance LayoutParameterDetails = FOnGetDetailCustomizationInstance::CreateStatic(
-		&FMaterialEditorParameterDetails::MakeInstance, MaterialEditorInstance);
-	MaterialParametersView->RegisterInstancedCustomPropertyLayout(UMaterialEditorPreviewParameters::StaticClass(), LayoutParameterDetails);
-	MaterialParametersView->SetObject(MaterialEditorInstance, true);
-	MaterialParametersView->OnFinishedChangingProperties().AddSP(this, &FMaterialEditor::OnFinishedChangingProperties);
+	MaterialParametersOverviewWidget = SNew(SMaterialParametersOverviewPanel)
+		.InMaterialEditorInstance(MaterialEditorInstance);
+	MaterialParametersOverviewWidget->GetGenerator()->OnFinishedChangingProperties().AddSP(this, &FMaterialEditor::OnFinishedChangingParametersFromOverview);
 	
 	IMaterialEditorModule* MaterialEditorModule = &FModuleManager::LoadModuleChecked<IMaterialEditorModule>("MaterialEditor");
 	if (MaterialEditorModule->MaterialLayersEnabled())
@@ -1002,6 +998,22 @@ void FMaterialEditor::OnFinishedChangingProperties(const FPropertyChangedEvent&
 		}
 		RefreshPreviewViewport();
 		UpdatePreviewMaterial();
+		MaterialParametersOverviewWidget->UpdateEditorInstance(MaterialEditorInstance);
+	}
+}
+
+void FMaterialEditor::OnFinishedChangingParametersFromOverview(const FPropertyChangedEvent& PropertyChangedEvent)
+{
+	bool bRefreshNodePreviews = false;
+	if (PropertyChangedEvent.Property != nullptr)
+	{
+		RefreshExpressionPreviews(true);
+		RefreshPreviewViewport();
+		UpdatePreviewMaterial();
+		TArray<TWeakObjectPtr<UObject>> SelectedObjects = MaterialDetailsView->GetSelectedObjects();
+		MaterialDetailsView->SetObjects(SelectedObjects, true);
+		Material->MarkPackageDirty();
+		SetMaterialDirty();
 	}
 }
 
@@ -1591,9 +1603,6 @@ void FMaterialEditor::UpdatePreviewMaterial( bool bForce )
 	}
 	if (MaterialEditorInstance != nullptr)
 	{
-		MaterialEditorInstance->RegenerateArrays();
-		MaterialParametersView->ForceRefresh();
-
 		if (MaterialLayersFunctionsInstance.IsValid())
 		{
 			MaterialLayersFunctionsInstance->SetEditorInstance(MaterialEditorInstance);
@@ -2989,8 +2998,7 @@ void FMaterialEditor::OnPromoteToParameter()
 	}
 	if (MaterialEditorInstance != nullptr)
 	{
-		MaterialEditorInstance->RegenerateArrays();
-		MaterialParametersView->ForceRefresh();
+		MaterialParametersOverviewWidget->UpdateEditorInstance(MaterialEditorInstance);
 	}
 }
 
@@ -3336,7 +3344,7 @@ TSharedRef<SDockTab> FMaterialEditor::SpawnTab_ParameterDefaults(const FSpawnTab
 		[
 			SNew(SBox)
 			[
-				MaterialParametersView.ToSharedRef()
+				MaterialParametersOverviewWidget.ToSharedRef()
 			]
 		];
 
@@ -3476,6 +3484,7 @@ UMaterialExpression* FMaterialEditor::CreateNewMaterialExpression(UClass* NewExp
 	RefreshExpressionPreviews();
 	GraphEditor->NotifyGraphChanged();
 	SetMaterialDirty();
+	MaterialParametersOverviewWidget->UpdateEditorInstance(MaterialEditorInstance);
 	return NewExpression;
 }
 
@@ -4080,7 +4089,7 @@ void FMaterialEditor::NotifyPostChange( const FPropertyChangedEvent& PropertyCha
 		{
 			MaterialLayersFunctionsInstance->SetEditorInstance(MaterialEditorInstance);
 		}
-
+		MaterialParametersOverviewWidget->UpdateEditorInstance(MaterialEditorInstance);
 		const FName NameOfPropertyThatChanged( *PropertyThatChanged->GetName() );
 		if ((NameOfPropertyThatChanged == GET_MEMBER_NAME_CHECKED(UMaterialInterface, PreviewMesh)) ||
 			(NameOfPropertyThatChanged == GET_MEMBER_NAME_CHECKED(UMaterial, bUsedWithSkeletalMesh)))
@@ -4330,6 +4339,7 @@ void FMaterialEditor::OnColorPickerCommitted(FLinearColor LinearColor)
 	}
 
 	RefreshExpressionPreviews();
+	MaterialParametersOverviewWidget->UpdateEditorInstance(MaterialEditorInstance);
 }
 
 TSharedRef<SGraphEditor> FMaterialEditor::CreateGraphEditorWidget()
@@ -4809,6 +4819,7 @@ void FMaterialEditor::OnNodeTitleCommitted(const FText& NewText, ETextCommit::Ty
 		const FScopedTransaction Transaction( LOCTEXT( "RenameNode", "Rename Node" ) );
 		NodeBeingChanged->Modify();
 		NodeBeingChanged->OnRenameNode(NewText.ToString());
+		MaterialParametersOverviewWidget->UpdateEditorInstance(MaterialEditorInstance);
 	}
 }
 
diff --git a/Engine/Source/Editor/MaterialEditor/Private/MaterialEditor.h b/Engine/Source/Editor/MaterialEditor/Private/MaterialEditor.h
--- a/Engine/Source/Editor/MaterialEditor/Private/MaterialEditor.h
+++ b/Engine/Source/Editor/MaterialEditor/Private/MaterialEditor.h
@@ -724,6 +724,7 @@ private:
 	TSharedRef<SDockTab> SpawnTab_LayerProperties(const FSpawnTabArgs& Args);
 
 	void OnFinishedChangingProperties(const FPropertyChangedEvent& PropertyChangedEvent);
+	void OnFinishedChangingParametersFromOverview(const FPropertyChangedEvent& PropertyChangedEvent);
 private:
 	/** List of open tool panels; used to ensure only one exists at any one time */
 	TMap< FName, TWeakPtr<class SDockableTab> > SpawnedToolPanels;
@@ -731,9 +732,6 @@ private:
 	/** Property View */
 	TSharedPtr<class IDetailsView> MaterialDetailsView;
 
-	/** Parameters View */
-	TSharedPtr<class IDetailsView> MaterialParametersView;
-
 	/** New Graph Editor */
 	TSharedPtr<class SGraphEditor> GraphEditor;
 
@@ -756,6 +754,9 @@ private:
 	/** Find results log as well as the search filter */
 	TSharedPtr<class SFindInMaterial> FindResults;
 
+	/** Parameter overview list View */
+	TSharedPtr<class SMaterialParametersOverviewPanel> MaterialParametersOverviewWidget;
+
 	/** Layer Properties View */
 	TSharedPtr<class SMaterialLayersFunctionsMaterialWrapper> MaterialLayersFunctionsInstance;
 
diff --git a/Engine/Source/Editor/MaterialEditor/Private/MaterialPropertyHelpers.cpp b/Engine/Source/Editor/MaterialEditor/Private/MaterialPropertyHelpers.cpp
--- a/Engine/Source/Editor/MaterialEditor/Private/MaterialPropertyHelpers.cpp
+++ b/Engine/Source/Editor/MaterialEditor/Private/MaterialPropertyHelpers.cpp
@@ -586,7 +586,7 @@ FReply FMaterialPropertyHelpers::OnClickedSaveNewFunctionInstance(class UMateria
 }
 
 
-FReply FMaterialPropertyHelpers::OnClickedSaveNewLayerInstance(class UMaterialFunctionInterface* Object, TSharedPtr<FStackSortedData> InSortedData)
+FReply FMaterialPropertyHelpers::OnClickedSaveNewLayerInstance(class UMaterialFunctionInterface* Object, TSharedPtr<FSortedParamData> InSortedData)
 {
 	const FString DefaultSuffix = TEXT("_Inst");
 	TArray<FEditorParameterGroup> ParameterGroups;
@@ -595,7 +595,7 @@ FReply FMaterialPropertyHelpers::OnClickedSaveNewLayerInstance(class UMaterialFu
 	{
 		FunctionPreviewMaterial = Object->GetPreviewMaterial();
 	}
-	for (TSharedPtr<FStackSortedData> Group : InSortedData->Children)
+	for (TSharedPtr<FSortedParamData> Group : InSortedData->Children)
 	{
 		FEditorParameterGroup DuplicatedGroup = FEditorParameterGroup();
 		DuplicatedGroup.GroupAssociation = Group->Group.GroupAssociation;
@@ -857,7 +857,7 @@ void FMaterialPropertyHelpers::ResetLayerAssetToDefault(TSharedPtr<IPropertyHand
 	
 }
 
-bool FMaterialPropertyHelpers::ShouldLayerAssetShowResetToDefault(TSharedPtr<IPropertyHandle> PropertyHandle, TSharedPtr<FStackSortedData> InParameterData, UMaterialInterface* InMaterial)
+bool FMaterialPropertyHelpers::ShouldLayerAssetShowResetToDefault(TSharedPtr<IPropertyHandle> PropertyHandle, TSharedPtr<FSortedParamData> InParameterData, UMaterialInterface* InMaterial)
 {
 	if (!InParameterData->Parameter)
 	{
diff --git a/Engine/Source/Editor/MaterialEditor/Private/SMaterialLayersFunctionsTree.cpp b/Engine/Source/Editor/MaterialEditor/Private/SMaterialLayersFunctionsTree.cpp
--- a/Engine/Source/Editor/MaterialEditor/Private/SMaterialLayersFunctionsTree.cpp
+++ b/Engine/Source/Editor/MaterialEditor/Private/SMaterialLayersFunctionsTree.cpp
@@ -92,7 +92,7 @@ void SMaterialLayersFunctionsInstanceTreeItem::RefreshOnRowChange(const FAssetDa
 	}
 }
 
-bool SMaterialLayersFunctionsInstanceTreeItem::GetFilterState(SMaterialLayersFunctionsInstanceTree* InTree, TSharedPtr<FStackSortedData> InStackData) const
+bool SMaterialLayersFunctionsInstanceTreeItem::GetFilterState(SMaterialLayersFunctionsInstanceTree* InTree, TSharedPtr<FSortedParamData> InStackData) const
 {
 	if (InStackData->ParameterInfo.Association == EMaterialParameterAssociation::LayerParameter)
 	{
@@ -105,7 +105,7 @@ bool SMaterialLayersFunctionsInstanceTreeItem::GetFilterState(SMaterialLayersFun
 	return false;
 }
 
-void SMaterialLayersFunctionsInstanceTreeItem::FilterClicked(const ECheckBoxState NewCheckedState, SMaterialLayersFunctionsInstanceTree* InTree, TSharedPtr<FStackSortedData> InStackData)
+void SMaterialLayersFunctionsInstanceTreeItem::FilterClicked(const ECheckBoxState NewCheckedState, SMaterialLayersFunctionsInstanceTree* InTree, TSharedPtr<FSortedParamData> InStackData)
 {
 	if (InStackData->ParameterInfo.Association == EMaterialParameterAssociation::LayerParameter)
 	{
@@ -117,7 +117,7 @@ void SMaterialLayersFunctionsInstanceTreeItem::FilterClicked(const ECheckBoxStat
 	}
 }
 
-ECheckBoxState SMaterialLayersFunctionsInstanceTreeItem::GetFilterChecked(SMaterialLayersFunctionsInstanceTree* InTree, TSharedPtr<FStackSortedData> InStackData) const
+ECheckBoxState SMaterialLayersFunctionsInstanceTreeItem::GetFilterChecked(SMaterialLayersFunctionsInstanceTree* InTree, TSharedPtr<FSortedParamData> InStackData) const
 {
 	return GetFilterState(InTree, InStackData) ? ECheckBoxState::Checked : ECheckBoxState::Unchecked;
 }
@@ -203,8 +203,8 @@ FReply SMaterialLayersFunctionsInstanceTreeItem::OnLayerDrop(const FDragDropEven
 	{
 		return FReply::Unhandled();
 	}
-	TSharedPtr<FStackSortedData> SwappingPropertyData = LayerPtr->StackParameterData;
-	TSharedPtr<FStackSortedData> SwappablePropertyData = StackParameterData;
+	TSharedPtr<FSortedParamData> SwappingPropertyData = LayerPtr->StackParameterData;
+	TSharedPtr<FSortedParamData> SwappablePropertyData = StackParameterData;
 	if (SwappingPropertyData.IsValid() && SwappablePropertyData.IsValid())
 	{
 		if (SwappingPropertyData != SwappablePropertyData)
@@ -334,7 +334,7 @@ void SMaterialLayersFunctionsInstanceTreeItem::Construct(const FArguments& InArg
 				];
 		}
 		const float ThumbnailSize = 24.0f;
-		TArray<TSharedPtr<FStackSortedData>> AssetChildren = StackParameterData->Children;
+		TArray<TSharedPtr<FSortedParamData>> AssetChildren = StackParameterData->Children;
 		if (AssetChildren.Num() > 0)
 		{
 			HeaderRowWidget->AddSlot()
@@ -346,7 +346,7 @@ void SMaterialLayersFunctionsInstanceTreeItem::Construct(const FArguments& InArg
 					SNullWidget::NullWidget
 				];
 		}
-		for (TSharedPtr<FStackSortedData> AssetChild : AssetChildren)
+		for (TSharedPtr<FSortedParamData> AssetChild : AssetChildren)
 		{
 			TSharedPtr<SBox> ThumbnailBox;
 			UObject* AssetObject;
@@ -950,7 +950,7 @@ void SMaterialLayersFunctionsInstanceTreeItem::Construct(const FArguments& InArg
 	FOnTableRowDragLeave LayerDragLeaveDelegate = FOnTableRowDragLeave::CreateSP(this, &SMaterialLayersFunctionsInstanceTreeItem::OnLayerDragLeave);
 	FOnTableRowDrop LayerDropDelegate = FOnTableRowDrop::CreateSP(this, &SMaterialLayersFunctionsInstanceTreeItem::OnLayerDrop);
 
-	STableRow< TSharedPtr<FStackSortedData> >::ConstructInternal(
+	STableRow< TSharedPtr<FSortedParamData> >::ConstructInternal(
 		STableRow::FArguments()
 		.Style(FEditorStyle::Get(), "DetailsView.TreeView.TableRow")
 		.ShowSelection(false)
@@ -1005,7 +1005,7 @@ void SMaterialLayersFunctionsInstanceTree::Construct(const FArguments& InArgs)
 	}
 #endif
 
-	STreeView<TSharedPtr<FStackSortedData>>::Construct(
+	STreeView<TSharedPtr<FSortedParamData>>::Construct(
 		STreeView::FArguments()
 		.TreeItemsSource(&LayerProperties)
 		.SelectionMode(ESelectionMode::None)
@@ -1017,7 +1017,7 @@ void SMaterialLayersFunctionsInstanceTree::Construct(const FArguments& InArgs)
 	SetParentsExpansionState();
 }
 
-TSharedRef< ITableRow > SMaterialLayersFunctionsInstanceTree::OnGenerateRowMaterialLayersFunctionsTreeView(TSharedPtr<FStackSortedData> Item, const TSharedRef< STableViewBase >& OwnerTable)
+TSharedRef< ITableRow > SMaterialLayersFunctionsInstanceTree::OnGenerateRowMaterialLayersFunctionsTreeView(TSharedPtr<FSortedParamData> Item, const TSharedRef< STableViewBase >& OwnerTable)
 {
 	TSharedRef< SMaterialLayersFunctionsInstanceTreeItem > ReturnRow = SNew(SMaterialLayersFunctionsInstanceTreeItem, OwnerTable)
 		.StackParameterData(Item)
@@ -1026,13 +1026,13 @@ TSharedRef< ITableRow > SMaterialLayersFunctionsInstanceTree::OnGenerateRowMater
 	return ReturnRow;
 }
 
-void SMaterialLayersFunctionsInstanceTree::OnGetChildrenMaterialLayersFunctionsTreeView(TSharedPtr<FStackSortedData> InParent, TArray< TSharedPtr<FStackSortedData> >& OutChildren)
+void SMaterialLayersFunctionsInstanceTree::OnGetChildrenMaterialLayersFunctionsTreeView(TSharedPtr<FSortedParamData> InParent, TArray< TSharedPtr<FSortedParamData> >& OutChildren)
 {
 	OutChildren = InParent->Children;
 }
 
 
-void SMaterialLayersFunctionsInstanceTree::OnExpansionChanged(TSharedPtr<FStackSortedData> Item, bool bIsExpanded)
+void SMaterialLayersFunctionsInstanceTree::OnExpansionChanged(TSharedPtr<FSortedParamData> Item, bool bIsExpanded)
 {
 	bool* ExpansionValue = MaterialEditorInstance->SourceInstance->LayerParameterExpansion.Find(Item->NodeKey);
 	if (ExpansionValue == nullptr)
@@ -1079,7 +1079,7 @@ void SMaterialLayersFunctionsInstanceTree::RefreshOnAssetChange(const struct FAs
 	RequestTreeRefresh();
 }
 
-void SMaterialLayersFunctionsInstanceTree::ResetAssetToDefault(TSharedPtr<IPropertyHandle> InHandle, TSharedPtr<FStackSortedData> InData)
+void SMaterialLayersFunctionsInstanceTree::ResetAssetToDefault(TSharedPtr<IPropertyHandle> InHandle, TSharedPtr<FSortedParamData> InData)
 {
 	FMaterialPropertyHelpers::ResetLayerAssetToDefault(FunctionInstanceHandle.ToSharedRef(), InData->Parameter, InData->ParameterInfo.Association, InData->ParameterInfo.Index, MaterialEditorInstance);
 	UpdateThumbnailMaterial(InData->ParameterInfo.Association, InData->ParameterInfo.Index, false);
@@ -1238,14 +1238,14 @@ void SMaterialLayersFunctionsInstanceTree::CreateGroupsWidget()
 					MaterialEditorInstance->StoredBlendPreviews.AddDefaulted(BlendChildren);
 				}
 					
-				TSharedRef<FStackSortedData> StackProperty = MakeShared<FStackSortedData>();
+				TSharedRef<FSortedParamData> StackProperty = MakeShared<FSortedParamData>();
 				StackProperty->StackDataType = EStackDataType::Stack;
 				StackProperty->Parameter = Parameter;
 				StackProperty->ParameterInfo.Index = LayerChildren - 1;
 				StackProperty->NodeKey = FString::FromInt(StackProperty->ParameterInfo.Index);
 
 
-				TSharedRef<FStackSortedData> ChildProperty = MakeShared<FStackSortedData>();
+				TSharedRef<FSortedParamData> ChildProperty = MakeShared<FSortedParamData>();
 				ChildProperty->StackDataType = EStackDataType::Asset;
 				ChildProperty->Parameter = Parameter;
 				ChildProperty->ParameterHandle = LayerHandle->AsArray()->GetElement(LayerChildren - 1);
@@ -1276,7 +1276,7 @@ void SMaterialLayersFunctionsInstanceTree::CreateGroupsWidget()
 				{
 					for (int32 Counter = BlendChildren - 1; Counter >= 0; Counter--)
 					{
-						ChildProperty = MakeShared<FStackSortedData>();
+						ChildProperty = MakeShared<FSortedParamData>();
 						ChildProperty->StackDataType = EStackDataType::Asset;
 						ChildProperty->Parameter = Parameter;
 						ChildProperty->ParameterHandle = BlendHandle->AsArray()->GetElement(Counter);	
@@ -1299,14 +1299,14 @@ void SMaterialLayersFunctionsInstanceTree::CreateGroupsWidget()
 						}
 						LayerProperties.Last()->Children.Add(ChildProperty);
 							
-						StackProperty = MakeShared<FStackSortedData>();
+						StackProperty = MakeShared<FSortedParamData>();
 						StackProperty->StackDataType = EStackDataType::Stack;
 						StackProperty->Parameter = Parameter;
 						StackProperty->ParameterInfo.Index = Counter;
 						StackProperty->NodeKey = FString::FromInt(StackProperty->ParameterInfo.Index);
 						LayerProperties.Add(StackProperty);
 
-						ChildProperty = MakeShared<FStackSortedData>();
+						ChildProperty = MakeShared<FSortedParamData>();
 						ChildProperty->StackDataType = EStackDataType::Asset;
 						ChildProperty->Parameter = Parameter;
 						ChildProperty->ParameterHandle = LayerHandle->AsArray()->GetElement(Counter);
@@ -1333,7 +1333,7 @@ void SMaterialLayersFunctionsInstanceTree::CreateGroupsWidget()
 			}
 			else
 			{
-				FLayerParameterUnsortedData NonLayerProperty;
+				FUnsortedParamData NonLayerProperty;
 				UDEditorScalarParameterValue* ScalarParam = Cast<UDEditorScalarParameterValue>(Parameter);
 
 				if (ScalarParam && ScalarParam->SliderMax > ScalarParam->SliderMin)
@@ -1388,7 +1388,7 @@ TSharedRef<SWidget> SMaterialLayersFunctionsInstanceTree::CreateThumbnailWidget(
 void SMaterialLayersFunctionsInstanceTree::UpdateThumbnailMaterial(TEnumAsByte<EMaterialParameterAssociation> InAssociation, int32 InIndex, bool bAlterBlendIndex)
 {
 	// Need to invert index b/c layer properties is generated in reverse order
-	TArray<TSharedPtr<FStackSortedData>> AssetChildren = LayerProperties[LayerProperties.Num() - 1 - InIndex]->Children;
+	TArray<TSharedPtr<FSortedParamData>> AssetChildren = LayerProperties[LayerProperties.Num() - 1 - InIndex]->Children;
 	UMaterialInstanceConstant* MaterialToUpdate = nullptr;
 	int32 ParameterIndex = InIndex;
 	if (InAssociation == EMaterialParameterAssociation::LayerParameter)
@@ -1405,9 +1405,9 @@ void SMaterialLayersFunctionsInstanceTree::UpdateThumbnailMaterial(TEnumAsByte<E
 	}
 
 	TArray<FEditorParameterGroup> ParameterGroups;
-	for (TSharedPtr<FStackSortedData> AssetChild : AssetChildren)
+	for (TSharedPtr<FSortedParamData> AssetChild : AssetChildren)
 	{
-		for (TSharedPtr<FStackSortedData> Group : AssetChild->Children)
+		for (TSharedPtr<FSortedParamData> Group : AssetChild->Children)
 		{
 			if (Group->ParameterInfo.Association == InAssociation)
 			{
@@ -1454,15 +1454,15 @@ FReply SMaterialLayersFunctionsInstanceTree::OnThumbnailDoubleClick(const FGeome
 	return FReply::Unhandled();
 }
 
-void SMaterialLayersFunctionsInstanceTree::ShowSubParameters(TSharedPtr<FStackSortedData> ParentParameter)
+void SMaterialLayersFunctionsInstanceTree::ShowSubParameters(TSharedPtr<FSortedParamData> ParentParameter)
 {
-	for (FLayerParameterUnsortedData Property : NonLayerProperties)
+	for (FUnsortedParamData Property : NonLayerProperties)
 	{
 		UDEditorParameterValue* Parameter = Property.Parameter;
 		if (Parameter->ParameterInfo.Index == ParentParameter->ParameterInfo.Index
 			&& Parameter->ParameterInfo.Association == ParentParameter->ParameterInfo.Association)
 		{
-			TSharedPtr<FStackSortedData> GroupProperty(new FStackSortedData());
+			TSharedPtr<FSortedParamData> GroupProperty(new FSortedParamData());
 			GroupProperty->StackDataType = EStackDataType::Group;
 			GroupProperty->ParameterInfo.Index = Parameter->ParameterInfo.Index;
 			GroupProperty->ParameterInfo.Association = Parameter->ParameterInfo.Association;
@@ -1470,7 +1470,7 @@ void SMaterialLayersFunctionsInstanceTree::ShowSubParameters(TSharedPtr<FStackSo
 			GroupProperty->NodeKey = FString::FromInt(GroupProperty->ParameterInfo.Index) + FString::FromInt(GroupProperty->ParameterInfo.Association) + Property.ParameterGroup.GroupName.ToString();
 
 			bool bAddNewGroup = true;
-			for (TSharedPtr<struct FStackSortedData> GroupChild : ParentParameter->Children)
+			for (TSharedPtr<struct FSortedParamData> GroupChild : ParentParameter->Children)
 			{
 				if (GroupChild->NodeKey == GroupProperty->NodeKey)
 				{
@@ -1482,7 +1482,7 @@ void SMaterialLayersFunctionsInstanceTree::ShowSubParameters(TSharedPtr<FStackSo
 				ParentParameter->Children.Add(GroupProperty);
 			}
 
-			TSharedPtr<FStackSortedData> ChildProperty(new FStackSortedData());
+			TSharedPtr<FSortedParamData> ChildProperty(new FSortedParamData());
 			ChildProperty->StackDataType = EStackDataType::Property;
 			ChildProperty->Parameter = Parameter;
 			ChildProperty->ParameterInfo.Index = Parameter->ParameterInfo.Index;
@@ -1499,7 +1499,7 @@ void SMaterialLayersFunctionsInstanceTree::ShowSubParameters(TSharedPtr<FStackSo
 				Property.ParameterNode->GetChildren(ParamChildren);
 				for (int32 ParamChildIdx = 0; ParamChildIdx < ParamChildren.Num(); ParamChildIdx++)
 				{
-					TSharedPtr<FStackSortedData> ParamChildProperty(new FStackSortedData());
+					TSharedPtr<FSortedParamData> ParamChildProperty(new FSortedParamData());
 					ParamChildProperty->StackDataType = EStackDataType::PropertyChild;
 					ParamChildProperty->ParameterNode = ParamChildren[ParamChildIdx];
 					ParamChildProperty->ParameterHandle = ParamChildProperty->ParameterNode->CreatePropertyHandle();
@@ -1509,7 +1509,7 @@ void SMaterialLayersFunctionsInstanceTree::ShowSubParameters(TSharedPtr<FStackSo
 					ChildProperty->Children.Add(ParamChildProperty);
 				}
 			}
-			for (TSharedPtr<struct FStackSortedData> GroupChild : ParentParameter->Children)
+			for (TSharedPtr<struct FSortedParamData> GroupChild : ParentParameter->Children)
 			{
 				if (GroupChild->Group.GroupName == Property.ParameterGroup.GroupName
 					&& GroupChild->ParameterInfo.Association == ChildProperty->ParameterInfo.Association
@@ -1675,7 +1675,7 @@ void SMaterialLayersFunctionsInstanceWrapper::SetEditorInstance(UMaterialEditorI
 /////////////// MATERIAL VERSION
 
 
-class SMaterialLayersFunctionsMaterialTreeItem : public STableRow< TSharedPtr<FStackSortedData> >
+class SMaterialLayersFunctionsMaterialTreeItem : public STableRow< TSharedPtr<FSortedParamData> >
 {
 public:
 
@@ -1685,7 +1685,7 @@ public:
 	{}
 
 	/** The item content. */
-	SLATE_ARGUMENT(TSharedPtr<FStackSortedData>, StackParameterData)
+	SLATE_ARGUMENT(TSharedPtr<FSortedParamData>, StackParameterData)
 	SLATE_ARGUMENT(UMaterialEditorPreviewParameters*, MaterialEditorInstance)
 	SLATE_ARGUMENT(SMaterialLayersFunctionsMaterialTree*, InTree)
 	SLATE_END_ARGS()
@@ -1730,7 +1730,7 @@ public:
 #endif
 			TSharedRef<SHorizontalBox> HeaderRowWidget = SNew(SHorizontalBox);
 			const float ThumbnailSize = 24.0f;
-			TArray<TSharedPtr<FStackSortedData>> AssetChildren = StackParameterData->Children;
+			TArray<TSharedPtr<FSortedParamData>> AssetChildren = StackParameterData->Children;
 			if (AssetChildren.Num() > 0)
 			{
 				HeaderRowWidget->AddSlot()
@@ -1742,7 +1742,7 @@ public:
 						SNullWidget::NullWidget
 					];
 			}
-			for (TSharedPtr<FStackSortedData> AssetChild : AssetChildren)
+			for (TSharedPtr<FSortedParamData> AssetChild : AssetChildren)
 			{
 				TSharedPtr<SBox> ThumbnailBox;
 				UObject* AssetObject;
@@ -2134,7 +2134,7 @@ public:
 				WrapperWidget
 			];
 
-		STableRow< TSharedPtr<FStackSortedData> >::ConstructInternal(
+		STableRow< TSharedPtr<FSortedParamData> >::ConstructInternal(
 			STableRow::FArguments()
 			.Style(FEditorStyle::Get(), "DetailsView.TreeView.TableRow")
 			.ShowSelection(false),
@@ -2160,7 +2160,7 @@ public:
 	virtual FReply OnMouseButtonDoubleClick(const FGeometry& InMyGeometry, const FPointerEvent& InMouseEvent) override { return FReply::Handled(); };
 
 	/** The node info to build the tree view row from. */
-	TSharedPtr<FStackSortedData> StackParameterData;
+	TSharedPtr<FSortedParamData> StackParameterData;
 
 	UMaterialEditorPreviewParameters* MaterialEditorInstance;
 };
@@ -2298,7 +2298,7 @@ void SMaterialLayersFunctionsMaterialTree::Construct(const FArguments& InArgs)
 
 	CreateGroupsWidget();
 
-	STreeView<TSharedPtr<FStackSortedData>>::Construct(
+	STreeView<TSharedPtr<FSortedParamData>>::Construct(
 		STreeView::FArguments()
 		.TreeItemsSource(&LayerProperties)
 		.SelectionMode(ESelectionMode::None)
@@ -2310,7 +2310,7 @@ void SMaterialLayersFunctionsMaterialTree::Construct(const FArguments& InArgs)
 	SetParentsExpansionState();
 }
 
-TSharedRef< ITableRow > SMaterialLayersFunctionsMaterialTree::OnGenerateRowMaterialLayersFunctionsTreeView(TSharedPtr<FStackSortedData> Item, const TSharedRef< STableViewBase >& OwnerTable)
+TSharedRef< ITableRow > SMaterialLayersFunctionsMaterialTree::OnGenerateRowMaterialLayersFunctionsTreeView(TSharedPtr<FSortedParamData> Item, const TSharedRef< STableViewBase >& OwnerTable)
 {
 	TSharedRef< SMaterialLayersFunctionsMaterialTreeItem > ReturnRow = SNew(SMaterialLayersFunctionsMaterialTreeItem, OwnerTable)
 		.StackParameterData(Item)
@@ -2319,12 +2319,12 @@ TSharedRef< ITableRow > SMaterialLayersFunctionsMaterialTree::OnGenerateRowMater
 	return ReturnRow;
 }
 
-void SMaterialLayersFunctionsMaterialTree::OnGetChildrenMaterialLayersFunctionsTreeView(TSharedPtr<FStackSortedData> InParent, TArray< TSharedPtr<FStackSortedData> >& OutChildren)
+void SMaterialLayersFunctionsMaterialTree::OnGetChildrenMaterialLayersFunctionsTreeView(TSharedPtr<FSortedParamData> InParent, TArray< TSharedPtr<FSortedParamData> >& OutChildren)
 {
 	OutChildren = InParent->Children;
 }
 
-void SMaterialLayersFunctionsMaterialTree::OnExpansionChanged(TSharedPtr<FStackSortedData> Item, bool bIsExpanded)
+void SMaterialLayersFunctionsMaterialTree::OnExpansionChanged(TSharedPtr<FSortedParamData> Item, bool bIsExpanded)
 {
 	bool* ExpansionValue = MaterialEditorInstance->OriginalMaterial->LayerParameterExpansion.Find(Item->NodeKey);
 	if (ExpansionValue == nullptr)
@@ -2443,14 +2443,14 @@ void SMaterialLayersFunctionsMaterialTree::CreateGroupsWidget()
 				uint32 BlendChildren;
 				BlendHandle->GetNumChildren(BlendChildren);
 
-				TSharedPtr<FStackSortedData> StackProperty(new FStackSortedData());
+				TSharedPtr<FSortedParamData> StackProperty(new FSortedParamData());
 				StackProperty->StackDataType = EStackDataType::Stack;
 				StackProperty->Parameter = Parameter;
 				StackProperty->ParameterInfo.Index = LayerChildren - 1;
 				StackProperty->NodeKey = FString::FromInt(StackProperty->ParameterInfo.Index);
 
 
-				TSharedPtr<FStackSortedData> ChildProperty(new FStackSortedData());
+				TSharedPtr<FSortedParamData> ChildProperty(new FSortedParamData());
 				ChildProperty->StackDataType = EStackDataType::Asset;
 				ChildProperty->Parameter = Parameter;
 				ChildProperty->ParameterHandle = LayerHandle->AsArray()->GetElement(LayerChildren - 1);
@@ -2466,7 +2466,7 @@ void SMaterialLayersFunctionsMaterialTree::CreateGroupsWidget()
 				{
 					for (int32 Counter = BlendChildren - 1; Counter >= 0; Counter--)
 					{
-						ChildProperty = MakeShareable(new FStackSortedData());
+						ChildProperty = MakeShareable(new FSortedParamData());
 						ChildProperty->StackDataType = EStackDataType::Asset;
 						ChildProperty->Parameter = Parameter;
 						ChildProperty->ParameterHandle = BlendHandle->AsArray()->GetElement(Counter);
@@ -2476,14 +2476,14 @@ void SMaterialLayersFunctionsMaterialTree::CreateGroupsWidget()
 						ChildProperty->NodeKey = FString::FromInt(ChildProperty->ParameterInfo.Index) + FString::FromInt(ChildProperty->ParameterInfo.Association);
 						LayerProperties.Last()->Children.Add(ChildProperty);
 
-						StackProperty = MakeShareable(new FStackSortedData());
+						StackProperty = MakeShareable(new FSortedParamData());
 						StackProperty->StackDataType = EStackDataType::Stack;
 						StackProperty->Parameter = Parameter;
 						StackProperty->ParameterInfo.Index = Counter;
 						StackProperty->NodeKey = FString::FromInt(StackProperty->ParameterInfo.Index);
 						LayerProperties.Add(StackProperty);
 
-						ChildProperty = MakeShareable(new FStackSortedData());
+						ChildProperty = MakeShareable(new FSortedParamData());
 						ChildProperty->StackDataType = EStackDataType::Asset;
 						ChildProperty->Parameter = Parameter;
 						ChildProperty->ParameterHandle = LayerHandle->AsArray()->GetElement(Counter);
@@ -2498,7 +2498,7 @@ void SMaterialLayersFunctionsMaterialTree::CreateGroupsWidget()
 			}
 			else
 			{
-				FLayerParameterUnsortedData NonLayerProperty;
+				FUnsortedParamData NonLayerProperty;
 				UDEditorScalarParameterValue* ScalarParam = Cast<UDEditorScalarParameterValue>(Parameter);
 
 				if (ScalarParam && ScalarParam->SliderMax > ScalarParam->SliderMin)
@@ -2528,15 +2528,15 @@ void SMaterialLayersFunctionsMaterialTree::CreateGroupsWidget()
 	SetParentsExpansionState();
 }
 
-void SMaterialLayersFunctionsMaterialTree::ShowSubParameters(TSharedPtr<FStackSortedData> ParentParameter)
+void SMaterialLayersFunctionsMaterialTree::ShowSubParameters(TSharedPtr<FSortedParamData> ParentParameter)
 {
-	for (FLayerParameterUnsortedData Property : NonLayerProperties)
+	for (FUnsortedParamData Property : NonLayerProperties)
 	{
 		UDEditorParameterValue* Parameter = Property.Parameter;
 		if (Parameter->ParameterInfo.Index == ParentParameter->ParameterInfo.Index
 			&& Parameter->ParameterInfo.Association == ParentParameter->ParameterInfo.Association)
 		{
-			TSharedPtr<FStackSortedData> GroupProperty(new FStackSortedData());
+			TSharedPtr<FSortedParamData> GroupProperty(new FSortedParamData());
 			GroupProperty->StackDataType = EStackDataType::Group;
 			GroupProperty->ParameterInfo.Index = Parameter->ParameterInfo.Index;
 			GroupProperty->ParameterInfo.Association = Parameter->ParameterInfo.Association;
@@ -2544,7 +2544,7 @@ void SMaterialLayersFunctionsMaterialTree::ShowSubParameters(TSharedPtr<FStackSo
 			GroupProperty->NodeKey = FString::FromInt(GroupProperty->ParameterInfo.Index) + FString::FromInt(GroupProperty->ParameterInfo.Association) + Property.ParameterGroup.GroupName.ToString();
 
 			bool bAddNewGroup = true;
-			for (TSharedPtr<struct FStackSortedData> GroupChild : ParentParameter->Children)
+			for (TSharedPtr<struct FSortedParamData> GroupChild : ParentParameter->Children)
 			{
 				if (GroupChild->NodeKey == GroupProperty->NodeKey)
 				{
@@ -2555,7 +2555,7 @@ void SMaterialLayersFunctionsMaterialTree::ShowSubParameters(TSharedPtr<FStackSo
 			{
 				ParentParameter->Children.Add(GroupProperty);
 			}
-			TSharedPtr<FStackSortedData> ChildProperty(new FStackSortedData());
+			TSharedPtr<FSortedParamData> ChildProperty(new FSortedParamData());
 			ChildProperty->StackDataType = EStackDataType::Property;
 			ChildProperty->Parameter = Parameter;
 			ChildProperty->ParameterInfo.Index = Parameter->ParameterInfo.Index;
@@ -2572,7 +2572,7 @@ void SMaterialLayersFunctionsMaterialTree::ShowSubParameters(TSharedPtr<FStackSo
 				Property.ParameterNode->GetChildren(ParamChildren);
 				for (int32 ParamChildIdx = 0; ParamChildIdx < ParamChildren.Num(); ParamChildIdx++)
 				{
-					TSharedPtr<FStackSortedData> ParamChildProperty(new FStackSortedData());
+					TSharedPtr<FSortedParamData> ParamChildProperty(new FSortedParamData());
 					ParamChildProperty->StackDataType = EStackDataType::PropertyChild;
 					ParamChildProperty->ParameterNode = ParamChildren[ParamChildIdx];
 					ParamChildProperty->ParameterHandle = ParamChildProperty->ParameterNode->CreatePropertyHandle();
@@ -2582,7 +2582,7 @@ void SMaterialLayersFunctionsMaterialTree::ShowSubParameters(TSharedPtr<FStackSo
 				}
 			}
 
-			for (TSharedPtr<struct FStackSortedData> GroupChild : ParentParameter->Children)
+			for (TSharedPtr<struct FSortedParamData> GroupChild : ParentParameter->Children)
 			{
 				if (GroupChild->Group.GroupName == Property.ParameterGroup.GroupName
 					&& GroupChild->ParameterInfo.Association == ChildProperty->ParameterInfo.Association
diff --git a/Engine/Source/Editor/MaterialEditor/Private/SMaterialLayersFunctionsTree.h b/Engine/Source/Editor/MaterialEditor/Private/SMaterialLayersFunctionsTree.h
--- a/Engine/Source/Editor/MaterialEditor/Private/SMaterialLayersFunctionsTree.h
+++ b/Engine/Source/Editor/MaterialEditor/Private/SMaterialLayersFunctionsTree.h
@@ -19,7 +19,7 @@ class IPropertyHandle;
 class UMaterialEditorInstanceConstant;
 class SMaterialLayersFunctionsInstanceTree;
 
-class SMaterialLayersFunctionsInstanceTreeItem : public STableRow< TSharedPtr<FStackSortedData> >
+class SMaterialLayersFunctionsInstanceTreeItem : public STableRow< TSharedPtr<FSortedParamData> >
 {
 public:
 
@@ -29,7 +29,7 @@ public:
 	{}
 
 	/** The item content. */
-	SLATE_ARGUMENT(TSharedPtr<FStackSortedData>, StackParameterData)
+	SLATE_ARGUMENT(TSharedPtr<FSortedParamData>, StackParameterData)
 	SLATE_ARGUMENT(UMaterialEditorInstanceConstant*, MaterialEditorInstance)
 	SLATE_ARGUMENT(SMaterialLayersFunctionsInstanceTree*, InTree)
 	SLATE_END_ARGS()
@@ -41,15 +41,15 @@ private:
 	bool bIsHoveredDragTarget;
 
 
-	FString GetCurvePath(UDEditorScalarParameterValue* Parameter) const;
+	FString GetCurvePath(class UDEditorScalarParameterValue* Parameter) const;
 	const FSlateBrush* GetBorderImage() const;
 
 public:
 
 	void RefreshOnRowChange(const FAssetData& AssetData, SMaterialLayersFunctionsInstanceTree* InTree);
-	bool GetFilterState(SMaterialLayersFunctionsInstanceTree* InTree, TSharedPtr<FStackSortedData> InStackData) const;
-	void FilterClicked(const ECheckBoxState NewCheckedState, SMaterialLayersFunctionsInstanceTree* InTree, TSharedPtr<FStackSortedData> InStackData);
-	ECheckBoxState GetFilterChecked(SMaterialLayersFunctionsInstanceTree* InTree, TSharedPtr<FStackSortedData> InStackData) const;
+	bool GetFilterState(SMaterialLayersFunctionsInstanceTree* InTree, TSharedPtr<FSortedParamData> InStackData) const;
+	void FilterClicked(const ECheckBoxState NewCheckedState, SMaterialLayersFunctionsInstanceTree* InTree, TSharedPtr<FSortedParamData> InStackData);
+	ECheckBoxState GetFilterChecked(SMaterialLayersFunctionsInstanceTree* InTree, TSharedPtr<FSortedParamData> InStackData) const;
 	FText GetLayerName(SMaterialLayersFunctionsInstanceTree* InTree, int32 Counter) const;
 	void OnNameChanged(const FText& InText, ETextCommit::Type CommitInfo, SMaterialLayersFunctionsInstanceTree* InTree, int32 Counter);
 
@@ -81,7 +81,7 @@ public:
 	void Construct(const FArguments& InArgs, const TSharedRef<STableViewBase>& InOwnerTableView);
 
 	/** The node info to build the tree view row from. */
-	TSharedPtr<FStackSortedData> StackParameterData;
+	TSharedPtr<FSortedParamData> StackParameterData;
 
 	SMaterialLayersFunctionsInstanceTree* Tree;
 
@@ -111,7 +111,7 @@ public:
 	FSimpleDelegate OnLayerPropertyChanged;
 };
 
-class SMaterialLayersFunctionsInstanceTree : public STreeView<TSharedPtr<FStackSortedData>>
+class SMaterialLayersFunctionsInstanceTree : public STreeView<TSharedPtr<FSortedParamData>>
 {
 	friend class SMaterialLayersFunctionsInstanceTreeItem;
 public:
@@ -125,9 +125,9 @@ public:
 
 	/** Constructs this widget with InArgs */
 	void Construct(const FArguments& InArgs);
-	TSharedRef< ITableRow > OnGenerateRowMaterialLayersFunctionsTreeView(TSharedPtr<FStackSortedData> Item, const TSharedRef< STableViewBase >& OwnerTable);
-	void OnGetChildrenMaterialLayersFunctionsTreeView(TSharedPtr<FStackSortedData> InParent, TArray< TSharedPtr<FStackSortedData> >& OutChildren);
-	void OnExpansionChanged(TSharedPtr<FStackSortedData> Item, bool bIsExpanded);
+	TSharedRef< ITableRow > OnGenerateRowMaterialLayersFunctionsTreeView(TSharedPtr<FSortedParamData> Item, const TSharedRef< STableViewBase >& OwnerTable);
+	void OnGetChildrenMaterialLayersFunctionsTreeView(TSharedPtr<FSortedParamData> InParent, TArray< TSharedPtr<FSortedParamData> >& OutChildren);
+	void OnExpansionChanged(TSharedPtr<FSortedParamData> Item, bool bIsExpanded);
 	void SetParentsExpansionState();
 
 	float OnGetLeftColumnWidth() const { return 1.0f - ColumnWidth; }
@@ -139,7 +139,7 @@ public:
 	struct FMaterialLayersFunctions* FunctionInstance;
 	TSharedPtr<IPropertyHandle> FunctionInstanceHandle;
 	void RefreshOnAssetChange(const struct FAssetData& InAssetData, int32 Index, EMaterialParameterAssociation MaterialType);
-	void ResetAssetToDefault(TSharedPtr<IPropertyHandle> InHandle, TSharedPtr<FStackSortedData> InData);
+	void ResetAssetToDefault(TSharedPtr<IPropertyHandle> InHandle, TSharedPtr<FSortedParamData> InData);
 	void AddLayer();
 	void RemoveLayer(int32 Index);
 	FReply ToggleLayerVisibility(int32 Index);
@@ -162,12 +162,12 @@ public:
 	FReply OnThumbnailDoubleClick(const FGeometry& Geometry, const FPointerEvent& MouseEvent, EMaterialParameterAssociation InAssociation, int32 InIndex);
 protected:
 
-	void ShowSubParameters(TSharedPtr<FStackSortedData> ParentParameter);
+	void ShowSubParameters(TSharedPtr<FSortedParamData> ParentParameter);
 
 private:
-	TArray<TSharedPtr<FStackSortedData>> LayerProperties;
+	TArray<TSharedPtr<FSortedParamData>> LayerProperties;
 
-	TArray<FLayerParameterUnsortedData> NonLayerProperties;
+	TArray<FUnsortedParamData> NonLayerProperties;
 
 	/** The actual width of the right column.  The left column is 1-ColumnWidth */
 	float ColumnWidth;
@@ -202,7 +202,7 @@ public:
 };
 
 
-class SMaterialLayersFunctionsMaterialTree : public STreeView<TSharedPtr<FStackSortedData>>
+class SMaterialLayersFunctionsMaterialTree : public STreeView<TSharedPtr<FSortedParamData>>
 {
 	friend class SMaterialLayersFunctionsMaterialTreeItem;
 public:
@@ -216,9 +216,9 @@ public:
 
 	/** Constructs this widget with InArgs */
 	void Construct(const FArguments& InArgs);
-	TSharedRef< ITableRow > OnGenerateRowMaterialLayersFunctionsTreeView(TSharedPtr<FStackSortedData> Item, const TSharedRef< STableViewBase >& OwnerTable);
-	void OnGetChildrenMaterialLayersFunctionsTreeView(TSharedPtr<FStackSortedData> InParent, TArray< TSharedPtr<FStackSortedData> >& OutChildren);
-	void OnExpansionChanged(TSharedPtr<FStackSortedData> Item, bool bIsExpanded);
+	TSharedRef< ITableRow > OnGenerateRowMaterialLayersFunctionsTreeView(TSharedPtr<FSortedParamData> Item, const TSharedRef< STableViewBase >& OwnerTable);
+	void OnGetChildrenMaterialLayersFunctionsTreeView(TSharedPtr<FSortedParamData> InParent, TArray< TSharedPtr<FSortedParamData> >& OutChildren);
+	void OnExpansionChanged(TSharedPtr<FSortedParamData> Item, bool bIsExpanded);
 	void SetParentsExpansionState();
 
 	float OnGetLeftColumnWidth() const { return 1.0f - ColumnWidth; }
@@ -238,12 +238,12 @@ public:
 
 protected:
 
-	void ShowSubParameters(TSharedPtr<FStackSortedData> ParentParameter);
+	void ShowSubParameters(TSharedPtr<FSortedParamData> ParentParameter);
 
 private:
-	TArray<TSharedPtr<FStackSortedData>> LayerProperties;
+	TArray<TSharedPtr<FSortedParamData>> LayerProperties;
 
-	TArray<FLayerParameterUnsortedData> NonLayerProperties;
+	TArray<FUnsortedParamData> NonLayerProperties;
 
 	/** The actual width of the right column.  The left column is 1-ColumnWidth */
 	float ColumnWidth;
diff --git a/Engine/Source/Editor/MaterialEditor/Private/SMaterialParametersOverviewWidget.cpp b/Engine/Source/Editor/MaterialEditor/Private/SMaterialParametersOverviewWidget.cpp
new file mode 100644
--- /dev/null
+++ b/Engine/Source/Editor/MaterialEditor/Private/SMaterialParametersOverviewWidget.cpp
@@ -0,0 +1,737 @@
+// Copyright 1998-2018 Epic Games, Inc. All Rights Reserved.
+
+#include "SMaterialParametersOverviewWidget.h"
+#include "MaterialEditor/DEditorFontParameterValue.h"
+#include "MaterialEditor/DEditorMaterialLayersParameterValue.h"
+#include "MaterialEditor/DEditorScalarParameterValue.h"
+#include "MaterialEditor/DEditorStaticComponentMaskParameterValue.h"
+#include "MaterialEditor/DEditorStaticSwitchParameterValue.h"
+#include "MaterialEditor/DEditorTextureParameterValue.h"
+#include "MaterialEditor/DEditorVectorParameterValue.h"
+#include "MaterialEditor/MaterialEditorInstanceConstant.h"
+#include "Materials/Material.h"
+#include "PropertyHandle.h"
+#include "Modules/ModuleManager.h"
+#include "PropertyEditorModule.h"
+#include "ISinglePropertyView.h"
+#include "Materials/MaterialInstanceConstant.h"
+#include "EditorStyleSet.h"
+#include "PropertyCustomizationHelpers.h"
+#include "ScopedTransaction.h"
+#include "IPropertyRowGenerator.h"
+#include "Widgets/Views/STreeView.h"
+#include "IDetailTreeNode.h"
+#include "AssetThumbnail.h"
+#include "MaterialEditorInstanceDetailCustomization.h"
+#include "MaterialPropertyHelpers.h"
+#include "DetailWidgetRow.h"
+#include "Widgets/Input/SCheckBox.h"
+#include "EditorSupportDelegates.h"
+#include "Widgets/Images/SImage.h"
+#include "MaterialEditor/MaterialEditorPreviewParameters.h"
+#include "Widgets/Input/SButton.h"
+#include "Widgets/Text/SInlineEditableTextBlock.h"
+#include "Materials/MaterialFunctionInstance.h"
+#include "Framework/Application/SlateApplication.h"
+#include "Toolkits/AssetEditorManager.h"
+#include "Widgets/Input/SEditableTextBox.h"
+#include "Curves/CurveLinearColor.h"
+#include "Curves/CurveLinearColorAtlas.h"
+#include "Widgets/Views/SExpanderArrow.h"
+#include "Widgets/Layout/SWidgetSwitcher.h"
+
+
+#define LOCTEXT_NAMESPACE "MaterialLayerCustomization"
+
+FString SMaterialParametersOverviewTreeItem::GetCurvePath(UDEditorScalarParameterValue* Parameter) const
+{
+	FString Path = Parameter->AtlasData.Curve->GetPathName();
+	return Path;
+}
+
+const FSlateBrush* SMaterialParametersOverviewTreeItem::GetBorderImage() const
+{
+	if (IsHovered())
+	{
+		return FEditorStyle::GetBrush("DetailsView.CategoryMiddle_Hovered");
+	}
+	else
+	{
+		return FEditorStyle::GetBrush("DetailsView.CategoryMiddle");
+	}
+}
+
+void SMaterialParametersOverviewTreeItem::RefreshOnRowChange(const FAssetData& AssetData, TSharedPtr<SMaterialParametersOverviewTree> InTree)
+{
+	if (InTree.IsValid())
+	{
+		InTree->CreateGroupsWidget();
+	}
+}
+
+void SMaterialParametersOverviewTreeItem::Construct(const FArguments& InArgs, const TSharedRef<STableViewBase>& InOwnerTableView)
+{
+	StackParameterData = InArgs._StackParameterData;
+	MaterialEditorInstance = InArgs._MaterialEditorInstance;
+	Tree = InArgs._InTree;
+	ColumnSizeData.LeftColumnWidth = TAttribute<float>(Tree.Pin().Get(), &SMaterialParametersOverviewTree::OnGetLeftColumnWidth);
+	ColumnSizeData.RightColumnWidth = TAttribute<float>(Tree.Pin().Get(), &SMaterialParametersOverviewTree::OnGetRightColumnWidth);
+	ColumnSizeData.OnWidthChanged = SSplitter::FOnSlotResized::CreateSP(Tree.Pin().Get(), &SMaterialParametersOverviewTree::OnSetColumnWidth);
+
+	TSharedRef<SWidget> LeftSideWidget = SNullWidget::NullWidget;
+	TSharedRef<SWidget> RightSideWidget = SNullWidget::NullWidget;
+	FText NameOverride;
+	TSharedRef<SVerticalBox> WrapperWidget = SNew(SVerticalBox);
+
+// GROUP --------------------------------------------------
+	if (StackParameterData->StackDataType == EStackDataType::Group)
+	{
+		NameOverride = FText::FromName(StackParameterData->Group.GroupName);
+		LeftSideWidget = SNew(STextBlock)
+			.Text(NameOverride)
+			.TextStyle(FEditorStyle::Get(), "TinyText");
+	}
+// END GROUP
+
+// PROPERTY ----------------------------------------------
+	if (StackParameterData->StackDataType == EStackDataType::Property)
+	{
+		UDEditorStaticComponentMaskParameterValue* CompMaskParam = Cast<UDEditorStaticComponentMaskParameterValue>(StackParameterData->Parameter);
+		UDEditorVectorParameterValue* VectorParam = Cast<UDEditorVectorParameterValue>(StackParameterData->Parameter);
+		UDEditorScalarParameterValue* ScalarParam = Cast<UDEditorScalarParameterValue>(StackParameterData->Parameter);
+
+		TAttribute<bool> IsParamEnabled = TAttribute<bool>::Create(TAttribute<bool>::FGetter::CreateStatic(&FMaterialPropertyHelpers::IsOverriddenExpression, StackParameterData->Parameter));
+		NameOverride = FText::FromName(StackParameterData->Parameter->ParameterInfo.Name);
+
+		IDetailTreeNode& Node = *StackParameterData->ParameterNode;
+		TSharedPtr<IDetailPropertyRow> GeneratedRow = StaticCastSharedPtr<IDetailPropertyRow>(Node.GetRow());
+		IDetailPropertyRow& Row = *GeneratedRow.Get();
+		Row.DisplayName(NameOverride);
+
+		if (VectorParam && VectorParam->bIsUsedAsChannelMask)
+		{
+			FOnGetPropertyComboBoxStrings GetMaskStrings = FOnGetPropertyComboBoxStrings::CreateStatic(&FMaterialPropertyHelpers::GetVectorChannelMaskComboBoxStrings);
+			FOnGetPropertyComboBoxValue GetMaskValue = FOnGetPropertyComboBoxValue::CreateStatic(&FMaterialPropertyHelpers::GetVectorChannelMaskValue, StackParameterData->Parameter);
+			FOnPropertyComboBoxValueSelected SetMaskValue = FOnPropertyComboBoxValueSelected::CreateStatic(&FMaterialPropertyHelpers::SetVectorChannelMaskValue, StackParameterData->ParameterNode->CreatePropertyHandle(), StackParameterData->Parameter, (UObject*)MaterialEditorInstance);
+
+			FDetailWidgetRow& CustomWidget = Row.CustomWidget();
+			CustomWidget
+			.FilterString(NameOverride)
+			.NameContent()
+			[
+				SNew(STextBlock)
+				.Text(NameOverride)
+				.ToolTipText(FMaterialPropertyHelpers::GetParameterExpressionDescription(StackParameterData->Parameter, MaterialEditorInstance))
+				.Font(FEditorStyle::GetFontStyle(TEXT("PropertyWindow.NormalFont")))
+			]
+			.ValueContent()
+			.MaxDesiredWidth(200.0f)
+			[
+				SNew(SHorizontalBox)
+				+ SHorizontalBox::Slot()
+				.FillWidth(1.0f)
+				[
+					SNew(SHorizontalBox)
+					+ SHorizontalBox::Slot()
+					.HAlign(HAlign_Left)
+					.AutoWidth()
+					[
+						PropertyCustomizationHelpers::MakePropertyComboBox(StackParameterData->ParameterNode->CreatePropertyHandle(), GetMaskStrings, GetMaskValue, SetMaskValue)
+					]
+				]
+			];
+		}
+		else if (ScalarParam && ScalarParam->AtlasData.bIsUsedAsAtlasPosition)
+		{
+			const FText ParameterName = FText::FromName(StackParameterData->Parameter->ParameterInfo.Name);
+
+			FDetailWidgetRow& CustomWidget = Row.CustomWidget();
+			CustomWidget
+				.FilterString(ParameterName)
+				.NameContent()
+				[
+					SNew(STextBlock)
+					.Text(ParameterName)
+					.ToolTipText(FMaterialPropertyHelpers::GetParameterExpressionDescription(StackParameterData->Parameter, MaterialEditorInstance))
+					.Font(FEditorStyle::GetFontStyle(TEXT("PropertyWindow.NormalFont")))
+				]
+				.ValueContent()
+				.HAlign(HAlign_Fill)
+				.MaxDesiredWidth(400.0f)
+				[
+					SNew(SObjectPropertyEntryBox)
+					.ObjectPath(this, &SMaterialParametersOverviewTreeItem::GetCurvePath, ScalarParam)
+					.AllowedClass(UCurveLinearColor::StaticClass())
+					.NewAssetFactories(TArray<UFactory*>())
+					.DisplayThumbnail(true)
+					.ThumbnailPool(Tree.Pin()->GetTreeThumbnailPool())
+					.OnShouldSetAsset(FOnShouldSetAsset::CreateStatic(&FMaterialPropertyHelpers::OnShouldSetCurveAsset, ScalarParam->AtlasData.Atlas))
+					.OnObjectChanged(FOnSetObject::CreateStatic(&FMaterialPropertyHelpers::SetPositionFromCurveAsset, ScalarParam->AtlasData.Atlas, ScalarParam, StackParameterData->ParameterHandle, (UObject*)MaterialEditorInstance))
+					.DisplayCompactSize(true)
+				];
+			
+		}
+		else if (!CompMaskParam)
+		{
+			FNodeWidgets StoredNodeWidgets = Node.CreateNodeWidgets();
+			TSharedRef<SWidget> StoredRightSideWidget = StoredNodeWidgets.ValueWidget.ToSharedRef();
+			StackParameterData->ParameterNode->CreatePropertyHandle()->MarkResetToDefaultCustomized(true);
+			FDetailWidgetRow& CustomWidget = Row.CustomWidget();
+			CustomWidget
+			.FilterString(NameOverride)
+			.NameContent()
+			[
+				SNew(SHorizontalBox)
+				+SHorizontalBox::Slot()
+				.VAlign(VAlign_Center)
+				[
+					SNew(STextBlock)
+					.Text(NameOverride)
+					.ToolTipText(FMaterialPropertyHelpers::GetParameterExpressionDescription(StackParameterData->Parameter, MaterialEditorInstance))
+					.Font(FEditorStyle::GetFontStyle(TEXT("PropertyWindow.NormalFont")))
+				]
+			]
+			.ValueContent()
+			[
+				StoredRightSideWidget
+			];
+
+		}
+		else
+		{
+			TSharedPtr<IPropertyHandle> RMaskProperty = StackParameterData->ParameterNode->CreatePropertyHandle()->GetChildHandle("R");
+			TSharedPtr<IPropertyHandle> GMaskProperty = StackParameterData->ParameterNode->CreatePropertyHandle()->GetChildHandle("G");
+			TSharedPtr<IPropertyHandle> BMaskProperty = StackParameterData->ParameterNode->CreatePropertyHandle()->GetChildHandle("B");
+			TSharedPtr<IPropertyHandle> AMaskProperty = StackParameterData->ParameterNode->CreatePropertyHandle()->GetChildHandle("A");
+			FDetailWidgetRow& CustomWidget = Row.CustomWidget();
+			CustomWidget
+			.FilterString(NameOverride)
+			.NameContent()
+			[
+				SNew(SHorizontalBox)
+				+ SHorizontalBox::Slot()
+				.VAlign(VAlign_Center)
+				[
+					SNew(STextBlock)
+					.Text(NameOverride)
+					.ToolTipText(FMaterialPropertyHelpers::GetParameterExpressionDescription(StackParameterData->Parameter, MaterialEditorInstance))
+					.Font(FEditorStyle::GetFontStyle(TEXT("PropertyWindow.NormalFont")))
+				]
+			]
+			.ValueContent()
+			.MaxDesiredWidth(200.0f)
+			[
+				SNew(SHorizontalBox)
+				+ SHorizontalBox::Slot()
+				.FillWidth(1.0f)
+				[
+					SNew(SHorizontalBox)
+					+ SHorizontalBox::Slot()
+					.HAlign(HAlign_Left)
+					.AutoWidth()
+					[
+						RMaskProperty->CreatePropertyNameWidget(FText::GetEmpty(), FText::GetEmpty(), false)
+					]
+					+ SHorizontalBox::Slot()
+					.HAlign(HAlign_Left)
+					.AutoWidth()
+					[
+						RMaskProperty->CreatePropertyValueWidget()
+					]
+					+ SHorizontalBox::Slot()
+					.HAlign(HAlign_Left)
+					.Padding(FMargin(10.0f, 0.0f, 0.0f, 0.0f))
+					.AutoWidth()
+					[
+						GMaskProperty->CreatePropertyNameWidget(FText::GetEmpty(), FText::GetEmpty(), false)
+					]
+					+ SHorizontalBox::Slot()
+					.HAlign(HAlign_Left)
+					.AutoWidth()
+					[
+						GMaskProperty->CreatePropertyValueWidget()
+					]
+					+ SHorizontalBox::Slot()
+					.HAlign(HAlign_Left)
+					.Padding(FMargin(10.0f, 0.0f, 0.0f, 0.0f))
+					.AutoWidth()
+					[
+						BMaskProperty->CreatePropertyNameWidget(FText::GetEmpty(), FText::GetEmpty(), false)
+					]
+					+ SHorizontalBox::Slot()
+					.HAlign(HAlign_Left)
+					.AutoWidth()
+					[
+						BMaskProperty->CreatePropertyValueWidget()
+					]
+					+ SHorizontalBox::Slot()
+					.HAlign(HAlign_Left)
+					.Padding(FMargin(10.0f, 0.0f, 0.0f, 0.0f))
+					.AutoWidth()
+					[
+						AMaskProperty->CreatePropertyNameWidget(FText::GetEmpty(), FText::GetEmpty(), false)
+					]
+					+ SHorizontalBox::Slot()
+					.HAlign(HAlign_Left)
+					.AutoWidth()
+					[
+						AMaskProperty->CreatePropertyValueWidget()
+					]
+				]	
+			];
+		}
+
+		FNodeWidgets NodeWidgets = Node.CreateNodeWidgets();
+		LeftSideWidget = NodeWidgets.NameWidget.ToSharedRef();
+		RightSideWidget = NodeWidgets.ValueWidget.ToSharedRef();
+	}
+// END PROPERTY
+
+// PROPERTY CHILD ----------------------------------------
+	if (StackParameterData->StackDataType == EStackDataType::PropertyChild)
+	{
+		FNodeWidgets NodeWidgets = StackParameterData->ParameterNode->CreateNodeWidgets();
+		LeftSideWidget = NodeWidgets.NameWidget.ToSharedRef();
+		RightSideWidget = NodeWidgets.ValueWidget.ToSharedRef();
+	}
+// END PROPERTY CHILD
+
+// FINAL WRAPPER
+	{
+		WrapperWidget->AddSlot()
+			.AutoHeight()
+			[
+				SNew(SBorder)
+				.Padding(0.0f)
+				.BorderImage(this, &SMaterialParametersOverviewTreeItem::GetBorderImage)
+				[
+					SNew(SSplitter)
+					.Style(FEditorStyle::Get(), "DetailsView.Splitter")
+					.PhysicalSplitterHandleSize(1.0f)
+					.HitDetectionSplitterHandleSize(5.0f)
+					+ SSplitter::Slot()
+					.Value(ColumnSizeData.LeftColumnWidth)
+					.OnSlotResized(ColumnSizeData.OnWidthChanged)
+					.Value(0.25f)
+					[
+						SNew(SHorizontalBox)
+						+ SHorizontalBox::Slot()
+						.AutoWidth()
+						.VAlign(VAlign_Center)
+						.Padding(FMargin(3.0f))
+						[
+							SNew(SExpanderArrow, SharedThis(this))
+						]
+						+ SHorizontalBox::Slot()
+						.Padding(FMargin(2.0f))
+						.VAlign(VAlign_Center)
+						[
+							LeftSideWidget
+						]
+					]
+					+ SSplitter::Slot()
+					.Value(ColumnSizeData.RightColumnWidth)
+					.OnSlotResized(ColumnSizeData.OnWidthChanged)
+					[
+						SNew(SHorizontalBox)
+						+ SHorizontalBox::Slot()
+						.MaxWidth(350.0f)
+						.Padding(FMargin(5.0f, 2.0f, 0.0f, 2.0f))
+						[
+							RightSideWidget
+						]
+					]
+				]
+			];
+	}
+
+
+	this->ChildSlot
+		[
+			WrapperWidget
+		];
+
+
+	STableRow< TSharedPtr<FSortedParamData> >::ConstructInternal(
+		STableRow::FArguments()
+		.Style(FEditorStyle::Get(), "DetailsView.TreeView.TableRow")
+		.ShowSelection(false),
+		InOwnerTableView
+	);
+}
+
+void SMaterialParametersOverviewTree::Construct(const FArguments& InArgs)
+{
+	bHasAnyParameters = false;
+	ColumnWidth = 0.5f;
+	MaterialEditorInstance = InArgs._InMaterialEditorInstance;
+	Owner = InArgs._InOwner;
+	CreateGroupsWidget();
+
+
+	STreeView<TSharedPtr<FSortedParamData>>::Construct(
+		STreeView::FArguments()
+		.TreeItemsSource(&SortedParameters)
+		.SelectionMode(ESelectionMode::None)
+		.OnGenerateRow(this, &SMaterialParametersOverviewTree::OnGenerateRowMaterialLayersFunctionsTreeView)
+		.OnGetChildren(this, &SMaterialParametersOverviewTree::OnGetChildrenMaterialLayersFunctionsTreeView)
+		.OnExpansionChanged(this, &SMaterialParametersOverviewTree::OnExpansionChanged)
+	);
+}
+
+TSharedRef< ITableRow > SMaterialParametersOverviewTree::OnGenerateRowMaterialLayersFunctionsTreeView(TSharedPtr<FSortedParamData> Item, const TSharedRef< STableViewBase >& OwnerTable)
+{
+	TSharedRef< SMaterialParametersOverviewTreeItem > ReturnRow = SNew(SMaterialParametersOverviewTreeItem, OwnerTable)
+		.StackParameterData(Item)
+		.MaterialEditorInstance(MaterialEditorInstance)
+		.InTree(SharedThis(this));
+	return ReturnRow;
+}
+
+void SMaterialParametersOverviewTree::OnGetChildrenMaterialLayersFunctionsTreeView(TSharedPtr<FSortedParamData> InParent, TArray< TSharedPtr<FSortedParamData> >& OutChildren)
+{
+	OutChildren = InParent->Children;
+}
+
+void SMaterialParametersOverviewTree::OnExpansionChanged(TSharedPtr<FSortedParamData> Item, bool bIsExpanded)
+{
+	bool* ExpansionValue = MaterialEditorInstance->OriginalMaterial->ParameterOverviewExpansion.Find(Item->NodeKey);
+	if (ExpansionValue == nullptr)
+	{
+		MaterialEditorInstance->OriginalMaterial->ParameterOverviewExpansion.Add(Item->NodeKey, bIsExpanded);
+	}
+	else if (*ExpansionValue != bIsExpanded)
+	{
+		MaterialEditorInstance->OriginalMaterial->ParameterOverviewExpansion.Emplace(Item->NodeKey, bIsExpanded);
+	}
+	// Expand any children that are also expanded
+	for (auto Child : Item->Children)
+	{
+		bool* ChildExpansionValue = MaterialEditorInstance->OriginalMaterial->ParameterOverviewExpansion.Find(Child->NodeKey);
+		if (ChildExpansionValue != nullptr && *ChildExpansionValue == true)
+		{
+			SetItemExpansion(Child, true);
+		}
+	}
+}
+
+void SMaterialParametersOverviewTree::SetParentsExpansionState()
+{
+	for (const auto& Pair : SortedParameters)
+	{
+		if (Pair->Children.Num())
+		{
+			bool* bIsExpanded = MaterialEditorInstance->OriginalMaterial->ParameterOverviewExpansion.Find(Pair->NodeKey);
+			if (bIsExpanded)
+			{
+				SetItemExpansion(Pair, *bIsExpanded);
+			}
+			else
+			{
+				SetItemExpansion(Pair, true);
+			}
+		}
+	}
+}
+
+TSharedPtr<class FAssetThumbnailPool> SMaterialParametersOverviewTree::GetTreeThumbnailPool()
+{
+	return Generator->GetGeneratedThumbnailPool();
+}
+
+void SMaterialParametersOverviewTree::CreateGroupsWidget()
+{
+	check(MaterialEditorInstance);
+	MaterialEditorInstance->RegenerateArrays();
+	UnsortedParameters.Empty();
+	SortedParameters.Empty();
+	FPropertyEditorModule& Module = FModuleManager::LoadModuleChecked<FPropertyEditorModule>("PropertyEditor");
+	if (!Generator.IsValid())
+	{
+		FPropertyRowGeneratorArgs Args;
+		Generator = Module.CreatePropertyRowGenerator(Args);
+
+		TArray<UObject*> Objects;
+		Objects.Add(MaterialEditorInstance);
+		Generator->SetObjects(Objects);
+	}
+	else
+	{
+		TArray<UObject*> Objects;
+		Objects.Add(MaterialEditorInstance);
+		Generator->SetObjects(Objects);
+	}
+	const TArray<TSharedRef<IDetailTreeNode>> TestData = Generator->GetRootTreeNodes();
+	TSharedPtr<IDetailTreeNode> Category = TestData[0];
+	TSharedPtr<IDetailTreeNode> ParameterGroups;
+	TArray<TSharedRef<IDetailTreeNode>> Children;
+	Category->GetChildren(Children);
+
+	for (int32 ChildIdx = 0; ChildIdx < Children.Num(); ChildIdx++)
+	{
+		if (Children[ChildIdx]->CreatePropertyHandle().IsValid() &&
+			Children[ChildIdx]->CreatePropertyHandle()->GetProperty()->GetName() == "ParameterGroups")
+		{
+			ParameterGroups = Children[ChildIdx];
+			break;
+		}
+	}
+
+	Children.Empty();
+	ParameterGroups->GetChildren(Children);
+	for (int32 GroupIdx = 0; GroupIdx < Children.Num(); ++GroupIdx)
+	{
+		bHasAnyParameters = true;
+		TArray<void*> GroupPtrs;
+		TSharedPtr<IPropertyHandle> ChildHandle = Children[GroupIdx]->CreatePropertyHandle();
+		ChildHandle->AccessRawData(GroupPtrs);
+		auto GroupIt = GroupPtrs.CreateConstIterator();
+		const FEditorParameterGroup* ParameterGroupPtr = reinterpret_cast<FEditorParameterGroup*>(*GroupIt);
+		const FEditorParameterGroup& ParameterGroup = *ParameterGroupPtr;
+
+		for (int32 ParamIdx = 0; ParamIdx < ParameterGroup.Parameters.Num(); ParamIdx++)
+		{
+			UDEditorParameterValue* Parameter = ParameterGroup.Parameters[ParamIdx];
+
+			TSharedPtr<IPropertyHandle> ParametersArrayProperty = ChildHandle->GetChildHandle("Parameters");
+			TSharedPtr<IPropertyHandle> ParameterProperty = ParametersArrayProperty->GetChildHandle(ParamIdx);
+			TSharedPtr<IPropertyHandle> ParameterValueProperty = ParameterProperty->GetChildHandle("ParameterValue");
+
+			{
+				FUnsortedParamData NonLayerProperty;
+				UDEditorScalarParameterValue* ScalarParam = Cast<UDEditorScalarParameterValue>(Parameter);
+
+				if (ScalarParam && ScalarParam->SliderMax > ScalarParam->SliderMin)
+				{
+					ParameterValueProperty->SetInstanceMetaData("UIMin", FString::Printf(TEXT("%f"), ScalarParam->SliderMin));
+					ParameterValueProperty->SetInstanceMetaData("UIMax", FString::Printf(TEXT("%f"), ScalarParam->SliderMax));
+				}
+
+				NonLayerProperty.Parameter = Parameter;
+				NonLayerProperty.ParameterGroup = ParameterGroup;
+				NonLayerProperty.ParameterNode = Generator->FindTreeNode(ParameterValueProperty);
+				NonLayerProperty.ParameterHandle = NonLayerProperty.ParameterNode->CreatePropertyHandle();
+				NonLayerProperty.UnsortedName = Parameter->ParameterInfo.Name;
+
+				UnsortedParameters.Add(NonLayerProperty);
+			}
+		}
+	}
+	ShowSubParameters();
+	RequestTreeRefresh();
+	SetParentsExpansionState();
+}
+
+
+
+void SMaterialParametersOverviewTree::ShowSubParameters()
+{
+	for (FUnsortedParamData Property : UnsortedParameters)
+	{
+		UDEditorParameterValue* Parameter = Property.Parameter;
+		{
+			TSharedPtr<FSortedParamData> GroupProperty(new FSortedParamData());
+			GroupProperty->StackDataType = EStackDataType::Group;
+			GroupProperty->ParameterInfo.Index = Parameter->ParameterInfo.Index;
+			GroupProperty->ParameterInfo.Association = Parameter->ParameterInfo.Association;
+			GroupProperty->Group = Property.ParameterGroup;
+			GroupProperty->NodeKey = FString::FromInt(GroupProperty->ParameterInfo.Index) + FString::FromInt(GroupProperty->ParameterInfo.Association) + Property.ParameterGroup.GroupName.ToString();
+
+			bool bAddNewGroup = true;
+			for (TSharedPtr<struct FSortedParamData> GroupChild : SortedParameters)
+			{
+				if (GroupChild->NodeKey == GroupProperty->NodeKey)
+				{
+					bAddNewGroup = false;
+				}
+			}
+			if (bAddNewGroup)
+			{
+				SortedParameters.Add(GroupProperty);
+			}
+
+			TSharedPtr<FSortedParamData> ChildProperty(new FSortedParamData());
+			ChildProperty->StackDataType = EStackDataType::Property;
+			ChildProperty->Parameter = Parameter;
+			ChildProperty->ParameterInfo.Index = Parameter->ParameterInfo.Index;
+			ChildProperty->ParameterInfo.Association = Parameter->ParameterInfo.Association;
+			ChildProperty->ParameterNode = Property.ParameterNode;
+			ChildProperty->PropertyName = Property.UnsortedName;
+			ChildProperty->NodeKey = FString::FromInt(ChildProperty->ParameterInfo.Index) + FString::FromInt(ChildProperty->ParameterInfo.Association) +  Property.ParameterGroup.GroupName.ToString() + Property.UnsortedName.ToString();
+
+
+			UDEditorStaticComponentMaskParameterValue* CompMaskParam = Cast<UDEditorStaticComponentMaskParameterValue>(Parameter);
+			if (!CompMaskParam)
+			{
+				TArray<TSharedRef<IDetailTreeNode>> ParamChildren;
+				Property.ParameterNode->GetChildren(ParamChildren);
+				for (int32 ParamChildIdx = 0; ParamChildIdx < ParamChildren.Num(); ParamChildIdx++)
+				{
+					TSharedPtr<FSortedParamData> ParamChildProperty(new FSortedParamData());
+					ParamChildProperty->StackDataType = EStackDataType::PropertyChild;
+					ParamChildProperty->ParameterNode = ParamChildren[ParamChildIdx];
+					ParamChildProperty->ParameterHandle = ParamChildProperty->ParameterNode->CreatePropertyHandle();
+					ParamChildProperty->ParameterInfo.Index = Parameter->ParameterInfo.Index;
+					ParamChildProperty->ParameterInfo.Association = Parameter->ParameterInfo.Association;
+					ParamChildProperty->Parameter = ChildProperty->Parameter;
+					ChildProperty->Children.Add(ParamChildProperty);
+				}
+			}
+			for (TSharedPtr<struct FSortedParamData> GroupChild : SortedParameters)
+			{
+				if (GroupChild->Group.GroupName == Property.ParameterGroup.GroupName
+					&& GroupChild->ParameterInfo.Association == ChildProperty->ParameterInfo.Association
+					&&  GroupChild->ParameterInfo.Index == ChildProperty->ParameterInfo.Index)
+				{
+					GroupChild->Children.Add(ChildProperty);
+				}
+			}
+
+		}
+	}
+}
+
+const FSlateBrush* SMaterialParametersOverviewPanel::GetBackgroundImage() const
+{
+	return FEditorStyle::GetBrush("DetailsView.CategoryTop_Hovered");
+}
+
+int32 SMaterialParametersOverviewPanel::GetPanelIndex() const
+{
+	return NestedTree && NestedTree->HasAnyParameters() ? 1 : 0;
+}
+
+void SMaterialParametersOverviewPanel::Refresh()
+{
+	TSharedPtr<SHorizontalBox> HeaderBox;
+	NestedTree->CreateGroupsWidget();
+
+	FOnClicked 	OnChildButtonClicked = FOnClicked();
+	if (MaterialEditorInstance->OriginalFunction)
+	{
+		OnChildButtonClicked = FOnClicked::CreateStatic(&FMaterialPropertyHelpers::OnClickedSaveNewFunctionInstance, ImplicitConv<UMaterialFunctionInterface*>(MaterialEditorInstance->OriginalFunction), ImplicitConv<UMaterialInterface*>(MaterialEditorInstance->PreviewMaterial), ImplicitConv<UObject*>(MaterialEditorInstance));
+	}
+	else
+	{
+		OnChildButtonClicked = FOnClicked::CreateStatic(&FMaterialPropertyHelpers::OnClickedSaveNewMaterialInstance, ImplicitConv<UMaterialInterface*>(MaterialEditorInstance->OriginalMaterial), ImplicitConv<UObject*>(MaterialEditorInstance));
+	}
+
+	{
+		this->ChildSlot
+			[
+				SNew(SVerticalBox)
+				+SVerticalBox::Slot()
+				.AutoHeight()
+				[
+					SNew(SBorder)
+					.BorderImage(this, &SMaterialParametersOverviewPanel::GetBackgroundImage)
+					.Padding(FMargin(4.0f))
+					[
+						SNew(SVerticalBox)
+						+ SVerticalBox::Slot()
+						.AutoHeight()
+						[
+							SAssignNew(HeaderBox, SHorizontalBox)
+							+ SHorizontalBox::Slot()
+							.Padding(FMargin(3.0f, 1.0f))
+							.HAlign(HAlign_Left)
+							.AutoWidth()
+							.VAlign(VAlign_Center)
+							[
+								SNew(STextBlock)
+								.Text(LOCTEXT("ParameterDefaults", "Parameter Defaults"))
+								.Font(FEditorStyle::GetFontStyle("DetailsView.CategoryFontStyle"))
+								.ShadowOffset(FVector2D(1.0f, 1.0f))
+							]
+						]
+						+ SVerticalBox::Slot()
+						.Padding(FMargin(3.0f, 2.0f, 3.0f, 3.0f))
+						.AutoHeight()
+						[
+							SNew(SBorder)
+							.BorderImage(FEditorStyle::GetBrush("DetailsView.CategoryTop"))
+							[
+								SNew(SWidgetSwitcher)
+								.WidgetIndex(this, &SMaterialParametersOverviewPanel::GetPanelIndex)
+								+SWidgetSwitcher::Slot()
+								[
+									SNew(STextBlock)
+									.Text(LOCTEXT("AddParams", "Add parameters to see them here."))
+								]
+								+SWidgetSwitcher::Slot()
+								[
+									NestedTree.ToSharedRef()
+								]	
+							]
+						]
+					]
+				]
+			];
+
+		HeaderBox->AddSlot()
+			.FillWidth(1.0f)
+			[
+				SNullWidget::NullWidget
+			];
+		
+		if (NestedTree->HasAnyParameters())
+		{
+			HeaderBox->AddSlot()
+				.AutoWidth()
+				.Padding(2.0f)
+				[
+					SNew(SButton)
+					.ButtonStyle(FEditorStyle::Get(), "FlatButton.DarkGrey")
+					.HAlign(HAlign_Center)
+					.OnClicked(OnChildButtonClicked)
+					.ToolTipText(LOCTEXT("SaveToChildInstance", "Save To Child Instance"))
+					.Content()
+					[
+						SNew(SHorizontalBox)
+						+ SHorizontalBox::Slot()
+						.AutoWidth()
+						[
+							SNew(STextBlock)
+							.Font(FEditorStyle::Get().GetFontStyle("FontAwesome.10"))
+							.TextStyle(FEditorStyle::Get(), "NormalText.Important")
+							.Text(FText::FromString(FString(TEXT("\xf0c7 \xf149"))) /*fa-filter*/)
+						]
+						+ SHorizontalBox::Slot()
+						.AutoWidth()
+						[
+							SNew(STextBlock)
+							.TextStyle(FEditorStyle::Get(), "NormalText.Important")
+							.Text(FText::FromString(FString(TEXT(" Save Child"))) /*fa-filter*/)
+						]
+					]
+				];
+		}
+	}
+	
+}
+
+
+void SMaterialParametersOverviewPanel::Construct(const FArguments& InArgs)
+{
+	NestedTree = SNew(SMaterialParametersOverviewTree)
+		.InMaterialEditorInstance(InArgs._InMaterialEditorInstance)
+		.InOwner(SharedThis(this));
+
+	MaterialEditorInstance = InArgs._InMaterialEditorInstance;
+	FEditorSupportDelegates::UpdateUI.AddSP(this, &SMaterialParametersOverviewPanel::Refresh);
+
+}
+
+void SMaterialParametersOverviewPanel::UpdateEditorInstance(UMaterialEditorPreviewParameters* InMaterialEditorInstance)
+{
+	NestedTree->MaterialEditorInstance = InMaterialEditorInstance;
+	Refresh();
+}
+
+
+TSharedPtr<class IPropertyRowGenerator> SMaterialParametersOverviewPanel::GetGenerator()
+{
+	 return NestedTree->GetGenerator();
+}
+
+#undef LOCTEXT_NAMESPACE
\ No newline at end of file
diff --git a/Engine/Source/Editor/MaterialEditor/Private/SMaterialParametersOverviewWidget.h b/Engine/Source/Editor/MaterialEditor/Private/SMaterialParametersOverviewWidget.h
new file mode 100644
--- /dev/null
+++ b/Engine/Source/Editor/MaterialEditor/Private/SMaterialParametersOverviewWidget.h
@@ -0,0 +1,152 @@
+// Copyright 1998-2018 Epic Games, Inc. All Rights Reserved.
+
+#pragma once
+
+#include "CoreMinimal.h"
+#include "UObject/ObjectMacros.h"
+#include "UObject/Object.h"
+#include "Engine/EngineTypes.h"
+#include "Widgets/DeclarativeSyntaxSupport.h"
+#include "Widgets/SCompoundWidget.h"
+#include "Widgets/Views/STreeView.h"
+#include "Materials/Material.h"
+#include "IDetailTreeNode.h"
+#include "IDetailPropertyRow.h"
+#include "MaterialPropertyHelpers.h"
+
+class IPropertyHandle;
+class SMaterialParametersOverviewTree;
+class UMaterialEditorPreviewParameters;
+
+// ********* SMaterialParametersOverviewTreeItem *******
+class SMaterialParametersOverviewTreeItem : public STableRow< TSharedPtr<FSortedParamData> >
+{
+public:
+
+	SLATE_BEGIN_ARGS(SMaterialParametersOverviewTreeItem)
+		: _StackParameterData(nullptr),
+		_MaterialEditorInstance(nullptr)
+	{}
+
+	/** The item content. */
+	SLATE_ARGUMENT(TSharedPtr<FSortedParamData>, StackParameterData)
+	SLATE_ARGUMENT(UMaterialEditorPreviewParameters*, MaterialEditorInstance)
+	SLATE_ARGUMENT(TSharedPtr<SMaterialParametersOverviewTree>, InTree)
+	SLATE_END_ARGS()
+
+	void RefreshOnRowChange(const FAssetData& AssetData, TSharedPtr<SMaterialParametersOverviewTree> InTree);
+
+	/**
+	* Construct the widget
+	*
+	* @param InArgs   A declaration from which to construct the widget
+	*/
+	void Construct(const FArguments& InArgs, const TSharedRef<STableViewBase>& InOwnerTableView);
+
+private:
+	FString GetCurvePath(class UDEditorScalarParameterValue* Parameter) const;
+
+	const FSlateBrush* GetBorderImage() const;
+
+private:
+
+	/** The node info to build the tree view row from. */
+	TSharedPtr<FSortedParamData> StackParameterData;
+
+	/** The tree that contains this item */
+	TWeakPtr<SMaterialParametersOverviewTree> Tree;
+
+	/** The set of material parameters this is associated with */
+	UMaterialEditorPreviewParameters* MaterialEditorInstance;
+
+	FMaterialTreeColumnSizeData ColumnSizeData;
+};
+
+// ********* SMaterialParametersOverviewPanel *******
+class SMaterialParametersOverviewPanel : public SCompoundWidget
+{
+public:
+	SLATE_BEGIN_ARGS(SMaterialParametersOverviewPanel)
+		: _InMaterialEditorInstance(nullptr)
+	{}
+
+	SLATE_ARGUMENT(UMaterialEditorPreviewParameters*, InMaterialEditorInstance)
+	SLATE_END_ARGS()
+	void Refresh();
+	void Construct(const FArguments& InArgs);
+	void UpdateEditorInstance(UMaterialEditorPreviewParameters* InMaterialEditorInstance);
+	TSharedPtr<class IPropertyRowGenerator> GetGenerator();
+
+private:
+	const FSlateBrush* GetBackgroundImage() const;
+
+	int32 GetPanelIndex() const;
+
+private:
+	/** The set of material parameters this is associated with */
+	class UMaterialEditorPreviewParameters* MaterialEditorInstance;
+
+	/** The tree contained in this item */
+	TSharedPtr<class SMaterialParametersOverviewTree> NestedTree;
+};
+
+// ********* SMaterialParametersOverviewTree *******
+class SMaterialParametersOverviewTree : public STreeView<TSharedPtr<FSortedParamData>>
+{
+	friend class SMaterialParametersOverviewTreeItem;
+
+public:
+
+	SLATE_BEGIN_ARGS(SMaterialParametersOverviewTree)
+		: _InMaterialEditorInstance(nullptr)
+		, _InOwner(nullptr)
+	{}
+
+	SLATE_ARGUMENT(UMaterialEditorPreviewParameters*, InMaterialEditorInstance)
+	SLATE_ARGUMENT(TSharedPtr<SMaterialParametersOverviewPanel>, InOwner)
+	SLATE_END_ARGS()
+
+	/** Constructs this widget with InArgs */
+	void Construct(const FArguments& InArgs);
+	TSharedRef< ITableRow > OnGenerateRowMaterialLayersFunctionsTreeView(TSharedPtr<FSortedParamData> Item, const TSharedRef< STableViewBase >& OwnerTable);
+	void OnGetChildrenMaterialLayersFunctionsTreeView(TSharedPtr<FSortedParamData> InParent, TArray< TSharedPtr<FSortedParamData> >& OutChildren);
+	void OnExpansionChanged(TSharedPtr<FSortedParamData> Item, bool bIsExpanded);
+	void SetParentsExpansionState();
+
+	float OnGetLeftColumnWidth() const { return 1.0f - ColumnWidth; }
+	float OnGetRightColumnWidth() const { return ColumnWidth; }
+	void OnSetColumnWidth(float InWidth) { ColumnWidth = InWidth; }
+	TSharedPtr<class FAssetThumbnailPool> GetTreeThumbnailPool();
+
+	/** Object that stores all of the possible parameters we can edit */
+	UMaterialEditorPreviewParameters* MaterialEditorInstance;
+
+	/** Builds the custom parameter groups category */
+	void CreateGroupsWidget();
+
+	TWeakPtr<SMaterialParametersOverviewPanel> GetOwner() { return Owner; }
+
+	TSharedPtr<class IPropertyRowGenerator> GetGenerator() { return Generator; }
+
+	bool HasAnyParameters() const { return bHasAnyParameters; }
+
+protected:
+
+	void ShowSubParameters();
+
+private:
+
+	TArray<TSharedPtr<FSortedParamData>> SortedParameters;
+
+	TArray<FUnsortedParamData> UnsortedParameters;
+
+	/** The actual width of the right column.  The left column is 1-ColumnWidth */
+	float ColumnWidth;
+
+	TWeakPtr<SMaterialParametersOverviewPanel> Owner;
+
+	TSharedPtr<class IPropertyRowGenerator> Generator;
+
+	bool bHasAnyParameters;
+
+};
diff --git a/Engine/Source/Editor/MaterialEditor/Public/MaterialPropertyHelpers.h b/Engine/Source/Editor/MaterialEditor/Public/MaterialPropertyHelpers.h
--- a/Engine/Source/Editor/MaterialEditor/Public/MaterialPropertyHelpers.h
+++ b/Engine/Source/Editor/MaterialEditor/Public/MaterialPropertyHelpers.h
@@ -34,7 +34,7 @@ enum EStackDataType
 };
 
 USTRUCT()
-struct MATERIALEDITOR_API FStackSortedData
+struct MATERIALEDITOR_API FSortedParamData
 {
 	GENERATED_USTRUCT_BODY()
 
@@ -53,12 +53,12 @@ public:
 
 	TSharedPtr<class IPropertyHandle> ParameterHandle;
 
-	TArray<TSharedPtr<struct FStackSortedData>> Children;
+	TArray<TSharedPtr<struct FSortedParamData>> Children;
 
 	FString NodeKey;
 };
 
-struct FLayerParameterUnsortedData
+struct FUnsortedParamData
 {
 	class UDEditorParameterValue* Parameter;
 	FEditorParameterGroup ParameterGroup;
@@ -164,7 +164,7 @@ public:
 	/** Specific resets for layer and blend asses */
 	static void ResetLayerAssetToDefault(TSharedPtr<IPropertyHandle> PropertyHandle, class UDEditorParameterValue* InParameter, TEnumAsByte<EMaterialParameterAssociation> InAssociation, int32 Index, UMaterialEditorInstanceConstant* MaterialEditorInstance);
 	/** If reset to default button should show for a layer or blend asset*/
-	static bool ShouldLayerAssetShowResetToDefault(TSharedPtr<IPropertyHandle> PropertyHandle, TSharedPtr<struct FStackSortedData> InParameterData, UMaterialInterface* InMaterial);
+	static bool ShouldLayerAssetShowResetToDefault(TSharedPtr<IPropertyHandle> PropertyHandle, TSharedPtr<struct FSortedParamData> InParameterData, UMaterialInterface* InMaterial);
 
 	static void OnMaterialLayerAssetChanged(const struct FAssetData& InAssetData, int32 Index, EMaterialParameterAssociation MaterialType, TSharedPtr<class IPropertyHandle> InHandle, FMaterialLayersFunctions* InMaterialFunction);
 
@@ -175,7 +175,7 @@ public:
 	static void CopyMaterialToInstance(class UMaterialInstanceConstant* ChildInstance, TArray<struct FEditorParameterGroup> &ParameterGroups);
 	static void TransitionAndCopyParameters(class UMaterialInstanceConstant* ChildInstance, TArray<struct FEditorParameterGroup> &ParameterGroups, bool bForceCopy = false);
 	static FReply OnClickedSaveNewFunctionInstance(class UMaterialFunctionInterface* Object, class UMaterialInterface* PreviewMaterial, UObject* EditorObject);
-	static FReply OnClickedSaveNewLayerInstance(class UMaterialFunctionInterface* Object, TSharedPtr<FStackSortedData> InSortedData);
+	static FReply OnClickedSaveNewLayerInstance(class UMaterialFunctionInterface* Object, TSharedPtr<FSortedParamData> InSortedData);
 
 	static void GetVectorChannelMaskComboBoxStrings(TArray<TSharedPtr<FString>>& OutComboBoxStrings, TArray<TSharedPtr<class SToolTip>>& OutToolTips, TArray<bool>& OutRestrictedItems);
 	static FString GetVectorChannelMaskValue(class UDEditorParameterValue* InParameter);
diff --git a/Engine/Source/Editor/PropertyEditor/Private/PropertyRowGenerator.cpp b/Engine/Source/Editor/PropertyEditor/Private/PropertyRowGenerator.cpp
--- a/Engine/Source/Editor/PropertyEditor/Private/PropertyRowGenerator.cpp
+++ b/Engine/Source/Editor/PropertyEditor/Private/PropertyRowGenerator.cpp
@@ -63,7 +63,10 @@ public:
 			: TSharedPtr<class FAssetThumbnailPool>();
 	}
 
-	virtual void NotifyFinishedChangingProperties(const FPropertyChangedEvent& PropertyChangedEvent) override {}
+	virtual void NotifyFinishedChangingProperties(const FPropertyChangedEvent& PropertyChangedEvent) override 
+	{
+		Generator->OnFinishedChangingProperties().Broadcast(PropertyChangedEvent);
+	}
 
 	virtual bool DontUpdateValueWhileEditing() const override { return false; }
 
diff --git a/Engine/Source/Editor/PropertyEditor/Private/PropertyRowGenerator.h b/Engine/Source/Editor/PropertyEditor/Private/PropertyRowGenerator.h
--- a/Engine/Source/Editor/PropertyEditor/Private/PropertyRowGenerator.h
+++ b/Engine/Source/Editor/PropertyEditor/Private/PropertyRowGenerator.h
@@ -110,6 +110,7 @@ public:
 
 	const FCustomPropertyTypeLayoutMap& GetInstancedPropertyTypeLayoutMap() const;
 	void UpdateDetailRows();
+	virtual FOnFinishedChangingProperties& OnFinishedChangingProperties() override { return OnFinishedChangingPropertiesDelegate; }
 
 private:
 	void PreSetObject(int32 NumNewObjects, bool bHasStructRoots);
@@ -148,6 +149,8 @@ private:
 	TSharedRef<IPropertyUtilities> PropertyUtilities;
 	/** Utility class for accessing internal helper methods */
 	TSharedRef<IPropertyGenerationUtilities> PropertyGenerationUtilities;
+	/** Delegate called when the details panel finishes editing a property (after post edit change is called) */
+	FOnFinishedChangingProperties OnFinishedChangingPropertiesDelegate;
 
 	bool bViewingClassDefaultObject;
 };
diff --git a/Engine/Source/Editor/PropertyEditor/Public/IPropertyRowGenerator.h b/Engine/Source/Editor/PropertyEditor/Public/IPropertyRowGenerator.h
--- a/Engine/Source/Editor/PropertyEditor/Public/IPropertyRowGenerator.h
+++ b/Engine/Source/Editor/PropertyEditor/Public/IPropertyRowGenerator.h
@@ -71,4 +71,5 @@ public:
 	virtual void UnregisterInstancedCustomPropertyTypeLayout(FName PropertyTypeName, TSharedPtr<IPropertyTypeIdentifier> Identifier = nullptr) = 0;
 
 	virtual TSharedPtr<class FAssetThumbnailPool> GetGeneratedThumbnailPool() = 0;
+	virtual FOnFinishedChangingProperties& OnFinishedChangingProperties() = 0;
 };
\ No newline at end of file
diff --git a/Engine/Source/Editor/UnrealEd/Classes/MaterialEditor/MaterialEditorPreviewParameters.h b/Engine/Source/Editor/UnrealEd/Classes/MaterialEditor/MaterialEditorPreviewParameters.h
--- a/Engine/Source/Editor/UnrealEd/Classes/MaterialEditor/MaterialEditorPreviewParameters.h
+++ b/Engine/Source/Editor/UnrealEd/Classes/MaterialEditor/MaterialEditorPreviewParameters.h
@@ -39,8 +39,6 @@ public:
 	UPROPERTY()
 	class UMaterial* OriginalMaterial;
 
-	TWeakPtr<class IDetailsView> DetailsView;
-
 	//~ Begin UObject Interface.
 	virtual void PostEditChangeProperty(FPropertyChangedEvent& PropertyChangedEvent) override;
 #if WITH_EDITOR
diff --git a/Engine/Source/Editor/UnrealEd/Private/PreviewMaterial.cpp b/Engine/Source/Editor/UnrealEd/Private/PreviewMaterial.cpp
--- a/Engine/Source/Editor/UnrealEd/Private/PreviewMaterial.cpp
+++ b/Engine/Source/Editor/UnrealEd/Private/PreviewMaterial.cpp
@@ -267,13 +267,6 @@ void UMaterialEditorPreviewParameters::PostEditChangeProperty(FPropertyChangedEv
 				OriginalFunction->PreviewMaterial->PostEditChangeProperty(PropertyChangedEvent);
 			}
 		}
-	
-		if(DetailsView.IsValid())
-		{
-			// Tell our source instance to update itself so the preview updates.
-			DetailsView.Pin()->OnFinishedChangingProperties().Broadcast(PropertyChangedEvent);
-		}
-
 	}
 }
 
diff --git a/Engine/Source/Runtime/Engine/Classes/Materials/MaterialInterface.h b/Engine/Source/Runtime/Engine/Classes/Materials/MaterialInterface.h
--- a/Engine/Source/Runtime/Engine/Classes/Materials/MaterialInterface.h
+++ b/Engine/Source/Runtime/Engine/Classes/Materials/MaterialInterface.h
@@ -261,6 +261,9 @@ public:
 	UPROPERTY()
 	TMap<FString, bool> LayerParameterExpansion;
 
+	UPROPERTY()
+	TMap<FString, bool> ParameterOverviewExpansion;
+
 	/** Importing data and options used for this material */
 	UPROPERTY(EditAnywhere, Instanced, Category = ImportSettings)
 	class UAssetImportData* AssetImportData;
diff --git a/Engine/Source/Runtime/Engine/Private/Materials/Material.cpp b/Engine/Source/Runtime/Engine/Private/Materials/Material.cpp
--- a/Engine/Source/Runtime/Engine/Private/Materials/Material.cpp
+++ b/Engine/Source/Runtime/Engine/Private/Materials/Material.cpp
@@ -1179,10 +1179,9 @@ void UMaterial::OverrideVectorParameterDefault(const FMaterialParameterInfo& Par
 	bool bShouldRecacheMaterialExpressions = false;
 
 	FMaterialResource* Resource = GetMaterialResource(InFeatureLevel);
-	// Iterate over both the 2D textures and cube texture expressions.
 	const TArray<TRefCountPtr<FMaterialUniformExpression> >& UniformExpressions = Resource->GetUniformVectorParameterExpressions();
 
-	// Iterate over each of the material's texture expressions.
+	// Iterate over each of the material's vector expressions.
 	for (FMaterialUniformExpression* UniformExpression : UniformExpressions)
 	{
 		if (UniformExpression->GetType() == &FMaterialUniformExpressionVectorParameter::StaticType)
@@ -1299,15 +1298,16 @@ bool UMaterial::IsCompilingOrHadCompileError(ERHIFeatureLevel::Type InFeatureLev
 #if WITH_EDITOR
 bool UMaterial::SetVectorParameterValueEditorOnly(FName ParameterName, FLinearColor InValue)
 {
+	bool bSetNewValue = false;
+	UMaterialExpressionVectorParameter* Parameter = nullptr;
 	for (UMaterialExpression* Expression : Expressions)
 	{
-		if (UMaterialExpressionVectorParameter* Parameter = Cast<UMaterialExpressionVectorParameter>(Expression))
+		Parameter = Cast<UMaterialExpressionVectorParameter>(Expression);
+		if (Parameter && Parameter->SetParameterValue(ParameterName, InValue))
 		{
-			if (Parameter->SetParameterValue(ParameterName, InValue))
-			{
-				return true;
-				// Warning: in the case of duplicate parameters with different default values, this will find the first in the expression array, not necessarily the one that's used for rendering
-			}
+			bSetNewValue = true;
+			break;
+			// Warning: in the case of duplicate parameters with different default values, this will find the first in the expression array, not necessarily the one that's used for rendering
 		}
 		else if (UMaterialExpressionMaterialFunctionCall* FunctionCall = Cast<UMaterialExpressionMaterialFunctionCall>(Expression))
 		{
@@ -1328,29 +1328,41 @@ bool UMaterial::SetVectorParameterValueEditorOnly(FName ParameterName, FLinearCo
 							{
 								if (FunctionExpressionParameter->SetParameterValue(ParameterName, InValue))
 								{
-									return true;
+									bSetNewValue = true;
+									break;
 								}
 							}
 						}
 					}
+					if (bSetNewValue)
+					{
+						break;
+					}
 				}
 			}
 		}
 	}
-	return false;
+	if (bSetNewValue)
+	{
+		UProperty* ParamProperty = FindField<UProperty>(UMaterialExpressionVectorParameter::StaticClass(), GET_MEMBER_NAME_STRING_CHECKED(UMaterialExpressionVectorParameter, DefaultValue));
+		FPropertyChangedEvent PropertyChangedEvent(ParamProperty);
+		Parameter->PostEditChangeProperty(PropertyChangedEvent);
+	}
+	return bSetNewValue;
 }
 
 bool UMaterial::SetScalarParameterValueEditorOnly(FName ParameterName, float InValue)
 {
+	bool bSetNewValue = false;
+	UMaterialExpressionScalarParameter* Parameter = nullptr;
 	for (UMaterialExpression* Expression : Expressions)
 	{
-		if (UMaterialExpressionScalarParameter* Parameter = Cast<UMaterialExpressionScalarParameter>(Expression))
+		Parameter = Cast<UMaterialExpressionScalarParameter>(Expression);
+		if (Parameter && Parameter->SetParameterValue(ParameterName, InValue))
 		{
-			if (Parameter->SetParameterValue(ParameterName, InValue))
-			{
-				return true;
-				// Warning: in the case of duplicate parameters with different default values, this will find the first in the expression array, not necessarily the one that's used for rendering
-			}
+			bSetNewValue = true;
+			break;
+			// Warning: in the case of duplicate parameters with different default values, this will find the first in the expression array, not necessarily the one that's used for rendering
 		}
 		else if (UMaterialExpressionMaterialFunctionCall* FunctionCall = Cast<UMaterialExpressionMaterialFunctionCall>(Expression))
 		{
@@ -1367,20 +1379,33 @@ bool UMaterial::SetScalarParameterValueEditorOnly(FName ParameterName, float InV
 					{
 						for (UMaterialExpression* FunctionExpression : *ExpressionPtr)
 						{
-							if (UMaterialExpressionScalarParameter* FunctionExpressionParameter = Cast<UMaterialExpressionScalarParameter>(FunctionExpression))
+							Parameter = Cast<UMaterialExpressionScalarParameter>(FunctionExpression);
+							if (!bSetNewValue && Parameter)
 							{
-								if (FunctionExpressionParameter->SetParameterValue(ParameterName, InValue))
+								if (Parameter->SetParameterValue(ParameterName, InValue))
 								{
-									return true;
+									bSetNewValue = true;
+									break;
 								}
 							}
 						}
 					}
+					if (bSetNewValue)
+					{
+						break;
+					}
 				}
 			}
 		}
 	}
-	return false;
+	if (bSetNewValue)
+	{
+		UProperty* ParamProperty = FindField<UProperty>(UMaterialExpressionScalarParameter::StaticClass(), GET_MEMBER_NAME_STRING_CHECKED(UMaterialExpressionScalarParameter, DefaultValue));
+		FPropertyChangedEvent PropertyChangedEvent(ParamProperty);
+		Parameter->PostEditChangeProperty(PropertyChangedEvent);
+	}
+
+	return bSetNewValue;
 };
 
 bool UMaterial::SetTextureParameterValueEditorOnly(FName ParameterName, class UTexture* InValue)
