Fix compile error

diff --git a/Engine/Source/Editor/MaterialEditor/Public/MaterialEditorModule.h b/Engine/Source/Editor/MaterialEditor/Public/MaterialEditorModule.h
--- a/Engine/Source/Editor/MaterialEditor/Public/MaterialEditorModule.h
+++ b/Engine/Source/Editor/MaterialEditor/Public/MaterialEditorModule.h
@@ -4,6 +4,7 @@
 
 #include "CoreMinimal.h"
 #include "Modules/ModuleInterface.h"
+#include "Modules/ModuleManager.h"
 #include "Framework/Commands/UICommandList.h"
 #include "Framework/MultiBox/MultiBoxExtender.h"
 #include "Toolkits/AssetEditorToolkit.h"
