Removing compile on load to avoid dirtying, will investigate another avenue for fixing vector/scalar params updating in world after load.
#rnx

diff --git a/Engine/Source/Editor/MaterialEditor/Private/MaterialEditor.cpp b/Engine/Source/Editor/MaterialEditor/Private/MaterialEditor.cpp
--- a/Engine/Source/Editor/MaterialEditor/Private/MaterialEditor.cpp
+++ b/Engine/Source/Editor/MaterialEditor/Private/MaterialEditor.cpp
@@ -693,8 +693,6 @@ void FMaterialEditor::InitMaterialEditor( const EToolkitMode::Type Mode, const T
 	Material->MaterialGraph->RebuildGraph();
 	RecenterEditor();
 
-	UMaterialEditingLibrary::RecompileMaterial(OriginalMaterial);
-
 	//Make sure the preview material is initialized.
 	UpdatePreviewMaterial(true);
 	RegenerateCodeView(true);
