Fix for post process materials not previewing correctly. Stopping nodes blinking on parameter refresh

diff --git a/Engine/Source/Editor/MaterialEditor/Private/MaterialEditor.cpp b/Engine/Source/Editor/MaterialEditor/Private/MaterialEditor.cpp
--- a/Engine/Source/Editor/MaterialEditor/Private/MaterialEditor.cpp
+++ b/Engine/Source/Editor/MaterialEditor/Private/MaterialEditor.cpp
@@ -1003,7 +1003,7 @@ void FMaterialEditor::OnFinishedChangingProperties(const FPropertyChangedEvent&
 void FMaterialEditor::OnFinishedChangingParametersFromOverview(const FPropertyChangedEvent& PropertyChangedEvent)
 {
 	bool bRefreshNodePreviews = false;
-	if (PropertyChangedEvent.Property != nullptr)
+	if (PropertyChangedEvent.Property != nullptr && PropertyChangedEvent.ChangeType != EPropertyChangeType::Interactive)
 	{
 		RefreshExpressionPreviews(true);
 		RefreshPreviewViewport();
@@ -3393,6 +3393,10 @@ void FMaterialEditor::SetPreviewExpression(UMaterialExpression* NewPreviewExpres
 			{
 				ExpressionPreviewMaterial->MaterialDomain = MD_UI;
 			}
+			else if (Material->IsPostProcessMaterial())
+			{
+				ExpressionPreviewMaterial->MaterialDomain = MD_PostProcess;
+			}
 		}
 
 		if (FunctionOutput)
diff --git a/Engine/Source/Editor/UnrealEd/Private/PreviewMaterial.cpp b/Engine/Source/Editor/UnrealEd/Private/PreviewMaterial.cpp
--- a/Engine/Source/Editor/UnrealEd/Private/PreviewMaterial.cpp
+++ b/Engine/Source/Editor/UnrealEd/Private/PreviewMaterial.cpp
@@ -251,7 +251,7 @@ FMaterialResource* UPreviewMaterial::AllocateResource()
 
 void UMaterialEditorPreviewParameters::PostEditChangeProperty(FPropertyChangedEvent& PropertyChangedEvent)
 {
-	if (PreviewMaterial)
+	if (PreviewMaterial && PropertyChangedEvent.ChangeType != EPropertyChangeType::Interactive)
 	{
 		UProperty* PropertyThatChanged = PropertyChangedEvent.Property;
 		if (OriginalFunction == nullptr)
