Moving sort priority down in the material parameters panel

diff --git a/Engine/Source/Editor/MaterialEditor/Private/MaterialEditorDetailCustomization.cpp b/Engine/Source/Editor/MaterialEditor/Private/MaterialEditorDetailCustomization.cpp
--- a/Engine/Source/Editor/MaterialEditor/Private/MaterialEditorDetailCustomization.cpp
+++ b/Engine/Source/Editor/MaterialEditor/Private/MaterialEditorDetailCustomization.cpp
@@ -492,7 +492,10 @@ void FMaterialExpressionParameterDetails::CustomizeDetails( IDetailLayoutBuilder
 	check(ScalarParameterObjects.Num() == DefaultValueHandles.Num());
 	
 	Category.AddProperty("ParameterName");
-	
+
+	const FName MaterialExpressionCategory = TEXT("MaterialExpression");
+	IDetailCategoryBuilder& ExpressionCategory = DetailLayout.EditCategory(MaterialExpressionCategory);
+
 	// Get a handle to the property we are about to edit
 	GroupPropertyHandle = DetailLayout.GetProperty( "Group" );
 
@@ -506,39 +509,39 @@ void FMaterialExpressionParameterDetails::CustomizeDetails( IDetailLayoutBuilder
 
 		PopulateGroups();	
 
-		Category.AddCustomRow(GroupPropertyHandle->GetPropertyDisplayName())
+		ExpressionCategory.AddCustomRow(GroupPropertyHandle->GetPropertyDisplayName())
 			.NameContent()
 			[
 				SNew(STextBlock)
 				.Text(GroupPropertyHandle->GetPropertyDisplayName())
-			.Font(IDetailLayoutBuilder::GetDetailFont())
+				.Font(IDetailLayoutBuilder::GetDetailFont())
 			]
-		.ValueContent()
+			.ValueContent()
 			[
 				SAssignNew(NewComboButton, SComboButton)
 				.ContentPadding(FMargin(2.0f, 2.0f))
-			.ButtonContent()
-			[
-				SAssignNew(NewEditBox, SEditableText)
-				.Text(this, &FMaterialExpressionParameterDetails::OnGetText)
-			.OnTextCommitted(this, &FMaterialExpressionParameterDetails::OnTextCommitted)
-			]
-		.MenuContent()
-			[
-				SNew(SVerticalBox)
-				+ SVerticalBox::Slot()
-			.AutoHeight()
-			.MaxHeight(400.0f)
-			[
-				SAssignNew(NewListView, SListView<TSharedPtr<FString>>)
-				.ListItemsSource(&GroupsSource)
-			.OnGenerateRow(this, &FMaterialExpressionParameterDetails::MakeDetailsGroupViewWidget)
-			.OnSelectionChanged(this, &FMaterialExpressionParameterDetails::OnSelectionChanged)
-			]
-			]
+				.ButtonContent()
+				[
+					SAssignNew(NewEditBox, SEditableText)
+					.Text(this, &FMaterialExpressionParameterDetails::OnGetText)
+					.OnTextCommitted(this, &FMaterialExpressionParameterDetails::OnTextCommitted)
+				]
+				.MenuContent()
+				[
+					SNew(SVerticalBox)
+					+ SVerticalBox::Slot()
+					.AutoHeight()
+					.MaxHeight(400.0f)
+					[
+						SAssignNew(NewListView, SListView<TSharedPtr<FString>>)
+						.ListItemsSource(&GroupsSource)
+						.OnGenerateRow(this, &FMaterialExpressionParameterDetails::MakeDetailsGroupViewWidget)
+						.OnSelectionChanged(this, &FMaterialExpressionParameterDetails::OnSelectionChanged)
+					]
+				]
 			];
 
-		Category.AddProperty("SortPriority");
+		ExpressionCategory.AddProperty("SortPriority");
 	}
 
 	GroupComboButton = NewComboButton;
diff --git a/Engine/Source/Runtime/Engine/Classes/Materials/MaterialExpression.h b/Engine/Source/Runtime/Engine/Classes/Materials/MaterialExpression.h
--- a/Engine/Source/Runtime/Engine/Classes/Materials/MaterialExpression.h
+++ b/Engine/Source/Runtime/Engine/Classes/Materials/MaterialExpression.h
@@ -136,7 +136,7 @@ class ENGINE_API UMaterialExpression : public UObject
 
 #if WITH_EDITORONLY_DATA
 	/** A description that level designers can add (shows in the material editor UI). */
-	UPROPERTY(EditAnywhere, Category=MaterialExpression, meta=(MultiLine=true))
+	UPROPERTY(EditAnywhere, Category=MaterialExpression, meta=(MultiLine=true, DisplayAfter = "SortPriority"))
 	FString Desc;
 
 	/** Set to true by RecursiveUpdateRealtimePreview() if the expression's preview needs to be updated in realtime in the material editor. */
diff --git a/Engine/Source/Runtime/Engine/Classes/Materials/MaterialExpressionFunctionInput.h b/Engine/Source/Runtime/Engine/Classes/Materials/MaterialExpressionFunctionInput.h
--- a/Engine/Source/Runtime/Engine/Classes/Materials/MaterialExpressionFunctionInput.h
+++ b/Engine/Source/Runtime/Engine/Classes/Materials/MaterialExpressionFunctionInput.h
@@ -67,7 +67,7 @@ class UMaterialExpressionFunctionInput : public UMaterialExpression
 	uint32 bUsePreviewValueAsDefault:1;
 
 	/** Controls where the input is displayed relative to the other inputs. */
-	UPROPERTY(EditAnywhere, BlueprintReadWrite, Category=MaterialExpressionFunctionInput)
+	UPROPERTY(EditAnywhere, BlueprintReadWrite, Category=MaterialExpression)
 	int32 SortPriority;
 
 	/** 
diff --git a/Engine/Source/Runtime/Engine/Classes/Materials/MaterialExpressionFunctionOutput.h b/Engine/Source/Runtime/Engine/Classes/Materials/MaterialExpressionFunctionOutput.h
--- a/Engine/Source/Runtime/Engine/Classes/Materials/MaterialExpressionFunctionOutput.h
+++ b/Engine/Source/Runtime/Engine/Classes/Materials/MaterialExpressionFunctionOutput.h
@@ -26,7 +26,7 @@ class UMaterialExpressionFunctionOutput : public UMaterialExpression
 	FString Description;
 
 	/** Controls where the output is displayed relative to the other outputs. */
-	UPROPERTY(EditAnywhere, Category=MaterialExpressionFunctionOutput)
+	UPROPERTY(EditAnywhere, Category=MaterialExpression)
 	int32 SortPriority;
 
 	/** Stores the expression in the material function connected to this output. */
