Reset to default now handles base properties and lightmass properties correctly

diff --git a/Engine/Source/Editor/MaterialEditor/Private/MaterialEditorInstanceDetailCustomization.cpp b/Engine/Source/Editor/MaterialEditor/Private/MaterialEditorInstanceDetailCustomization.cpp
--- a/Engine/Source/Editor/MaterialEditor/Private/MaterialEditorInstanceDetailCustomization.cpp
+++ b/Engine/Source/Editor/MaterialEditor/Private/MaterialEditorInstanceDetailCustomization.cpp
@@ -750,6 +750,13 @@ void FMaterialInstanceParameterDetails::CreateLightmassOverrideWidgets(IDetailLa
 	TSharedPtr<IPropertyHandle> ExportResolutionScaleProperty = LightmassSettings->GetChildHandle("ExportResolutionScale");
 
 
+	FIsResetToDefaultVisible IsCastShadowAsMaskedPropertyResetVisible = FIsResetToDefaultVisible::CreateLambda([this](TSharedPtr<IPropertyHandle> InHandle) {
+		return MaterialEditorInstance->LightmassSettings.CastShadowAsMasked.ParameterValue != MaterialEditorInstance->Parent->GetCastShadowAsMasked();
+	});
+	FResetToDefaultHandler ResetCastShadowAsMaskedPropertyHandler = FResetToDefaultHandler::CreateLambda([this](TSharedPtr<IPropertyHandle> InHandle) {
+		MaterialEditorInstance->LightmassSettings.CastShadowAsMasked.ParameterValue = MaterialEditorInstance->Parent->GetCastShadowAsMasked();
+	});
+	FResetToDefaultOverride ResetCastShadowAsMaskedPropertyOverride = FResetToDefaultOverride::Create(IsCastShadowAsMaskedPropertyResetVisible, ResetCastShadowAsMaskedPropertyHandler);
 	IDetailPropertyRow& CastShadowAsMaskedPropertyRow = LightmassSettingsGroup.AddPropertyRow(CastShadowAsMaskedProperty->GetChildHandle(0).ToSharedRef());
 	CastShadowAsMaskedPropertyRow
 		.DisplayName(CastShadowAsMaskedProperty->GetPropertyDisplayName())
@@ -759,8 +766,16 @@ void FMaterialInstanceParameterDetails::CreateLightmassOverrideWidgets(IDetailLa
 			MaterialEditorInstance->PostEditChange();
 			FEditorSupportDelegates::RedrawAllViewports.Broadcast();
 		}))
-		.Visibility(TAttribute<EVisibility>::Create(TAttribute<EVisibility>::FGetter::CreateSP(this, &FMaterialInstanceParameterDetails::IsOverriddenAndVisible, IsOverrideCastShadowAsMaskedEnabled)));
-
+		.Visibility(TAttribute<EVisibility>::Create(TAttribute<EVisibility>::FGetter::CreateSP(this, &FMaterialInstanceParameterDetails::IsOverriddenAndVisible, IsOverrideCastShadowAsMaskedEnabled)))
+		.OverrideResetToDefault(ResetCastShadowAsMaskedPropertyOverride);
+
+	FIsResetToDefaultVisible IsEmissiveBoostPropertyResetVisible = FIsResetToDefaultVisible::CreateLambda([this](TSharedPtr<IPropertyHandle> InHandle) {
+		return MaterialEditorInstance->LightmassSettings.EmissiveBoost.ParameterValue != MaterialEditorInstance->Parent->GetEmissiveBoost();
+	});
+	FResetToDefaultHandler ResetEmissiveBoostPropertyHandler = FResetToDefaultHandler::CreateLambda([this](TSharedPtr<IPropertyHandle> InHandle) {
+		MaterialEditorInstance->LightmassSettings.EmissiveBoost.ParameterValue = MaterialEditorInstance->Parent->GetEmissiveBoost();
+	});
+	FResetToDefaultOverride ResetEmissiveBoostPropertyOverride = FResetToDefaultOverride::Create(IsEmissiveBoostPropertyResetVisible, ResetEmissiveBoostPropertyHandler);
 	IDetailPropertyRow& EmissiveBoostPropertyRow = LightmassSettingsGroup.AddPropertyRow(EmissiveBoostProperty->GetChildHandle(0).ToSharedRef());
 	EmissiveBoostPropertyRow
 		.DisplayName(EmissiveBoostProperty->GetPropertyDisplayName())
@@ -770,8 +785,16 @@ void FMaterialInstanceParameterDetails::CreateLightmassOverrideWidgets(IDetailLa
 			MaterialEditorInstance->PostEditChange();
 			FEditorSupportDelegates::RedrawAllViewports.Broadcast();
 		}))
-		.Visibility(TAttribute<EVisibility>::Create(TAttribute<EVisibility>::FGetter::CreateSP(this, &FMaterialInstanceParameterDetails::IsOverriddenAndVisible, IsOverrideEmissiveBoostEnabled)));
+		.Visibility(TAttribute<EVisibility>::Create(TAttribute<EVisibility>::FGetter::CreateSP(this, &FMaterialInstanceParameterDetails::IsOverriddenAndVisible, IsOverrideEmissiveBoostEnabled)))
+		.OverrideResetToDefault(ResetEmissiveBoostPropertyOverride);
 	
+	FIsResetToDefaultVisible IsDiffuseBoostPropertyResetVisible = FIsResetToDefaultVisible::CreateLambda([this](TSharedPtr<IPropertyHandle> InHandle) {
+		return MaterialEditorInstance->LightmassSettings.DiffuseBoost.ParameterValue != MaterialEditorInstance->Parent->GetDiffuseBoost();
+	});
+	FResetToDefaultHandler ResetDiffuseBoostPropertyHandler = FResetToDefaultHandler::CreateLambda([this](TSharedPtr<IPropertyHandle> InHandle) {
+		MaterialEditorInstance->LightmassSettings.DiffuseBoost.ParameterValue = MaterialEditorInstance->Parent->GetDiffuseBoost();
+	});
+	FResetToDefaultOverride ResetDiffuseBoostPropertyOverride = FResetToDefaultOverride::Create(IsDiffuseBoostPropertyResetVisible, ResetDiffuseBoostPropertyHandler);
 	IDetailPropertyRow& DiffuseBoostPropertyRow = LightmassSettingsGroup.AddPropertyRow(DiffuseBoostProperty->GetChildHandle(0).ToSharedRef());
 	DiffuseBoostPropertyRow
 		.DisplayName(DiffuseBoostProperty->GetPropertyDisplayName())
@@ -781,8 +804,16 @@ void FMaterialInstanceParameterDetails::CreateLightmassOverrideWidgets(IDetailLa
 			MaterialEditorInstance->PostEditChange();
 			FEditorSupportDelegates::RedrawAllViewports.Broadcast();
 		}))
-		.Visibility(TAttribute<EVisibility>::Create(TAttribute<EVisibility>::FGetter::CreateSP(this, &FMaterialInstanceParameterDetails::IsOverriddenAndVisible, IsOverrideDiffuseBoostEnabled)));
-
+		.Visibility(TAttribute<EVisibility>::Create(TAttribute<EVisibility>::FGetter::CreateSP(this, &FMaterialInstanceParameterDetails::IsOverriddenAndVisible, IsOverrideDiffuseBoostEnabled)))
+		.OverrideResetToDefault(ResetDiffuseBoostPropertyOverride);
+
+	FIsResetToDefaultVisible IsExportResolutionScalePropertyResetVisible = FIsResetToDefaultVisible::CreateLambda([this](TSharedPtr<IPropertyHandle> InHandle) {
+		return MaterialEditorInstance->LightmassSettings.ExportResolutionScale.ParameterValue != MaterialEditorInstance->Parent->GetDiffuseBoost();
+	});
+	FResetToDefaultHandler ResetExportResolutionScalePropertyHandler = FResetToDefaultHandler::CreateLambda([this](TSharedPtr<IPropertyHandle> InHandle) {
+		MaterialEditorInstance->LightmassSettings.ExportResolutionScale.ParameterValue = MaterialEditorInstance->Parent->GetDiffuseBoost();
+	});
+	FResetToDefaultOverride ResetExportResolutionScalePropertyOverride = FResetToDefaultOverride::Create(IsExportResolutionScalePropertyResetVisible, ResetExportResolutionScalePropertyHandler);
 	IDetailPropertyRow& ExportResolutionScalePropertyRow = LightmassSettingsGroup.AddPropertyRow(ExportResolutionScaleProperty->GetChildHandle(0).ToSharedRef());
 	ExportResolutionScalePropertyRow
 		.DisplayName(ExportResolutionScaleProperty->GetPropertyDisplayName())
@@ -792,7 +823,8 @@ void FMaterialInstanceParameterDetails::CreateLightmassOverrideWidgets(IDetailLa
 			MaterialEditorInstance->PostEditChange();
 			FEditorSupportDelegates::RedrawAllViewports.Broadcast();
 		}))
-		.Visibility(TAttribute<EVisibility>::Create(TAttribute<EVisibility>::FGetter::CreateSP(this, &FMaterialInstanceParameterDetails::IsOverriddenAndVisible, IsOverrideExportResolutionScaleEnabled)));
+		.Visibility(TAttribute<EVisibility>::Create(TAttribute<EVisibility>::FGetter::CreateSP(this, &FMaterialInstanceParameterDetails::IsOverriddenAndVisible, IsOverrideExportResolutionScaleEnabled)))
+		.OverrideResetToDefault(ResetExportResolutionScalePropertyOverride);
 }
 
 void FMaterialInstanceParameterDetails::CreateBasePropertyOverrideWidgets(IDetailLayoutBuilder& DetailLayout)
@@ -815,40 +847,80 @@ void FMaterialInstanceParameterDetails::CreateBasePropertyOverrideWidgets(IDetai
 	TSharedPtr<IPropertyHandle> TwoSidedProperty = BasePropertyOverridePropery->GetChildHandle("TwoSided");
 	TSharedPtr<IPropertyHandle> DitheredLODTransitionProperty = BasePropertyOverridePropery->GetChildHandle("DitheredLODTransition");
 
+	FIsResetToDefaultVisible IsOpacityClipMaskValuePropertyResetVisible = FIsResetToDefaultVisible::CreateLambda([this](TSharedPtr<IPropertyHandle> InHandle) {
+		return MaterialEditorInstance->BasePropertyOverrides.OpacityMaskClipValue != MaterialEditorInstance->Parent->GetOpacityMaskClipValue();
+	});
+	FResetToDefaultHandler ResetOpacityClipMaskValuePropertyHandler = FResetToDefaultHandler::CreateLambda([this](TSharedPtr<IPropertyHandle> InHandle) {
+		MaterialEditorInstance->BasePropertyOverrides.OpacityMaskClipValue = MaterialEditorInstance->Parent->GetOpacityMaskClipValue();
+	});
+	FResetToDefaultOverride ResetOpacityClipMaskValuePropertyOverride = FResetToDefaultOverride::Create(IsOpacityClipMaskValuePropertyResetVisible, ResetOpacityClipMaskValuePropertyHandler);
 	IDetailPropertyRow& OpacityClipMaskValuePropertyRow = BasePropertyOverrideGroup.AddPropertyRow(OpacityClipMaskValueProperty.ToSharedRef());
 	OpacityClipMaskValuePropertyRow
 		.DisplayName(OpacityClipMaskValueProperty->GetPropertyDisplayName())
 		.ToolTip(OpacityClipMaskValueProperty->GetToolTipText())
 		.EditCondition(IsOverrideOpacityClipMaskValueEnabled, FOnBooleanValueChanged::CreateSP(this, &FMaterialInstanceParameterDetails::OnOverrideOpacityClipMaskValueChanged))
-		.Visibility(TAttribute<EVisibility>::Create(TAttribute<EVisibility>::FGetter::CreateSP(this, &FMaterialInstanceParameterDetails::IsOverriddenAndVisible, IsOverrideOpacityClipMaskValueEnabled)));
-
+		.Visibility(TAttribute<EVisibility>::Create(TAttribute<EVisibility>::FGetter::CreateSP(this, &FMaterialInstanceParameterDetails::IsOverriddenAndVisible, IsOverrideOpacityClipMaskValueEnabled)))
+		.OverrideResetToDefault(ResetOpacityClipMaskValuePropertyOverride);
+
+	FIsResetToDefaultVisible IsBlendModePropertyResetVisible = FIsResetToDefaultVisible::CreateLambda([this](TSharedPtr<IPropertyHandle> InHandle) {
+		return MaterialEditorInstance->BasePropertyOverrides.BlendMode != MaterialEditorInstance->Parent->GetBlendMode();
+	});
+	FResetToDefaultHandler ResetBlendModePropertyHandler = FResetToDefaultHandler::CreateLambda([this](TSharedPtr<IPropertyHandle> InHandle) {
+		MaterialEditorInstance->BasePropertyOverrides.BlendMode = MaterialEditorInstance->Parent->GetBlendMode();
+	}); 
+	FResetToDefaultOverride ResetBlendModePropertyOverride = FResetToDefaultOverride::Create(IsBlendModePropertyResetVisible, ResetBlendModePropertyHandler);
 	IDetailPropertyRow& BlendModePropertyRow = BasePropertyOverrideGroup.AddPropertyRow(BlendModeProperty.ToSharedRef());
 	BlendModePropertyRow
 		.DisplayName(BlendModeProperty->GetPropertyDisplayName())
 		.ToolTip(BlendModeProperty->GetToolTipText())
 		.EditCondition(IsOverrideBlendModeEnabled, FOnBooleanValueChanged::CreateSP(this, &FMaterialInstanceParameterDetails::OnOverrideBlendModeChanged))
-		.Visibility(TAttribute<EVisibility>::Create(TAttribute<EVisibility>::FGetter::CreateSP(this, &FMaterialInstanceParameterDetails::IsOverriddenAndVisible, IsOverrideBlendModeEnabled)));
-
+		.Visibility(TAttribute<EVisibility>::Create(TAttribute<EVisibility>::FGetter::CreateSP(this, &FMaterialInstanceParameterDetails::IsOverriddenAndVisible, IsOverrideBlendModeEnabled)))
+		.OverrideResetToDefault(ResetBlendModePropertyOverride);
+
+	FIsResetToDefaultVisible IsShadingModelPropertyResetVisible = FIsResetToDefaultVisible::CreateLambda([this](TSharedPtr<IPropertyHandle> InHandle) {
+		return MaterialEditorInstance->BasePropertyOverrides.ShadingModel != MaterialEditorInstance->Parent->GetShadingModel();
+	});
+	FResetToDefaultHandler ResetShadingModelPropertyHandler = FResetToDefaultHandler::CreateLambda([this](TSharedPtr<IPropertyHandle> InHandle) {
+		MaterialEditorInstance->BasePropertyOverrides.ShadingModel = MaterialEditorInstance->Parent->GetShadingModel();
+	});
+	FResetToDefaultOverride ResetShadingModelPropertyOverride = FResetToDefaultOverride::Create(IsShadingModelPropertyResetVisible, ResetShadingModelPropertyHandler);
 	IDetailPropertyRow& ShadingModelPropertyRow = BasePropertyOverrideGroup.AddPropertyRow(ShadingModelProperty.ToSharedRef());
 	ShadingModelPropertyRow
 		.DisplayName(ShadingModelProperty->GetPropertyDisplayName())
 		.ToolTip(ShadingModelProperty->GetToolTipText())
 		.EditCondition(IsOverrideShadingModelEnabled, FOnBooleanValueChanged::CreateSP(this, &FMaterialInstanceParameterDetails::OnOverrideShadingModelChanged))
-		.Visibility(TAttribute<EVisibility>::Create(TAttribute<EVisibility>::FGetter::CreateSP(this, &FMaterialInstanceParameterDetails::IsOverriddenAndVisible, IsOverrideShadingModelEnabled)));
-
+		.Visibility(TAttribute<EVisibility>::Create(TAttribute<EVisibility>::FGetter::CreateSP(this, &FMaterialInstanceParameterDetails::IsOverriddenAndVisible, IsOverrideShadingModelEnabled)))
+		.OverrideResetToDefault(ResetShadingModelPropertyOverride);
+
+	FIsResetToDefaultVisible IsTwoSidedPropertyResetVisible = FIsResetToDefaultVisible::CreateLambda([this](TSharedPtr<IPropertyHandle> InHandle) {
+		return MaterialEditorInstance->BasePropertyOverrides.TwoSided != MaterialEditorInstance->Parent->IsTwoSided();
+	});
+	FResetToDefaultHandler ResetTwoSidedValuePropertyHandler = FResetToDefaultHandler::CreateLambda([this](TSharedPtr<IPropertyHandle> InHandle) {
+		MaterialEditorInstance->BasePropertyOverrides.TwoSided = MaterialEditorInstance->Parent->IsTwoSided();
+	});
+	FResetToDefaultOverride ResetTwoSidedPropertyOverride = FResetToDefaultOverride::Create(IsTwoSidedPropertyResetVisible, ResetTwoSidedValuePropertyHandler);
 	IDetailPropertyRow& TwoSidedPropertyRow = BasePropertyOverrideGroup.AddPropertyRow(TwoSidedProperty.ToSharedRef());
 	TwoSidedPropertyRow
 		.DisplayName(TwoSidedProperty->GetPropertyDisplayName())
 		.ToolTip(TwoSidedProperty->GetToolTipText())
 		.EditCondition(IsOverrideTwoSidedEnabled, FOnBooleanValueChanged::CreateSP(this, &FMaterialInstanceParameterDetails::OnOverrideTwoSidedChanged))
-		.Visibility(TAttribute<EVisibility>::Create(TAttribute<EVisibility>::FGetter::CreateSP(this, &FMaterialInstanceParameterDetails::IsOverriddenAndVisible, IsOverrideTwoSidedEnabled)));
-
+		.Visibility(TAttribute<EVisibility>::Create(TAttribute<EVisibility>::FGetter::CreateSP(this, &FMaterialInstanceParameterDetails::IsOverriddenAndVisible, IsOverrideTwoSidedEnabled)))
+		.OverrideResetToDefault(ResetTwoSidedPropertyOverride);
+
+	FIsResetToDefaultVisible IsDitheredLODTransitionPropertyResetVisible = FIsResetToDefaultVisible::CreateLambda([this](TSharedPtr<IPropertyHandle> InHandle) {
+		return MaterialEditorInstance->BasePropertyOverrides.DitheredLODTransition != MaterialEditorInstance->Parent->IsDitheredLODTransition();
+	});
+	FResetToDefaultHandler ResetDitheredLODTransitionPropertyHandler = FResetToDefaultHandler::CreateLambda([this](TSharedPtr<IPropertyHandle> InHandle) {
+		MaterialEditorInstance->BasePropertyOverrides.DitheredLODTransition = MaterialEditorInstance->Parent->IsDitheredLODTransition();
+	});
+	FResetToDefaultOverride ResetDitheredLODTransitionPropertyOverride = FResetToDefaultOverride::Create(IsDitheredLODTransitionPropertyResetVisible, ResetDitheredLODTransitionPropertyHandler);
 	IDetailPropertyRow& DitheredLODTransitionPropertyRow = BasePropertyOverrideGroup.AddPropertyRow(DitheredLODTransitionProperty.ToSharedRef());
 	DitheredLODTransitionPropertyRow
 		.DisplayName(DitheredLODTransitionProperty->GetPropertyDisplayName())
 		.ToolTip(DitheredLODTransitionProperty->GetToolTipText())
 		.EditCondition(IsOverrideDitheredLODTransitionEnabled, FOnBooleanValueChanged::CreateSP(this, &FMaterialInstanceParameterDetails::OnOverrideDitheredLODTransitionChanged))
-		.Visibility(TAttribute<EVisibility>::Create(TAttribute<EVisibility>::FGetter::CreateSP(this, &FMaterialInstanceParameterDetails::IsOverriddenAndVisible, IsOverrideDitheredLODTransitionEnabled)));
+		.Visibility(TAttribute<EVisibility>::Create(TAttribute<EVisibility>::FGetter::CreateSP(this, &FMaterialInstanceParameterDetails::IsOverriddenAndVisible, IsOverrideDitheredLODTransitionEnabled)))
+		.OverrideResetToDefault(ResetDitheredLODTransitionPropertyOverride);
 }
 
 EVisibility FMaterialInstanceParameterDetails::IsOverriddenAndVisible(TAttribute<bool> IsOverridden) const
diff --git a/Engine/Source/Editor/MaterialEditor/Private/MaterialInstanceEditor.cpp b/Engine/Source/Editor/MaterialEditor/Private/MaterialInstanceEditor.cpp
--- a/Engine/Source/Editor/MaterialEditor/Private/MaterialInstanceEditor.cpp
+++ b/Engine/Source/Editor/MaterialEditor/Private/MaterialInstanceEditor.cpp
@@ -1201,6 +1201,7 @@ void FMaterialInstanceEditor::RebuildMaterialInstanceEditor()
 	if( MaterialEditorInstance )
 	{
 		ReInitMaterialFunctionProxies();
+		MaterialEditorInstance->CopyBasePropertiesFromParent();
 		MaterialEditorInstance->RegenerateArrays();
 		RebuildInheritanceList(); // Required b/c recompiled parent materials result in invalid weak object pointers
 		UpdatePropertyWindow();
diff --git a/Engine/Source/Editor/UnrealEd/Classes/MaterialEditor/MaterialEditorInstanceConstant.h b/Engine/Source/Editor/UnrealEd/Classes/MaterialEditor/MaterialEditorInstanceConstant.h
--- a/Engine/Source/Editor/UnrealEd/Classes/MaterialEditor/MaterialEditorInstanceConstant.h
+++ b/Engine/Source/Editor/UnrealEd/Classes/MaterialEditor/MaterialEditorInstanceConstant.h
@@ -299,6 +299,8 @@ class UNREALED_API UMaterialEditorInstanceConstant : public UObject
 	 */
 	void SetSourceInstance(UMaterialInstanceConstant* MaterialInterface);
 
+	void CopyBasePropertiesFromParent();
+
 	void SetSourceFunction(UMaterialFunctionInstance* MaterialFunction);
 
 	/** 
@@ -336,3 +338,5 @@ class UNREALED_API UMaterialEditorInstanceConstant : public UObject
 	bool bShowOnlyOverrides;
 };
 
+
+
diff --git a/Engine/Source/Editor/UnrealEd/Private/PreviewMaterial.cpp b/Engine/Source/Editor/UnrealEd/Private/PreviewMaterial.cpp
--- a/Engine/Source/Editor/UnrealEd/Private/PreviewMaterial.cpp
+++ b/Engine/Source/Editor/UnrealEd/Private/PreviewMaterial.cpp
@@ -1482,6 +1482,36 @@ void UMaterialEditorInstanceConstant::SetSourceInstance(UMaterialInstanceConstan
 	Parent = SourceInstance->Parent;
 	PhysMaterial = SourceInstance->PhysMaterial;
 
+	CopyBasePropertiesFromParent();
+
+	RegenerateArrays();
+
+	//propagate changes to the base material so the instance will be updated if it has a static permutation resource
+	FStaticParameterSet NewStaticParameters;
+	BuildStaticParametersForSourceInstance(NewStaticParameters);
+	SourceInstance->UpdateStaticPermutation(NewStaticParameters);
+}
+
+void UMaterialEditorInstanceConstant::SetSourceFunction(UMaterialFunctionInstance* MaterialFunction)
+{
+	SourceFunction = MaterialFunction;
+	bIsFunctionPreviewMaterial = !!(SourceFunction);
+}
+
+void UMaterialEditorInstanceConstant::UpdateSourceInstanceParent()
+{
+	// If the parent was changed to the source instance, set it to NULL
+	if( Parent == SourceInstance )
+	{
+		Parent = NULL;
+	}
+
+	SourceInstance->SetParentEditorOnly( Parent );
+}
+
+
+void UMaterialEditorInstanceConstant::CopyBasePropertiesFromParent()
+{
 	BasePropertyOverrides = SourceInstance->BasePropertyOverrides;
 	// Copy the overrides (if not yet overridden), so they match their true values in the UI
 	if (!BasePropertyOverrides.bOverride_OpacityMaskClipValue)
@@ -1506,44 +1536,20 @@ void UMaterialEditorInstanceConstant::SetSourceInstance(UMaterialInstanceConstan
 	}
 
 	// Copy the Lightmass settings...
-	LightmassSettings.CastShadowAsMasked.bOverride = SourceInstance->GetOverrideCastShadowAsMasked();
+	// The lightmass functions (GetCastShadowAsMasked, etc.) check if the value is overridden and returns the current value if so, otherwise returns the parent value
+	// So we don't need to wrap these in the same "if not overriding" as above
 	LightmassSettings.CastShadowAsMasked.ParameterValue = SourceInstance->GetCastShadowAsMasked();
-	LightmassSettings.EmissiveBoost.bOverride = SourceInstance->GetOverrideEmissiveBoost();
 	LightmassSettings.EmissiveBoost.ParameterValue = SourceInstance->GetEmissiveBoost();
-	LightmassSettings.DiffuseBoost.bOverride = SourceInstance->GetOverrideDiffuseBoost();
 	LightmassSettings.DiffuseBoost.ParameterValue = SourceInstance->GetDiffuseBoost();
-	LightmassSettings.ExportResolutionScale.bOverride = SourceInstance->GetOverrideExportResolutionScale();
 	LightmassSettings.ExportResolutionScale.ParameterValue = SourceInstance->GetExportResolutionScale();
 
 	//Copy refraction settings
 	SourceInstance->GetRefractionSettings(RefractionDepthBias);
 
-	bOverrideSubsurfaceProfile = SourceInstance->bOverrideSubsurfaceProfile;
-	SubsurfaceProfile = SourceInstance->SubsurfaceProfile;
-
-	RegenerateArrays();
-
-	//propagate changes to the base material so the instance will be updated if it has a static permutation resource
-	FStaticParameterSet NewStaticParameters;
-	BuildStaticParametersForSourceInstance(NewStaticParameters);
-	SourceInstance->UpdateStaticPermutation(NewStaticParameters);
-}
-
-void UMaterialEditorInstanceConstant::SetSourceFunction(UMaterialFunctionInstance* MaterialFunction)
-{
-	SourceFunction = MaterialFunction;
-	bIsFunctionPreviewMaterial = !!(SourceFunction);
-}
-
-void UMaterialEditorInstanceConstant::UpdateSourceInstanceParent()
-{
-	// If the parent was changed to the source instance, set it to NULL
-	if( Parent == SourceInstance )
+	if (!bOverrideSubsurfaceProfile)
 	{
-		Parent = NULL;
+		SubsurfaceProfile = SourceInstance->SubsurfaceProfile;
 	}
-
-	SourceInstance->SetParentEditorOnly( Parent );
 }
 
 #if WITH_EDITOR
