Adding a custom reset to default for refraction depth bias

diff --git a/Engine/Source/Editor/MaterialEditor/Private/MaterialEditorInstanceDetailCustomization.cpp b/Engine/Source/Editor/MaterialEditor/Private/MaterialEditorInstanceDetailCustomization.cpp
--- a/Engine/Source/Editor/MaterialEditor/Private/MaterialEditorInstanceDetailCustomization.cpp
+++ b/Engine/Source/Editor/MaterialEditor/Private/MaterialEditorInstanceDetailCustomization.cpp
@@ -214,8 +214,20 @@ void FMaterialInstanceParameterDetails::CustomizeDetails(IDetailLayoutBuilder& D
 		DetailLayout.HideProperty("ParameterGroups");
 
 		{
+			FIsResetToDefaultVisible IsRefractionDepthBiasPropertyResetVisible = FIsResetToDefaultVisible::CreateLambda([this](TSharedPtr<IPropertyHandle> InHandle) {
+				float BiasValue;
+				float ParentBiasValue;
+				return MaterialEditorInstance->SourceInstance->GetRefractionSettings(BiasValue) 
+					&& MaterialEditorInstance->Parent->GetRefractionSettings(ParentBiasValue)
+					&& BiasValue != ParentBiasValue;
+			});
+			FResetToDefaultHandler ResetRefractionDepthBiasPropertyHandler = FResetToDefaultHandler::CreateLambda([this](TSharedPtr<IPropertyHandle> InHandle) {
+				MaterialEditorInstance->Parent->GetRefractionSettings(MaterialEditorInstance->RefractionDepthBias);
+			});
+			FResetToDefaultOverride ResetRefractionDepthBiasPropertyOverride = FResetToDefaultOverride::Create(IsRefractionDepthBiasPropertyResetVisible, ResetRefractionDepthBiasPropertyHandler);
 			IDetailPropertyRow& PropertyRow = DefaultCategory.AddProperty("RefractionDepthBias");
 			PropertyRow.Visibility(TAttribute<EVisibility>::Create(TAttribute<EVisibility>::FGetter::CreateSP(this, &FMaterialInstanceParameterDetails::ShouldShowMaterialRefractionSettings)));
+			PropertyRow.OverrideResetToDefault(ResetRefractionDepthBiasPropertyOverride);
 		}
 
 		{
