Fix for UE-68218. Crash altering material's scalar values. #rb Marcus.Wassmer

diff --git a/Engine/Source/Editor/MaterialEditor/Private/MaterialEditor.cpp b/Engine/Source/Editor/MaterialEditor/Private/MaterialEditor.cpp
--- a/Engine/Source/Editor/MaterialEditor/Private/MaterialEditor.cpp
+++ b/Engine/Source/Editor/MaterialEditor/Private/MaterialEditor.cpp
@@ -3105,6 +3105,9 @@ void FMaterialEditor::SetVectorParameterDefaultOnDependentMaterials(FName Parame
 		MaterialsToOverride.Add(OriginalMaterial);
 	}
 
+
+	FMaterialUpdateContext UpdateContext(FMaterialUpdateContext::EOptions::SyncWithRenderingThread);
+
 	const ERHIFeatureLevel::Type FeatureLevel = GEditor->GetEditorWorldContext().World()->FeatureLevel;
 
 	for (int32 MaterialIndex = 0; MaterialIndex < MaterialsToOverride.Num(); MaterialIndex++)
@@ -3112,6 +3115,7 @@ void FMaterialEditor::SetVectorParameterDefaultOnDependentMaterials(FName Parame
 		UMaterial* CurrentMaterial = MaterialsToOverride[MaterialIndex];
 
 		CurrentMaterial->OverrideVectorParameterDefault(ParameterName, Value, bOverride, FeatureLevel);
+		UpdateContext.AddMaterial(CurrentMaterial);
 	}
 
 	// Update MI's that reference any of the materials affected
@@ -3181,6 +3185,9 @@ void FMaterialEditor::SetScalarParameterDefaultOnDependentMaterials(FName Parame
 		MaterialsToOverride.Add(OriginalMaterial);
 	}
 
+
+	FMaterialUpdateContext UpdateContext(FMaterialUpdateContext::EOptions::SyncWithRenderingThread);
+
 	const ERHIFeatureLevel::Type FeatureLevel = GEditor->GetEditorWorldContext().World()->FeatureLevel;
 
 	for (int32 MaterialIndex = 0; MaterialIndex < MaterialsToOverride.Num(); MaterialIndex++)
@@ -3188,6 +3195,7 @@ void FMaterialEditor::SetScalarParameterDefaultOnDependentMaterials(FName Parame
 		UMaterial* CurrentMaterial = MaterialsToOverride[MaterialIndex];
 
 		CurrentMaterial->OverrideScalarParameterDefault(ParameterName, Value, bOverride, FeatureLevel);
+		UpdateContext.AddMaterial(CurrentMaterial);
 	}
 
 	// Update MI's that reference any of the materials affected
