Adding a parameter's source function or material to its tooltip
#rnx

diff --git a/Engine/Source/Editor/MaterialEditor/Private/MaterialEditorInstanceDetailCustomization.cpp b/Engine/Source/Editor/MaterialEditor/Private/MaterialEditorInstanceDetailCustomization.cpp
--- a/Engine/Source/Editor/MaterialEditor/Private/MaterialEditorInstanceDetailCustomization.cpp
+++ b/Engine/Source/Editor/MaterialEditor/Private/MaterialEditorInstanceDetailCustomization.cpp
@@ -456,7 +456,7 @@ void FMaterialInstanceParameterDetails::CreateParameterValueWidget(UDEditorParam
 
 		PropertyRow
 			.DisplayName(FText::FromName(Parameter->ParameterInfo.Name))
-			.ToolTip(FMaterialPropertyHelpers::GetParameterExpressionDescription(Parameter, MaterialEditorInstance))
+			.ToolTip(FMaterialPropertyHelpers::GetParameterTooltip(Parameter, MaterialEditorInstance))
 			.EditCondition(IsParamEnabled, FOnBooleanValueChanged::CreateStatic(&FMaterialPropertyHelpers::OnOverrideParameter, Parameter, MaterialEditorInstance))
 			.Visibility(TAttribute<EVisibility>::Create(TAttribute<EVisibility>::FGetter::CreateStatic(&FMaterialPropertyHelpers::ShouldShowExpression, Parameter, MaterialEditorInstance, ShowHiddenDelegate)))
 			.OverrideResetToDefault(ResetOverride);
@@ -490,7 +490,7 @@ void FMaterialInstanceParameterDetails::CreateMaskParameterValueWidget(UDEditorP
 			[
 				SNew(STextBlock)
 				.Text(ParameterName)
-				.ToolTipText(FMaterialPropertyHelpers::GetParameterExpressionDescription(Parameter, MaterialEditorInstance))
+				.ToolTipText(FMaterialPropertyHelpers::GetParameterTooltip(Parameter, MaterialEditorInstance))
 				.Font(FEditorStyle::GetFontStyle(TEXT("PropertyWindow.NormalFont")))
 			]
 		.ValueContent()
@@ -586,7 +586,7 @@ void FMaterialInstanceParameterDetails::CreateVectorChannelMaskParameterValueWid
 		[
 			SNew(STextBlock)
 			.Text(ParameterName)
-			.ToolTipText(FMaterialPropertyHelpers::GetParameterExpressionDescription(Parameter, MaterialEditorInstance))
+			.ToolTipText(FMaterialPropertyHelpers::GetParameterTooltip(Parameter, MaterialEditorInstance))
 			.Font(FEditorStyle::GetFontStyle(TEXT("PropertyWindow.NormalFont")))
 		]
 		.ValueContent()
@@ -637,7 +637,7 @@ void FMaterialInstanceParameterDetails::CreateScalarAtlasPositionParameterValueW
 			[
 				SNew(STextBlock)
 				.Text(ParameterName)
-				.ToolTipText(FMaterialPropertyHelpers::GetParameterExpressionDescription(Parameter, MaterialEditorInstance))
+				.ToolTipText(FMaterialPropertyHelpers::GetParameterTooltip(Parameter, MaterialEditorInstance))
 				.Font(FEditorStyle::GetFontStyle(TEXT("PropertyWindow.NormalFont")))
 			]
 			.ValueContent()
diff --git a/Engine/Source/Editor/MaterialEditor/Private/MaterialPropertyHelpers.cpp b/Engine/Source/Editor/MaterialEditor/Private/MaterialPropertyHelpers.cpp
--- a/Engine/Source/Editor/MaterialEditor/Private/MaterialPropertyHelpers.cpp
+++ b/Engine/Source/Editor/MaterialEditor/Private/MaterialPropertyHelpers.cpp
@@ -704,7 +704,7 @@ FText FMaterialPropertyHelpers::GetParameterExpressionDescription(UDEditorParame
 	UMaterialEditorInstanceConstant* MaterialInstanceEditor = Cast<UMaterialEditorInstanceConstant>(MaterialEditorInstance);
 	if (MaterialInstanceEditor)
 	{
-		BaseMaterial = MaterialInstanceEditor->SourceInstance->GetMaterial();;
+		BaseMaterial = MaterialInstanceEditor->SourceInstance->GetMaterial();
 	}
 	UMaterialEditorPreviewParameters* MaterialEditor = Cast<UMaterialEditorPreviewParameters>(MaterialEditorInstance);
 	if (MaterialEditor)
@@ -726,6 +726,41 @@ FText FMaterialPropertyHelpers::GetParameterExpressionDescription(UDEditorParame
 	return FText::GetEmpty();
 }
 
+FText FMaterialPropertyHelpers::GetParameterTooltip(UDEditorParameterValue* Parameter, UObject* MaterialEditorInstance)
+{
+	UMaterial* BaseMaterial = nullptr;
+	FText FoundInLocationText = FText::FromString(FPaths::GetBaseFilename(Parameter->ParameterInfo.ParameterLocation.GetAssetPathName().ToString()));
+	UMaterialEditorInstanceConstant* MaterialInstanceEditor = Cast<UMaterialEditorInstanceConstant>(MaterialEditorInstance);
+	if (MaterialInstanceEditor)
+	{
+		BaseMaterial = MaterialInstanceEditor->SourceInstance->GetMaterial();
+	}
+	UMaterialEditorPreviewParameters* MaterialEditor = Cast<UMaterialEditorPreviewParameters>(MaterialEditorInstance);
+	if (MaterialEditor)
+	{
+		BaseMaterial = MaterialEditor->OriginalMaterial;
+	}
+
+	// TODO: This needs to support functions added by SourceInstance layers
+	if (BaseMaterial)
+	{
+		UMaterialExpression* MaterialExpression = BaseMaterial->FindExpressionByGUID<UMaterialExpression>(Parameter->ExpressionId);
+
+		if (MaterialExpression)
+		{
+			FText TooltipText = FText::Format(LOCTEXT("ParameterInfoLocationOnly", "Found in: {0}"), FoundInLocationText);
+			if (!MaterialExpression->Desc.IsEmpty())
+			{
+				TooltipText = FText::Format(LOCTEXT("ParameterInfoDescAndLocation", "{0} \nFound in: {1}"), FText::FromString(MaterialExpression->Desc), FoundInLocationText);
+			}
+			return TooltipText;
+		}
+	}
+
+	return FText::GetEmpty();
+}
+
+
 void FMaterialPropertyHelpers::ResetToDefault(TSharedPtr<IPropertyHandle> PropertyHandle, class UDEditorParameterValue* Parameter, UMaterialEditorInstanceConstant* MaterialEditorInstance)
 {
 	const FScopedTransaction Transaction( LOCTEXT( "ResetToDefault", "Reset To Default" ) );
diff --git a/Engine/Source/Editor/MaterialEditor/Public/MaterialPropertyHelpers.h b/Engine/Source/Editor/MaterialEditor/Public/MaterialPropertyHelpers.h
--- a/Engine/Source/Editor/MaterialEditor/Public/MaterialPropertyHelpers.h
+++ b/Engine/Source/Editor/MaterialEditor/Public/MaterialPropertyHelpers.h
@@ -150,6 +150,7 @@ public:
 	/** Gets the expression description of this parameter from the base material */
 	static	FText GetParameterExpressionDescription(class UDEditorParameterValue* Parameter, UObject* MaterialEditorInstance);
 	
+	static FText GetParameterTooltip(class UDEditorParameterValue* Parameter, UObject* MaterialEditorInstance);
 	/**
 	 * Called when a parameter is overridden;
 	 */
diff --git a/Engine/Source/Runtime/Engine/Classes/Materials/MaterialLayersFunctions.h b/Engine/Source/Runtime/Engine/Classes/Materials/MaterialLayersFunctions.h
--- a/Engine/Source/Runtime/Engine/Classes/Materials/MaterialLayersFunctions.h
+++ b/Engine/Source/Runtime/Engine/Classes/Materials/MaterialLayersFunctions.h
@@ -35,6 +35,11 @@ struct ENGINE_API FMaterialParameterInfo
 	UPROPERTY(EditAnywhere, BlueprintReadWrite, Category=ParameterInfo)
 	int32 Index;
 
+#if WITH_EDITORONLY_DATA
+	UPROPERTY(Transient)
+	FSoftObjectPath ParameterLocation;
+#endif
+
 	FMaterialParameterInfo(const TCHAR* InName, EMaterialParameterAssociation InAssociation = EMaterialParameterAssociation::GlobalParameter, int32 InIndex = INDEX_NONE)
 		: Name(InName)
 		, Association(InAssociation)
diff --git a/Engine/Source/Runtime/Engine/Private/Materials/MaterialExpressions.cpp b/Engine/Source/Runtime/Engine/Private/Materials/MaterialExpressions.cpp
--- a/Engine/Source/Runtime/Engine/Private/Materials/MaterialExpressions.cpp
+++ b/Engine/Source/Runtime/Engine/Private/Materials/MaterialExpressions.cpp
@@ -1981,6 +1981,15 @@ void UMaterialExpressionTextureSampleParameter::GetAllParameterInfo(TArray<FMate
 {
 	int32 CurrentSize = OutParameterInfo.Num();
 	FMaterialParameterInfo NewParameter(ParameterName, InBaseParameterInfo.Association, InBaseParameterInfo.Index);
+
+#if WITH_EDITOR
+	NewParameter.ParameterLocation = Material;
+	if (Function != nullptr)
+	{
+		NewParameter.ParameterLocation = Function;
+	}
+#endif
+
 	OutParameterInfo.AddUnique(NewParameter);
 
 	if(CurrentSize != OutParameterInfo.Num())
@@ -6187,6 +6196,15 @@ void UMaterialExpressionParameter::GetAllParameterInfo(TArray<FMaterialParameter
 {
 	int32 CurrentSize = OutParameterInfo.Num();
 	FMaterialParameterInfo NewParameter(ParameterName, InBaseParameterInfo.Association, InBaseParameterInfo.Index);
+
+#if WITH_EDITOR
+	NewParameter.ParameterLocation = Material;
+	if (Function != nullptr)
+	{
+		NewParameter.ParameterLocation = Function;
+	}
+#endif
+
 	OutParameterInfo.AddUnique(NewParameter);
 	if(CurrentSize != OutParameterInfo.Num())
 	{
@@ -9216,6 +9234,13 @@ void UMaterialExpressionFontSampleParameter::GetAllParameterInfo(TArray<FMateria
 {
 	int32 CurrentSize = OutParameterInfo.Num();
 	FMaterialParameterInfo NewParameter(ParameterName, InBaseParameterInfo.Association, InBaseParameterInfo.Index);
+#if WITH_EDITOR
+	NewParameter.ParameterLocation = Material;
+	if (Function != nullptr)
+	{
+		NewParameter.ParameterLocation = Function;
+	}
+#endif
 	OutParameterInfo.AddUnique(NewParameter);
 	if(CurrentSize != OutParameterInfo.Num())
 	{
