Speculative fix for editing material after changing Decal Blend Mode may cause editor to crash.  Move my previous change in material pre-edit and place in FMaterialEditor undo / redo instead AND in graph changes.  Can't repro the graph change version of the crash but I can repro the undo redo version and this change keeps that fix in place.

Reasons for the flush:
In Editor realtime rendering changing a material connection (or via undo/redo) changes what is rendered on the RT (specifically in FMeshDecalMeshProcessor::AddMeshBatch() Emmissive connection test)
before the Shader Map has been updated in PostEditChangeProperty() - Viewport render command enqueue has already happenned so we have a data race and we need to flush that before the material change.

diff --git a/Engine/Source/Editor/MaterialEditor/Private/MaterialEditor.cpp b/Engine/Source/Editor/MaterialEditor/Private/MaterialEditor.cpp
--- a/Engine/Source/Editor/MaterialEditor/Private/MaterialEditor.cpp
+++ b/Engine/Source/Editor/MaterialEditor/Private/MaterialEditor.cpp
@@ -3927,6 +3927,8 @@ FText FMaterialEditor::GetOriginalObjectName() const
 
 void FMaterialEditor::UpdateMaterialAfterGraphChange()
 {
+	FlushRenderingCommands();
+	
 	Material->MaterialGraph->LinkMaterialExpressionsFromGraph();
 
 	// Update the current preview material.
@@ -3951,6 +3953,8 @@ FMaterialRenderProxy* FMaterialEditor::GetExpressionPreview(UMaterialExpression*
 
 void FMaterialEditor::UndoGraphAction()
 {
+	FlushRenderingCommands();
+	
 	int32 NumExpressions = Material->Expressions.Num();
 	GEditor->UndoTransaction();
 
@@ -3962,6 +3966,8 @@ void FMaterialEditor::UndoGraphAction()
 
 void FMaterialEditor::RedoGraphAction()
 {
+	FlushRenderingCommands();
+	
 	// Clear selection, to avoid holding refs to nodes that go away
 	GraphEditor->ClearSelectionSet();
 
diff --git a/Engine/Source/Runtime/Engine/Private/Materials/Material.cpp b/Engine/Source/Runtime/Engine/Private/Materials/Material.cpp
--- a/Engine/Source/Runtime/Engine/Private/Materials/Material.cpp
+++ b/Engine/Source/Runtime/Engine/Private/Materials/Material.cpp
@@ -4280,9 +4280,6 @@ bool UMaterial::CanEditChange(const UProperty* InProperty) const
 
 void UMaterial::PreEditChange(UProperty* PropertyThatChanged)
 {
-	// In Editor realtime rendering changing a material connection (or via undo/redo) changes what is rendered on the RT (specifically in FMeshDecalMeshProcessor::AddMeshBatch() Emmissive connection test)
-	// before the Shader Map has been updated in PostEditChangeProperty() below - Viewport render command enqueue has already happenned so we need to flush that before the material change.
-	FlushRenderingCommands();
 	Super::PreEditChange(PropertyThatChanged);
 }
 
