Some UI improvements to matrerial scalar/vector parameters when using the custom primitive data.
Added a debug string field to display all the shading models in a material in UMaterial, this is displayed only when "From Material Expression" is the selected shading model on the material.

#rnx

diff --git a/Engine/Source/Editor/MaterialEditor/Private/MaterialEditorDetailCustomization.cpp b/Engine/Source/Editor/MaterialEditor/Private/MaterialEditorDetailCustomization.cpp
--- a/Engine/Source/Editor/MaterialEditor/Private/MaterialEditorDetailCustomization.cpp
+++ b/Engine/Source/Editor/MaterialEditor/Private/MaterialEditorDetailCustomization.cpp
@@ -11,6 +11,7 @@
 #include "DetailWidgetRow.h"
 #include "DetailCategoryBuilder.h"
 #include "Materials/MaterialExpressionScalarParameter.h"
+#include "Materials/MaterialExpressionVectorParameter.h"
 #include "MaterialLayersFunctionsCustomization.h"
 #include "PropertyEditorModule.h"
 #include "MaterialEditor/MaterialEditorPreviewParameters.h"
@@ -449,13 +450,15 @@ void FMaterialExpressionParameterDetails::CustomizeDetails( IDetailLayoutBuilder
 	// for expression parameters all their properties are in one category based on their class name.
 	FName DefaultCategory = NAME_None;
 	IDetailCategoryBuilder& Category = DetailLayout.EditCategory( DefaultCategory );
-
+	
 	TArray< TWeakObjectPtr<UObject> > Objects;
 	DetailLayout.GetObjectsBeingCustomized(Objects);
 
 	DefaultValueHandles.Reset();
 	ScalarParameterObjects.Reset();
 
+	const FName MaterialExpressionCategory = TEXT("MaterialExpression");
+
 	for (const auto& WeakObjectPtr : Objects)
 	{
 		UObject* Object = WeakObjectPtr.Get();
@@ -486,14 +489,51 @@ void FMaterialExpressionParameterDetails::CustomizeDetails( IDetailLayoutBuilder
 			}
 
 			OnSliderMinMaxEdited();
+
+			TSharedPtr<IPropertyHandle> PropertyHandle = DetailLayout.GetProperty("bUseCustomPrimitiveData", UMaterialExpressionScalarParameter::StaticClass());
+			if (PropertyHandle.IsValid() && PropertyHandle->IsValidHandle())
+			{
+				// Rebuild the layout when the bUseCustomPrimitiveData property changes
+				PropertyHandle->SetOnPropertyValueChanged(FSimpleDelegate::CreateLambda([&DetailLayout]()
+				{
+					DetailLayout.ForceRefreshDetails();
+				}));
+			}
+
+			if (ScalarParameter->bUseCustomPrimitiveData)
+			{
+				DetailLayout.HideCategory(TEXT("MaterialExpressionScalarParameter"));
+				DetailLayout.HideCategory(MaterialExpressionCategory);
+			}
+		}
+
+		UMaterialExpressionVectorParameter* VectorParameter = Cast<UMaterialExpressionVectorParameter>(Object);
+
+		if (VectorParameter)
+		{
+			TSharedPtr<IPropertyHandle> PropertyHandle = DetailLayout.GetProperty("bUseCustomPrimitiveData", UMaterialExpressionVectorParameter::StaticClass());
+			if (PropertyHandle.IsValid() && PropertyHandle->IsValidHandle())
+			{
+				// Rebuild the layout when the bUseCustomPrimitiveData property changes
+				PropertyHandle->SetOnPropertyValueChanged(FSimpleDelegate::CreateLambda([&DetailLayout]()
+				{
+					DetailLayout.ForceRefreshDetails();
+				}));
+			}
+			
+			if (VectorParameter->bUseCustomPrimitiveData)
+			{
+				DetailLayout.HideCategory(TEXT("MaterialExpressionVectorParameter"));
+				DetailLayout.HideCategory(MaterialExpressionCategory);
+			}
 		}
+
 	}
 
 	check(ScalarParameterObjects.Num() == DefaultValueHandles.Num());
 	
 	Category.AddProperty("ParameterName");
 
-	const FName MaterialExpressionCategory = TEXT("MaterialExpression");
 	IDetailCategoryBuilder& ExpressionCategory = DetailLayout.EditCategory(MaterialExpressionCategory);
 
 	// Get a handle to the property we are about to edit
@@ -809,12 +849,16 @@ void FMaterialDetailCustomization::CustomizeDetails( IDetailLayoutBuilder& Detai
 	DetailLayout.GetObjectsBeingCustomized( Objects );
 
 	bool bUIMaterial = true;
+	bool bIsShadingModelFromMaterialExpression = false;
 	for( TWeakObjectPtr<UObject>& Object : Objects )
 	{
 		UMaterial* Material = Cast<UMaterial>( Object.Get() );
 		if( Material )
 		{
 			bUIMaterial &= Material->IsUIMaterial();
+
+			// If any Object has its shading model from material expression
+			bIsShadingModelFromMaterialExpression |= Material->IsShadingModelFromMaterialExpression();
 		}
 		else
 		{
@@ -823,30 +867,20 @@ void FMaterialDetailCustomization::CustomizeDetails( IDetailLayoutBuilder& Detai
 		}
 	}
 
-	if( bUIMaterial )
+	// Material category
 	{
-		DetailLayout.HideCategory( TEXT("TranslucencySelfShadowing") );
-		DetailLayout.HideCategory( TEXT("Translucency") );
-		DetailLayout.HideCategory( TEXT("Tessellation") );
-		DetailLayout.HideCategory( TEXT("PostProcessMaterial") );
-		DetailLayout.HideCategory( TEXT("Lightmass") );
-		DetailLayout.HideCategory( TEXT("Thumbnail") );
-		DetailLayout.HideCategory( TEXT("MaterialInterface") );
-		DetailLayout.HideCategory( TEXT("PhysicalMaterial") );
-		DetailLayout.HideCategory( TEXT("Usage") );
+		IDetailCategoryBuilder& MaterialCategory = DetailLayout.EditCategory( TEXT("Material") );
 
-		// Materail category
-		{
-			IDetailCategoryBuilder& MaterialCategory = DetailLayout.EditCategory( TEXT("Material") );
+		TArray<TSharedRef<IPropertyHandle>> AllProperties;
+		MaterialCategory.GetDefaultProperties( AllProperties );
 
-			TArray<TSharedRef<IPropertyHandle>> AllProperties;
-			MaterialCategory.GetDefaultProperties( AllProperties );
+		for( TSharedRef<IPropertyHandle>& PropertyHandle : AllProperties )
+		{
+			UProperty* Property = PropertyHandle->GetProperty();
+			FName PropertyName = Property->GetFName();
 
-			for( TSharedRef<IPropertyHandle>& PropertyHandle : AllProperties )
+			if (bUIMaterial)
 			{
-				UProperty* Property = PropertyHandle->GetProperty();
-				FName PropertyName = Property->GetFName();
-
 				if(		PropertyName != GET_MEMBER_NAME_CHECKED(UMaterial, MaterialDomain) 
 					&&	PropertyName != GET_MEMBER_NAME_CHECKED(UMaterial, BlendMode) 
 					&&	PropertyName != GET_MEMBER_NAME_CHECKED(UMaterial, OpacityMaskClipValue) 
@@ -855,7 +889,28 @@ void FMaterialDetailCustomization::CustomizeDetails( IDetailLayoutBuilder& Detai
 					DetailLayout.HideProperty( PropertyHandle );
 				}
 			}
+
+#if WITH_EDITORONLY_DATA
+			// Hide the shading model 
+			if (!bIsShadingModelFromMaterialExpression && PropertyName == GET_MEMBER_NAME_CHECKED(UMaterial, UsedShadingModels))
+			{
+				 DetailLayout.HideProperty(PropertyHandle);
+			}
+#endif
 		}
+	}
+
+	if( bUIMaterial )
+	{
+		DetailLayout.HideCategory( TEXT("TranslucencySelfShadowing") );
+		DetailLayout.HideCategory( TEXT("Translucency") );
+		DetailLayout.HideCategory( TEXT("Tessellation") );
+		DetailLayout.HideCategory( TEXT("PostProcessMaterial") );
+		DetailLayout.HideCategory( TEXT("Lightmass") );
+		DetailLayout.HideCategory( TEXT("Thumbnail") );
+		DetailLayout.HideCategory( TEXT("MaterialInterface") );
+		DetailLayout.HideCategory( TEXT("PhysicalMaterial") );
+		DetailLayout.HideCategory( TEXT("Usage") );
 
 		// Mobile category
 		{
diff --git a/Engine/Source/Runtime/Engine/Classes/Materials/Material.h b/Engine/Source/Runtime/Engine/Classes/Materials/Material.h
--- a/Engine/Source/Runtime/Engine/Classes/Materials/Material.h
+++ b/Engine/Source/Runtime/Engine/Classes/Materials/Material.h
@@ -1756,6 +1756,7 @@ private:
 	friend class FMaterialUpdateContext;
 	friend class FMaterialResource;
 	friend class FMaterialEditor;
+	friend class FMaterialDetailCustomization;
 
 	// DO NOT CALL outside of FMaterialEditor!
 	ENGINE_API static void ForceNoCompilationInPostLoad(bool bForceNoCompilation);
