Now treating the subsurface profile override like the other material override poroperties in the Material Instance editor UI

diff --git a/Engine/Source/Editor/MaterialEditor/Private/MaterialEditorInstanceDetailCustomization.cpp b/Engine/Source/Editor/MaterialEditor/Private/MaterialEditorInstanceDetailCustomization.cpp
--- a/Engine/Source/Editor/MaterialEditor/Private/MaterialEditorInstanceDetailCustomization.cpp
+++ b/Engine/Source/Editor/MaterialEditor/Private/MaterialEditorInstanceDetailCustomization.cpp
@@ -231,17 +231,31 @@ void FMaterialInstanceParameterDetails::CustomizeDetails(IDetailLayoutBuilder& D
 		}
 
 		{
-			IDetailPropertyRow& PropertyRow = DefaultCategory.AddProperty("bOverrideSubsurfaceProfile");
-			PropertyRow.Visibility(TAttribute<EVisibility>::Create(TAttribute<EVisibility>::FGetter::CreateSP(this, &FMaterialInstanceParameterDetails::ShouldShowSubsurfaceProfile)));
+			// Add the material property override group
+			static FName GroupName(TEXT("MaterialPropertyOverrideGroup"));
+			IDetailGroup& MaterialPropertyOverrideGroup = DefaultCategory.AddGroup(GroupName, LOCTEXT("MaterialPropertyOverrideGroup", "Material Property Overrides"), false, false);
+			
+			// Hide the originals, these will be recreated manually
+			DetailLayout.HideProperty("bOverrideSubsurfaceProfile");
+			DetailLayout.HideProperty("SubsurfaceProfile");
+			DetailLayout.HideProperty("BasePropertyOverrides");
+
+			// Set up the override logic for the subsurface profile
+			TAttribute<bool> IsParamEnabled = TAttribute<bool>::Create(TAttribute<bool>::FGetter::CreateLambda([this](){ return (bool)MaterialEditorInstance->bOverrideSubsurfaceProfile; }));
+
+			IDetailPropertyRow& PropertyRow = MaterialPropertyOverrideGroup.AddPropertyRow(DetailLayout.GetProperty("SubsurfaceProfile"));
+			PropertyRow
+				.EditCondition(IsParamEnabled, 
+					FOnBooleanValueChanged::CreateLambda([this](bool NewValue) {
+						MaterialEditorInstance->bOverrideSubsurfaceProfile = (uint32)NewValue;
+						MaterialEditorInstance->PostEditChange();
+						FEditorSupportDelegates::RedrawAllViewports.Broadcast();
+				}))
+				.Visibility(TAttribute<EVisibility>::Create(TAttribute<EVisibility>::FGetter::CreateSP(this, &FMaterialInstanceParameterDetails::ShouldShowSubsurfaceProfile)));
+			
+			// Append the base property overrides to the Material Property Override Group
+			CreateBasePropertyOverrideWidgets(DetailLayout, MaterialPropertyOverrideGroup);
 		}
-
-		{
-			IDetailPropertyRow& PropertyRow = DefaultCategory.AddProperty("SubsurfaceProfile");
-			PropertyRow.Visibility(TAttribute<EVisibility>::Create(TAttribute<EVisibility>::FGetter::CreateSP(this, &FMaterialInstanceParameterDetails::ShouldShowSubsurfaceProfile)));
-		}
-
-		DetailLayout.HideProperty("BasePropertyOverrides");
-		CreateBasePropertyOverrideWidgets(DetailLayout);
 	}
 
 	// Add the preview mesh property directly from the material instance 
@@ -851,12 +865,9 @@ void FMaterialInstanceParameterDetails::CreateLightmassOverrideWidgets(IDetailLa
 		.OverrideResetToDefault(ResetExportResolutionScalePropertyOverride);
 }
 
-void FMaterialInstanceParameterDetails::CreateBasePropertyOverrideWidgets(IDetailLayoutBuilder& DetailLayout)
+void FMaterialInstanceParameterDetails::CreateBasePropertyOverrideWidgets(IDetailLayoutBuilder& DetailLayout, IDetailGroup& MaterialPropertyOverrideGroup)
 {
-	IDetailCategoryBuilder& DetailCategory = DetailLayout.EditCategory(NAME_None);
-	
-	static FName GroupName(TEXT("BasePropertyOverrideGroup"));
-	IDetailGroup& BasePropertyOverrideGroup = DetailCategory.AddGroup(GroupName, LOCTEXT("BasePropertyOverrideGroup", "Material Property Overrides"), false, false);
+	IDetailGroup& BasePropertyOverrideGroup = MaterialPropertyOverrideGroup;
 
 	TAttribute<bool> IsOverrideOpacityClipMaskValueEnabled = TAttribute<bool>::Create(TAttribute<bool>::FGetter::CreateSP(this, &FMaterialInstanceParameterDetails::OverrideOpacityClipMaskValueEnabled));
 	TAttribute<bool> IsOverrideBlendModeEnabled = TAttribute<bool>::Create(TAttribute<bool>::FGetter::CreateSP(this, &FMaterialInstanceParameterDetails::OverrideBlendModeEnabled));
diff --git a/Engine/Source/Editor/MaterialEditor/Private/MaterialEditorInstanceDetailCustomization.h b/Engine/Source/Editor/MaterialEditor/Private/MaterialEditorInstanceDetailCustomization.h
--- a/Engine/Source/Editor/MaterialEditor/Private/MaterialEditorInstanceDetailCustomization.h
+++ b/Engine/Source/Editor/MaterialEditor/Private/MaterialEditorInstanceDetailCustomization.h
@@ -79,7 +79,7 @@ private:
 	//Functions supporting BasePropertyOverrides
 
 	/** Creates all the base property override widgets. */
-	void CreateBasePropertyOverrideWidgets(IDetailLayoutBuilder& DetailLayout);
+	void CreateBasePropertyOverrideWidgets(IDetailLayoutBuilder& DetailLayout, IDetailGroup& MaterialPropertyOverrideGroup);
 
 	EVisibility IsOverriddenAndVisible(TAttribute<bool> IsOverridden) const;
 
