PR #5560: Realtime option User Interface material previews (Contributed by [REDACTED_PERSON])

diff --git a/Engine/Source/Editor/MaterialEditor/Private/SMaterialEditorViewport.cpp b/Engine/Source/Editor/MaterialEditor/Private/SMaterialEditorViewport.cpp
--- a/Engine/Source/Editor/MaterialEditor/Private/SMaterialEditorViewport.cpp
+++ b/Engine/Source/Editor/MaterialEditor/Private/SMaterialEditorViewport.cpp
@@ -29,6 +29,7 @@
 #include "AdvancedPreviewScene.h"
 #include "AssetViewerSettings.h"
 #include "Engine/PostProcessVolume.h"
+#include "Widgets/Input/SCheckBox.h"
 
 #define LOCTEXT_NAMESPACE "MaterialEditor"
 
@@ -890,7 +891,7 @@ void SMaterialEditorUIPreviewZoomer::SetPreviewMaterial(UMaterialInterface* InPr
 
 void SMaterialEditorUIPreviewViewport::Construct( const FArguments& InArgs, UMaterialInterface* PreviewMaterial )
 {
-
+	bRealtime = false;
 	ChildSlot
 	[
 		SNew( SVerticalBox )
@@ -905,6 +906,21 @@ void SMaterialEditorUIPreviewViewport::Construct( const FArguments& InArgs, UMat
 				SNew( SHorizontalBox )
 				+ SHorizontalBox::Slot()
 				.VAlign(VAlign_Center)
+				.Padding( 9.f )
+				.AutoWidth()
+				[
+					SNew( SCheckBox )
+					.CheckBoxContentUsesAutoWidth( true )
+					.OnCheckStateChanged(this, &SMaterialEditorUIPreviewViewport::OnRealtimeChanged)
+					.IsChecked(this, &SMaterialEditorUIPreviewViewport::IsRealtimeChecked)
+					.Content()
+					[
+						SNew(STextBlock)
+						.Text(LOCTEXT("Realtime", "Realtime"))
+					]
+				]
+				+ SHorizontalBox::Slot()
+				.VAlign(VAlign_Center)
 				.Padding( 3.f )
 				.AutoWidth()
 				[
@@ -1005,4 +1021,32 @@ void SMaterialEditorUIPreviewViewport::OnPreviewYCommitted( int32 NewValue, ETex
 	OnPreviewYChanged( NewValue );
 }
 
+void SMaterialEditorUIPreviewViewport::OnRealtimeChanged(ECheckBoxState State)
+{
+	if (State == ECheckBoxState::Checked)
+	{
+		bRealtime = true;
+		ActiveTimerHandle = RegisterActiveTimer(0.f, FWidgetActiveTimerDelegate::CreateSP(this, &SMaterialEditorUIPreviewViewport::EnsureTick));
+	}
+	else
+	{
+		bRealtime = false;
+		if (ActiveTimerHandle.IsValid())
+		{
+			UnRegisterActiveTimer(ActiveTimerHandle.Pin().ToSharedRef());
+		}
+	}
+}
+
+ECheckBoxState SMaterialEditorUIPreviewViewport::IsRealtimeChecked() const
+{
+	return bRealtime ? ECheckBoxState::Checked : ECheckBoxState::Unchecked;
+}
+
+EActiveTimerReturnType SMaterialEditorUIPreviewViewport::EnsureTick(double InCurrentTime, float InDeltaTime)
+{
+	// Keep the timer going if we're realtime
+	return bRealtime ? EActiveTimerReturnType::Continue : EActiveTimerReturnType::Stop;
+}
+
 #undef LOCTEXT_NAMESPACE
diff --git a/Engine/Source/Editor/MaterialEditor/Private/SMaterialEditorViewport.h b/Engine/Source/Editor/MaterialEditor/Private/SMaterialEditorViewport.h
--- a/Engine/Source/Editor/MaterialEditor/Private/SMaterialEditorViewport.h
+++ b/Engine/Source/Editor/MaterialEditor/Private/SMaterialEditorViewport.h
@@ -148,7 +148,13 @@ private:
 	void OnPreviewYCommitted( int32 NewValue, ETextCommit::Type );
 	TOptional<int32> OnGetPreviewXValue() const { return PreviewSize.X; }
 	TOptional<int32> OnGetPreviewYValue() const { return PreviewSize.Y; }
+	void OnRealtimeChanged(ECheckBoxState State);
+	ECheckBoxState IsRealtimeChecked() const;
+	EActiveTimerReturnType EnsureTick(double InCurrentTime, float InDeltaTime);
+
 private:
 	FIntPoint PreviewSize;
 	TSharedPtr<class SMaterialEditorUIPreviewZoomer> PreviewZoomer;
+	bool bRealtime;
+	TWeakPtr<FActiveTimerHandle> ActiveTimerHandle;
 };
