Adding Material Editing Library script-callable functions to find the source of parameters (ie did they originate in the material itself or a nested function)
#rnx

diff --git a/Engine/Source/Editor/MaterialEditor/Private/MaterialEditingLibrary.cpp b/Engine/Source/Editor/MaterialEditor/Private/MaterialEditingLibrary.cpp
--- a/Engine/Source/Editor/MaterialEditor/Private/MaterialEditingLibrary.cpp
+++ b/Engine/Source/Editor/MaterialEditor/Private/MaterialEditingLibrary.cpp
@@ -1046,3 +1046,87 @@ void UMaterialEditingLibrary::GetStaticSwitchParameterNames(UMaterialInterface*
 		}
 	}
 }
+
+bool UMaterialEditingLibrary::GetScalarParameterSource(UMaterialInterface* Material, const FName ParameterName, FSoftObjectPath& ParameterSource)
+{
+	if (Material)
+	{
+		TArray<FMaterialParameterInfo> MaterialInfo;
+		TArray<FGuid> MaterialGuids;
+		Material->GetAllScalarParameterInfo(MaterialInfo, MaterialGuids);
+		FMaterialParameterInfo* ParameterInfo = MaterialInfo.FindByPredicate([ParameterName](const FMaterialParameterInfo& Parameter)
+		{
+			return ParameterName == Parameter.Name;
+		});
+		
+		if (ParameterInfo)
+		{
+			ParameterSource = ParameterInfo->ParameterLocation;
+			return true;
+		}
+	}
+	return false;
+}
+
+bool UMaterialEditingLibrary::GetVectorParameterSource(UMaterialInterface* Material, const FName ParameterName, FSoftObjectPath& ParameterSource)
+{
+	if (Material)
+	{
+		TArray<FMaterialParameterInfo> MaterialInfo;
+		TArray<FGuid> MaterialGuids;
+		Material->GetAllVectorParameterInfo(MaterialInfo, MaterialGuids);
+		FMaterialParameterInfo* ParameterInfo = MaterialInfo.FindByPredicate([ParameterName](const FMaterialParameterInfo& Parameter)
+		{
+			return ParameterName == Parameter.Name;
+		});
+
+		if (ParameterInfo)
+		{
+			ParameterSource = ParameterInfo->ParameterLocation;
+			return true;
+		}
+	}
+	return false;
+}
+
+bool UMaterialEditingLibrary::GetTextureParameterSource(UMaterialInterface* Material, const FName ParameterName, FSoftObjectPath& ParameterSource)
+{
+	if (Material)
+	{
+		TArray<FMaterialParameterInfo> MaterialInfo;
+		TArray<FGuid> MaterialGuids;
+		Material->GetAllTextureParameterInfo(MaterialInfo, MaterialGuids);
+		FMaterialParameterInfo* ParameterInfo = MaterialInfo.FindByPredicate([ParameterName](const FMaterialParameterInfo& Parameter)
+		{
+			return ParameterName == Parameter.Name;
+		});
+
+		if (ParameterInfo)
+		{
+			ParameterSource = ParameterInfo->ParameterLocation;
+			return true;
+		}
+	}
+	return false;
+}
+
+bool UMaterialEditingLibrary::GetStaticSwitchParameterSource(UMaterialInterface* Material, const FName ParameterName, FSoftObjectPath& ParameterSource)
+{
+	if (Material)
+	{
+		TArray<FMaterialParameterInfo> MaterialInfo;
+		TArray<FGuid> MaterialGuids;
+		Material->GetAllStaticSwitchParameterInfo(MaterialInfo, MaterialGuids);
+		FMaterialParameterInfo* ParameterInfo = MaterialInfo.FindByPredicate([ParameterName](const FMaterialParameterInfo& Parameter)
+		{
+			return ParameterName == Parameter.Name;
+		});
+
+		if (ParameterInfo)
+		{
+			ParameterSource = ParameterInfo->ParameterLocation;
+			return true;
+		}
+	}
+	return false;
+}
\ No newline at end of file
diff --git a/Engine/Source/Editor/MaterialEditor/Public/MaterialEditingLibrary.h b/Engine/Source/Editor/MaterialEditor/Public/MaterialEditingLibrary.h
--- a/Engine/Source/Editor/MaterialEditor/Public/MaterialEditingLibrary.h
+++ b/Engine/Source/Editor/MaterialEditor/Public/MaterialEditingLibrary.h
@@ -236,4 +236,44 @@ public:
 	/** Gets all static switch parameter names */
 	UFUNCTION(BlueprintCallable, Category = "MaterialEditing")
 	static void GetStaticSwitchParameterNames(UMaterialInterface* Material, TArray<FName>& ParameterNames);
+
+	/**
+	*	Returns the path of the asset where the parameter originated, as well as true/false if it was found
+	*	@param	Material	The material or material instance you want to look up a parameter from
+	*	@param	ParameterName		The parameter name
+	*	@param	ParameterSource		The soft object path of the asset the parameter originates in 
+	*	@return	Whether or not the parameter was found in this material
+	*/
+	UFUNCTION(BlueprintCallable, Category = "MaterialEditing")
+	bool GetScalarParameterSource(UMaterialInterface* Material, const FName ParameterName, FSoftObjectPath& ParameterSource);
+
+	/**
+	*	Returns the path of the asset where the parameter originated, as well as true/false if it was found
+	*	@param	Material	The material or material instance you want to look up a parameter from
+	*	@param	ParameterName		The parameter name
+	*	@param	ParameterSource		The soft object path of the asset the parameter originates in
+	*	@return	Whether or not the parameter was found in this material
+	*/
+	UFUNCTION(BlueprintCallable, Category = "MaterialEditing")
+	bool GetVectorParameterSource(UMaterialInterface* Material, const FName ParameterName, FSoftObjectPath& ParameterSource);
+
+	/**
+	*	Returns the path of the asset where the parameter originated, as well as true/false if it was found
+	*	@param	Material	The material or material instance you want to look up a parameter from
+	*	@param	ParameterName		The parameter name
+	*	@param	ParameterSource		The soft object path of the asset the parameter originates in
+	*	@return	Whether or not the parameter was found in this material
+	*/
+	UFUNCTION(BlueprintCallable, Category = "MaterialEditing")
+	bool GetTextureParameterSource(UMaterialInterface* Material, const FName ParameterName, FSoftObjectPath& ParameterSource);
+
+	/**
+	*	Returns the path of the asset where the parameter originated, as well as true/false if it was found
+	*	@param	Material	The material or material instance you want to look up a parameter from
+	*	@param	ParameterName		The parameter name
+	*	@param	ParameterSource		The soft object path of the asset the parameter originates in
+	*	@return	Whether or not the parameter was found in this material
+	*/
+	UFUNCTION(BlueprintCallable, Category = "MaterialEditing")
+	bool GetStaticSwitchParameterSource(UMaterialInterface* Material, const FName ParameterName, FSoftObjectPath& ParameterSource);
 };
\ No newline at end of file
