Adding default text in case a parent material isn't found

diff --git a/Engine/Source/Editor/MaterialEditor/Private/MaterialInstanceEditor.cpp b/Engine/Source/Editor/MaterialEditor/Private/MaterialInstanceEditor.cpp
--- a/Engine/Source/Editor/MaterialEditor/Private/MaterialInstanceEditor.cpp
+++ b/Engine/Source/Editor/MaterialEditor/Private/MaterialInstanceEditor.cpp
@@ -862,6 +862,13 @@ TSharedRef< SWidget > FMaterialInstanceEditor::GenerateInheritanceMenu()
 	MenuBuilder.BeginSection(ParentName, LOCTEXT("ParentChain", "Parent Chain"));
 	if (bIsFunctionPreviewMaterial)
 	{
+		if (FunctionParentList.Num() == 0)
+		{
+			const FText NoParentText = LOCTEXT("NoParentFound", "No Parent Found");
+			TSharedRef<SWidget> NoParentWidget = SNew(STextBlock)
+				.Text(NoParentText);
+			MenuBuilder.AddWidget(NoParentWidget, FText::GetEmpty());
+		}
 		for (FAssetData FunctionParent : FunctionParentList)
 		{
 			FFormatNamedArguments Args;
@@ -875,6 +882,13 @@ TSharedRef< SWidget > FMaterialInstanceEditor::GenerateInheritanceMenu()
 	}
 	else
 	{
+		if (MaterialParentList.Num() == 0)
+		{
+			const FText NoParentText = LOCTEXT("NoParentFound", "No Parent Found");
+			TSharedRef<SWidget> NoParentWidget = SNew(STextBlock)
+				.Text(NoParentText);
+			MenuBuilder.AddWidget(NoParentWidget, FText::GetEmpty());
+		}
 		for (FAssetData MaterialParent : MaterialParentList)
 		{
 			FFormatNamedArguments Args;
