Fix for substance tutorial not showing up.

diff --git a/Engine/Source/Editor/MaterialEditor/MaterialEditor.Build.cs b/Engine/Source/Editor/MaterialEditor/MaterialEditor.Build.cs
--- a/Engine/Source/Editor/MaterialEditor/MaterialEditor.Build.cs
+++ b/Engine/Source/Editor/MaterialEditor/MaterialEditor.Build.cs
@@ -42,13 +42,13 @@ public class MaterialEditor : ModuleRules
                 "AdvancedPreviewScene",
                 "Projects",
                 "AssetRegistry",
-            }
+				"MainFrame",
+			}
 		);
 
 		DynamicallyLoadedModuleNames.AddRange(
 			new string[] {
 				"AssetTools",
-				"MainFrame",
 				"SceneOutliner",
 				"ClassViewer",
 				"ContentBrowser",
diff --git a/Engine/Source/Editor/MaterialEditor/Private/MaterialEditorModule.cpp b/Engine/Source/Editor/MaterialEditor/Private/MaterialEditorModule.cpp
--- a/Engine/Source/Editor/MaterialEditor/Private/MaterialEditorModule.cpp
+++ b/Engine/Source/Editor/MaterialEditor/Private/MaterialEditorModule.cpp
@@ -22,6 +22,7 @@
 #include "MaterialEditorSettings.h"
 
 #include "ISettingsModule.h"
+#include "Interfaces/IMainFrameModule.h"
 
 const FName MaterialEditorAppIdentifier = FName(TEXT("MaterialEditorApp"));
 const FName MaterialInstanceEditorAppIdentifier = FName(TEXT("MaterialInstanceEditorApp"));
@@ -198,13 +199,8 @@ public:
 
 	void OpenSubstancePluginGetter()
 	{
-		TSharedPtr<SWindow> Window = nullptr;
-		TSharedPtr<SWidget> Widget = FSlateApplication::Get().GetKeyboardFocusedWidget();
-		if (Widget.IsValid())
-		{
-			Window = FSlateApplication::Get().FindWidgetWindow(Widget.ToSharedRef());
-		}
-
+		IMainFrameModule& MainFrameModule = FModuleManager::LoadModuleChecked<IMainFrameModule>(TEXT("MainFrame"));
+		TSharedPtr<SWindow> Window = MainFrameModule.GetParentWindow();
 		IIntroTutorials::Get().LaunchTutorial(TEXT("/Engine/Tutorial/SubEditors/GettingSubstance"), Window);
 	}
 
