Fix Material Editor Bug
The references haven't been released when we quit the Material Editor.

diff --git a/Engine/Source/Editor/MaterialEditor/Private/MaterialInstanceEditor.cpp b/Engine/Source/Editor/MaterialEditor/Private/MaterialInstanceEditor.cpp
--- a/Engine/Source/Editor/MaterialEditor/Private/MaterialInstanceEditor.cpp
+++ b/Engine/Source/Editor/MaterialEditor/Private/MaterialInstanceEditor.cpp
@@ -622,6 +622,8 @@ FMaterialInstanceEditor::~FMaterialInstanceEditor()
 
 	if (MaterialEditorInstance)
 	{
+		MaterialEditorInstance->SourceInstance = nullptr;
+		MaterialEditorInstance->SourceFunction = nullptr;
 		MaterialEditorInstance->MarkPendingKill();
 		MaterialEditorInstance = nullptr;
 	}
