Fix for material updating too frequently
#rnx

diff --git a/Engine/Source/Editor/MaterialEditor/Private/SMaterialLayersFunctionsTree.cpp b/Engine/Source/Editor/MaterialEditor/Private/SMaterialLayersFunctionsTree.cpp
--- a/Engine/Source/Editor/MaterialEditor/Private/SMaterialLayersFunctionsTree.cpp
+++ b/Engine/Source/Editor/MaterialEditor/Private/SMaterialLayersFunctionsTree.cpp
@@ -1662,7 +1662,7 @@ void SMaterialLayersFunctionsInstanceWrapper::Construct(const FArguments& InArgs
 	LayerParameter = NestedTree->FunctionParameter;
 
 	MaterialEditorInstance = InArgs._InMaterialEditorInstance;
-	FEditorSupportDelegates::UpdateUI.AddSP(this, &SMaterialLayersFunctionsInstanceWrapper::Refresh);
+	Refresh();
 
 }
 
@@ -2279,7 +2279,7 @@ void SMaterialLayersFunctionsMaterialWrapper::Construct(const FArguments& InArgs
 
 	LayerParameter = NestedTree->FunctionParameter;
 	MaterialEditorInstance = InArgs._InMaterialEditorInstance;
-	FEditorSupportDelegates::UpdateUI.AddSP(this, &SMaterialLayersFunctionsMaterialWrapper::Refresh);
+	Refresh();
 
 }
 
diff --git a/Engine/Source/Editor/MaterialEditor/Private/SMaterialParametersOverviewWidget.cpp b/Engine/Source/Editor/MaterialEditor/Private/SMaterialParametersOverviewWidget.cpp
--- a/Engine/Source/Editor/MaterialEditor/Private/SMaterialParametersOverviewWidget.cpp
+++ b/Engine/Source/Editor/MaterialEditor/Private/SMaterialParametersOverviewWidget.cpp
@@ -735,7 +735,7 @@ void SMaterialParametersOverviewPanel::Construct(const FArguments& InArgs)
 		.InScrollbar(ExternalScrollbar);
 
 	MaterialEditorInstance = InArgs._InMaterialEditorInstance;
-	FEditorSupportDelegates::UpdateUI.AddSP(this, &SMaterialParametersOverviewPanel::Refresh);
+	Refresh();
 }
 
 void SMaterialParametersOverviewPanel::UpdateEditorInstance(UMaterialEditorPreviewParameters* InMaterialEditorInstance)
