Update GUID, ties to 8182408.

diff --git a/Engine/Source/Editor/MaterialEditor/Private/MaterialInstanceEditor.cpp b/Engine/Source/Editor/MaterialEditor/Private/MaterialInstanceEditor.cpp
--- a/Engine/Source/Editor/MaterialEditor/Private/MaterialInstanceEditor.cpp
+++ b/Engine/Source/Editor/MaterialEditor/Private/MaterialInstanceEditor.cpp
@@ -1275,7 +1275,7 @@ void FMaterialInstanceEditor::DrawSamplerWarningStrings(FCanvas* Canvas, int32&
 							UMaterialExpressionTextureSampleParameter* Expression = BaseMaterial->FindExpressionByGUID<UMaterialExpressionTextureSampleParameter>( TextureParameterValue->ExpressionId );
 
 							FString ErrorMessage;
-							if (!Expression->TextureIsValid(Texture, ErrorMessage))
+							if (Expression && !Expression->TextureIsValid(Texture, ErrorMessage))
 							{
 								Canvas->DrawShadowedString(
 									5,
