Support material expression convert from RVT Sample to RVT Parameter Sample

diff --git a/Engine/Source/Editor/MaterialEditor/Private/MaterialEditor.cpp b/Engine/Source/Editor/MaterialEditor/Private/MaterialEditor.cpp
--- a/Engine/Source/Editor/MaterialEditor/Private/MaterialEditor.cpp
+++ b/Engine/Source/Editor/MaterialEditor/Private/MaterialEditor.cpp
@@ -56,6 +56,8 @@
 #include "Materials/MaterialExpressionTextureBase.h"
 #include "Materials/MaterialExpressionTextureSample.h"
 #include "Materials/MaterialExpressionParticleSubUV.h"
+#include "Materials/MaterialExpressionRuntimeVirtualTextureSample.h"
+#include "Materials/MaterialExpressionRuntimeVirtualTextureSampleParameter.h"
 #include "Materials/MaterialExpressionScalarParameter.h"
 #include "Materials/MaterialExpressionStaticComponentMaskParameter.h"
 #include "Materials/MaterialExpressionTextureSampleParameter.h"
@@ -2846,6 +2848,7 @@ void FMaterialEditor::OnConvertObjects()
 				UMaterialExpressionScalarParameter* ScalarParameterExpression = Cast<UMaterialExpressionScalarParameter>(CurrentSelectedExpression);
 				UMaterialExpressionVectorParameter* VectorParameterExpression = Cast<UMaterialExpressionVectorParameter>(CurrentSelectedExpression);
 				UMaterialExpressionTextureObjectParameter* TextureObjectParameterExpression = Cast<UMaterialExpressionTextureObjectParameter>(CurrentSelectedExpression);
+				UMaterialExpressionRuntimeVirtualTextureSample* RuntimeVirtualTextureSampleExpression = Cast<UMaterialExpressionRuntimeVirtualTextureSample>(CurrentSelectedExpression);
 
 				// Setup the class to convert to
 				UClass* ClassToCreate = NULL;
@@ -2877,6 +2880,10 @@ void FMaterialEditor::OnConvertObjects()
 				{
 					ClassToCreate = UMaterialExpressionTextureSampleParameter2D::StaticClass();
 				}
+				else if (RuntimeVirtualTextureSampleExpression)
+				{
+					ClassToCreate = UMaterialExpressionRuntimeVirtualTextureSampleParameter::StaticClass();
+				}
 				else if (ComponentMaskExpression)
 				{
 					ClassToCreate = UMaterialExpressionStaticComponentMaskParameter::StaticClass();
@@ -2946,6 +2953,15 @@ void FMaterialEditor::OnConvertObjects()
 							NewTextureExpr->MipValueMode = TextureSampleExpression->MipValueMode;
 							NewGraphNode->ReconstructNode();
 						}
+						else if (RuntimeVirtualTextureSampleExpression)
+						{
+							bNeedsRefresh = true;
+							UMaterialExpressionRuntimeVirtualTextureSampleParameter* NewRuntimeVirtualTextureExpression = CastChecked<UMaterialExpressionRuntimeVirtualTextureSampleParameter>(NewExpression);
+							NewRuntimeVirtualTextureExpression->VirtualTexture = RuntimeVirtualTextureSampleExpression->VirtualTexture;
+							NewRuntimeVirtualTextureExpression->MaterialType = RuntimeVirtualTextureSampleExpression->MaterialType;
+							NewRuntimeVirtualTextureExpression->MipValueMode = RuntimeVirtualTextureSampleExpression->MipValueMode;
+							NewGraphNode->ReconstructNode();
+						}
 						else if (ComponentMaskExpression)
 						{
 							bNeedsRefresh = true;
diff --git a/Engine/Source/Editor/UnrealEd/Private/MaterialGraphNode.cpp b/Engine/Source/Editor/UnrealEd/Private/MaterialGraphNode.cpp
--- a/Engine/Source/Editor/UnrealEd/Private/MaterialGraphNode.cpp
+++ b/Engine/Source/Editor/UnrealEd/Private/MaterialGraphNode.cpp
@@ -13,14 +13,15 @@
 #include "Materials/MaterialExpressionConstant2Vector.h"
 #include "Materials/MaterialExpressionConstant3Vector.h"
 #include "Materials/MaterialExpressionConstant4Vector.h"
+#include "Materials/MaterialExpressionCustomOutput.h"
 #include "Materials/MaterialExpressionFontSample.h"
 #include "Materials/MaterialExpressionFunctionInput.h"
 #include "Materials/MaterialExpressionFunctionOutput.h"
 #include "Materials/MaterialExpressionMaterialAttributeLayers.h"
+#include "Materials/MaterialExpressionRuntimeVirtualTextureSample.h"
 #include "Materials/MaterialExpressionScalarParameter.h"
 #include "Materials/MaterialExpressionStaticBool.h"
 #include "Materials/MaterialExpressionStaticBoolParameter.h"
-#include "Materials/MaterialExpressionCustomOutput.h"
 #include "Materials/MaterialExpressionTextureBase.h"
 #include "Materials/MaterialExpressionTextureCoordinate.h"
 #include "Materials/MaterialExpressionTextureSample.h"
@@ -355,6 +356,7 @@ void UMaterialGraphNode::GetContextMenuActions(const FGraphNodeContextMenuBuilde
 				|| MaterialExpression->IsA(UMaterialExpressionConstant3Vector::StaticClass())
 				|| MaterialExpression->IsA(UMaterialExpressionConstant4Vector::StaticClass())
 				|| (MaterialExpression->IsA(UMaterialExpressionTextureSample::StaticClass()) && !MaterialExpression->HasAParameterName())
+				|| (MaterialExpression->IsA(UMaterialExpressionRuntimeVirtualTextureSample::StaticClass()) && !MaterialExpression->HasAParameterName())
 				|| MaterialExpression->IsA(UMaterialExpressionTextureObject::StaticClass())
 				|| MaterialExpression->IsA(UMaterialExpressionComponentMask::StaticClass())
 				|| MaterialExpression->IsA(UMaterialExpressionMaterialFunctionCall::StaticClass()))
diff --git a/Engine/Source/Editor/UnrealEd/Private/MaterialGraphNode_Base.cpp b/Engine/Source/Editor/UnrealEd/Private/MaterialGraphNode_Base.cpp
--- a/Engine/Source/Editor/UnrealEd/Private/MaterialGraphNode_Base.cpp
+++ b/Engine/Source/Editor/UnrealEd/Private/MaterialGraphNode_Base.cpp
@@ -112,20 +112,36 @@ void UMaterialGraphNode_Base::ReplaceNode(UMaterialGraphNode_Base* OldNode)
 	// Copy Outputs from old node
 	for (int32 PinIndex = 0; PinIndex < OldOutputPins.Num(); PinIndex++)
 	{
-		// If we can't find an equivalent output in this node, just use the first
-		// The user will have to fix up any issues from the mismatch
-		int32 FoundPinIndex = 0;
-
 		// Try to find an equivalent output in this node
-		for (int32 NewPinIndex = 0; NewPinIndex < NewOutputPins.Num(); NewPinIndex++)
+		int32 FoundPinIndex = -1;
+		if (FoundPinIndex == -1)
 		{
-			if (OldOutputPins[PinIndex]->PinType == NewOutputPins[NewPinIndex]->PinType)
+			// First check names
+			for (int32 NewPinIndex = 0; NewPinIndex < NewOutputPins.Num(); NewPinIndex++)
 			{
-				FoundPinIndex = NewPinIndex;
-				break;
+				if (OldOutputPins[PinIndex]->PinName == NewOutputPins[NewPinIndex]->PinName)
+				{
+					FoundPinIndex = NewPinIndex;
+					break;
+				}
 			}
 		}
-		
+		if (FoundPinIndex == -1)
+		{
+			// Now check types
+			for (int32 NewPinIndex = 0; NewPinIndex < NewOutputPins.Num(); NewPinIndex++)
+			{
+				if (OldOutputPins[PinIndex]->PinType == NewOutputPins[NewPinIndex]->PinType)
+				{
+					FoundPinIndex = NewPinIndex;
+					break;
+				}
+			}
+		}
+
+		// If we can't find an equivalent output in this node, just use the first
+		// The user will have to fix up any issues from the mismatch
+		FoundPinIndex = FMath::Max(FoundPinIndex, 0);
 		if (FoundPinIndex < NewOutputPins.Num())
 		{
 			ModifyAndCopyPersistentPinData(*NewOutputPins[FoundPinIndex], *OldOutputPins[PinIndex]);
diff --git a/Engine/Source/Runtime/Engine/Classes/Materials/MaterialExpressionRuntimeVirtualTextureSample.h b/Engine/Source/Runtime/Engine/Classes/Materials/MaterialExpressionRuntimeVirtualTextureSample.h
--- a/Engine/Source/Runtime/Engine/Classes/Materials/MaterialExpressionRuntimeVirtualTextureSample.h
+++ b/Engine/Source/Runtime/Engine/Classes/Materials/MaterialExpressionRuntimeVirtualTextureSample.h
@@ -30,7 +30,7 @@ enum ERuntimeVirtualTextureMipValueMode
 
 /** Material expression for sampling from a runtime virtual texture. */
 UCLASS(collapsecategories, hidecategories=Object)
-class UMaterialExpressionRuntimeVirtualTextureSample : public UMaterialExpression
+class ENGINE_API UMaterialExpressionRuntimeVirtualTextureSample : public UMaterialExpression
 {
 	GENERATED_UCLASS_BODY()
 
