Adding scrollbar to material parameters panel

diff --git a/Engine/Source/Editor/MaterialEditor/Private/SMaterialParametersOverviewWidget.cpp b/Engine/Source/Editor/MaterialEditor/Private/SMaterialParametersOverviewWidget.cpp
--- a/Engine/Source/Editor/MaterialEditor/Private/SMaterialParametersOverviewWidget.cpp
+++ b/Engine/Source/Editor/MaterialEditor/Private/SMaterialParametersOverviewWidget.cpp
@@ -376,6 +376,7 @@ void SMaterialParametersOverviewTree::Construct(const FArguments& InArgs)
 		.OnGenerateRow(this, &SMaterialParametersOverviewTree::OnGenerateRowMaterialLayersFunctionsTreeView)
 		.OnGetChildren(this, &SMaterialParametersOverviewTree::OnGetChildrenMaterialLayersFunctionsTreeView)
 		.OnExpansionChanged(this, &SMaterialParametersOverviewTree::OnExpansionChanged)
+		.ExternalScrollbar(InArgs._InScrollbar)
 	);
 }
 
@@ -621,7 +622,6 @@ void SMaterialParametersOverviewPanel::Refresh()
 			[
 				SNew(SVerticalBox)
 				+SVerticalBox::Slot()
-				.AutoHeight()
 				[
 					SNew(SBorder)
 					.BorderImage(this, &SMaterialParametersOverviewPanel::GetBackgroundImage)
@@ -646,22 +646,34 @@ void SMaterialParametersOverviewPanel::Refresh()
 						]
 						+ SVerticalBox::Slot()
 						.Padding(FMargin(3.0f, 2.0f, 3.0f, 3.0f))
-						.AutoHeight()
 						[
 							SNew(SBorder)
 							.BorderImage(FEditorStyle::GetBrush("DetailsView.CategoryTop"))
 							[
-								SNew(SWidgetSwitcher)
-								.WidgetIndex(this, &SMaterialParametersOverviewPanel::GetPanelIndex)
-								+SWidgetSwitcher::Slot()
+								SNew(SOverlay)
+								+ SOverlay::Slot()
 								[
-									SNew(STextBlock)
-									.Text(LOCTEXT("AddParams", "Add parameters to see them here."))
+									SNew(SWidgetSwitcher)
+									.WidgetIndex(this, &SMaterialParametersOverviewPanel::GetPanelIndex)
+									+SWidgetSwitcher::Slot()
+									[
+										SNew(STextBlock)
+										.Text(LOCTEXT("AddParams", "Add parameters to see them here."))
+									]
+									+SWidgetSwitcher::Slot()
+									[
+										NestedTree.ToSharedRef()
+									]	
 								]
-								+SWidgetSwitcher::Slot()
+								+ SOverlay::Slot()
+								.HAlign(HAlign_Right)
 								[
-									NestedTree.ToSharedRef()
-								]	
+									SNew(SBox)
+									.WidthOverride(16.0f)
+									[
+										ExternalScrollbar.ToSharedRef()
+									]
+								]
 							]
 						]
 					]
@@ -713,9 +725,13 @@ void SMaterialParametersOverviewPanel::Refresh()
 
 void SMaterialParametersOverviewPanel::Construct(const FArguments& InArgs)
 {
+	ExternalScrollbar = SNew(SScrollBar)
+		.AlwaysShowScrollbar(true);
+
 	NestedTree = SNew(SMaterialParametersOverviewTree)
 		.InMaterialEditorInstance(InArgs._InMaterialEditorInstance)
-		.InOwner(SharedThis(this));
+		.InOwner(SharedThis(this))
+		.InScrollbar(ExternalScrollbar);
 
 	MaterialEditorInstance = InArgs._InMaterialEditorInstance;
 	FEditorSupportDelegates::UpdateUI.AddSP(this, &SMaterialParametersOverviewPanel::Refresh);
diff --git a/Engine/Source/Editor/MaterialEditor/Private/SMaterialParametersOverviewWidget.h b/Engine/Source/Editor/MaterialEditor/Private/SMaterialParametersOverviewWidget.h
--- a/Engine/Source/Editor/MaterialEditor/Private/SMaterialParametersOverviewWidget.h
+++ b/Engine/Source/Editor/MaterialEditor/Private/SMaterialParametersOverviewWidget.h
@@ -88,6 +88,8 @@ private:
 
 	/** The tree contained in this item */
 	TSharedPtr<class SMaterialParametersOverviewTree> NestedTree;
+
+	TSharedPtr<class SScrollBar> ExternalScrollbar;
 };
 
 // ********* SMaterialParametersOverviewTree *******
@@ -100,10 +102,12 @@ public:
 	SLATE_BEGIN_ARGS(SMaterialParametersOverviewTree)
 		: _InMaterialEditorInstance(nullptr)
 		, _InOwner(nullptr)
+		, _InScrollbar()
 	{}
 
 	SLATE_ARGUMENT(UMaterialEditorPreviewParameters*, InMaterialEditorInstance)
 	SLATE_ARGUMENT(TSharedPtr<SMaterialParametersOverviewPanel>, InOwner)
+	SLATE_ARGUMENT(TSharedPtr<SScrollBar>, InScrollbar)
 	SLATE_END_ARGS()
 
 	/** Constructs this widget with InArgs */
