Curve atlas parameters now filter out curves that aren't in the current atlas.
#rnx

diff --git a/Engine/Source/Editor/MaterialEditor/Private/MaterialEditorDetailCustomization.cpp b/Engine/Source/Editor/MaterialEditor/Private/MaterialEditorDetailCustomization.cpp
--- a/Engine/Source/Editor/MaterialEditor/Private/MaterialEditorDetailCustomization.cpp
+++ b/Engine/Source/Editor/MaterialEditor/Private/MaterialEditorDetailCustomization.cpp
@@ -36,6 +36,8 @@
 #include "Curves/CurveLinearColor.h"
 #include "IPropertyUtilities.h"
 #include "Engine/Texture.h"
+#include "Materials/MaterialExpressionCurveAtlasRowParameter.h"
+#include "Curves/CurveLinearColorAtlas.h"
 
 
 #define LOCTEXT_NAMESPACE "MaterialEditor"
@@ -449,6 +451,7 @@ void FMaterialEditorParameterDetails::CreateScalarAtlasPositionParameterValueWid
 				.NewAssetFactories(TArray<UFactory*>())
 				.DisplayThumbnail(true)
 				.ThumbnailPool(PropertyUtilities.Pin()->GetThumbnailPool())
+				.OnShouldFilterAsset(FOnShouldFilterAsset::CreateStatic(&FMaterialPropertyHelpers::OnShouldFilterCurveAsset, AtlasParameter->AtlasData.Atlas))
 				.OnShouldSetAsset(FOnShouldSetAsset::CreateStatic(&FMaterialPropertyHelpers::OnShouldSetCurveAsset, AtlasParameter->AtlasData.Atlas))
 				.OnObjectChanged(FOnSetObject::CreateStatic(&FMaterialPropertyHelpers::SetPositionFromCurveAsset, AtlasParameter->AtlasData.Atlas, AtlasParameter, ParameterProperty, (UObject*)MaterialEditorInstance))
 				.DisplayCompactSize(true)
@@ -507,6 +510,7 @@ void FMaterialExpressionParameterDetails::CustomizeDetails( IDetailLayoutBuilder
 	ScalarParameterObjects.Reset();
 
 	const FName MaterialExpressionCategory = TEXT("MaterialExpression");
+	IDetailCategoryBuilder& ExpressionCategory = DetailLayout.EditCategory(MaterialExpressionCategory);
 
 	for (const auto& WeakObjectPtr : Objects)
 	{
@@ -549,6 +553,51 @@ void FMaterialExpressionParameterDetails::CustomizeDetails( IDetailLayoutBuilder
 				}));
 			}
 
+			if (ScalarParameter->IsUsedAsAtlasPosition())
+			{
+				UMaterialExpressionCurveAtlasRowParameter* ColorCurveParameter = Cast<UMaterialExpressionCurveAtlasRowParameter>(ScalarParameter);
+				TSharedPtr<IPropertyHandle> CurveHandle = DetailLayout.GetProperty("Curve", UMaterialExpressionCurveAtlasRowParameter::StaticClass());
+				if (CurveHandle.IsValid() && CurveHandle->IsValidHandle())
+				{
+					const FName AtlasCategoryName = TEXT("MaterialExpressionCurveAtlasRowParameter");
+					IDetailCategoryBuilder& AtlasCategory = DetailLayout.EditCategory(AtlasCategoryName);
+					TSoftObjectPtr<UCurveLinearColorAtlas> Atlas = TSoftObjectPtr<UCurveLinearColorAtlas>(FSoftObjectPath(ColorCurveParameter->Atlas->GetPathName()));
+					CurveHandle->MarkHiddenByCustomization();
+					AtlasCategory.AddCustomRow(CurveHandle->GetPropertyDisplayName())
+						.NameContent()
+						[
+							SNew(STextBlock)
+							.Text(CurveHandle->GetPropertyDisplayName())
+							.Font(IDetailLayoutBuilder::GetDetailFont())
+						]
+						.ValueContent()
+						.HAlign(HAlign_Fill)
+						.MaxDesiredWidth(400.0f)
+						[
+							SNew(SObjectPropertyEntryBox)
+							.PropertyHandle(CurveHandle)
+							.AllowedClass(UCurveLinearColor::StaticClass())
+							.NewAssetFactories(TArray<UFactory*>())
+							.DisplayThumbnail(true)
+							.ThumbnailPool(DetailLayout.GetThumbnailPool())
+							.OnShouldFilterAsset(FOnShouldFilterAsset::CreateStatic(&FMaterialPropertyHelpers::OnShouldFilterCurveAsset, Atlas))
+							.OnShouldSetAsset(FOnShouldSetAsset::CreateStatic(&FMaterialPropertyHelpers::OnShouldSetCurveAsset, Atlas))
+						];
+
+
+						TSharedPtr<IPropertyHandle> AtlasHandle = DetailLayout.GetProperty("Atlas", UMaterialExpressionCurveAtlasRowParameter::StaticClass());
+						if (AtlasHandle.IsValid() && AtlasHandle->IsValidHandle())
+						{
+							// Rebuild the layout when the bUseCustomPrimitiveData property changes
+							AtlasHandle->SetOnPropertyValueChanged(FSimpleDelegate::CreateLambda([&DetailLayout]()
+							{
+								DetailLayout.ForceRefreshDetails();
+							}));
+						}
+
+				}
+			}
+
 			if (ScalarParameter->bUseCustomPrimitiveData)
 			{
 				DetailLayout.HideCategory(TEXT("MaterialExpressionScalarParameter"));
@@ -805,7 +854,7 @@ void FMaterialExpressionParameterDetails::CustomizeDetails( IDetailLayoutBuilder
 	
 	Category.AddProperty("ParameterName");
 
-	IDetailCategoryBuilder& ExpressionCategory = DetailLayout.EditCategory(MaterialExpressionCategory);
+
 
 	// Get a handle to the property we are about to edit
 	GroupPropertyHandle = DetailLayout.GetProperty( "Group" );
diff --git a/Engine/Source/Editor/MaterialEditor/Private/MaterialEditorInstanceDetailCustomization.cpp b/Engine/Source/Editor/MaterialEditor/Private/MaterialEditorInstanceDetailCustomization.cpp
--- a/Engine/Source/Editor/MaterialEditor/Private/MaterialEditorInstanceDetailCustomization.cpp
+++ b/Engine/Source/Editor/MaterialEditor/Private/MaterialEditorInstanceDetailCustomization.cpp
@@ -747,6 +747,7 @@ void FMaterialInstanceParameterDetails::CreateScalarAtlasPositionParameterValueW
 				.NewAssetFactories(TArray<UFactory*>())
 				.DisplayThumbnail(true)
 				.ThumbnailPool(PropertyUtilities.Pin()->GetThumbnailPool())
+				.OnShouldFilterAsset(FOnShouldFilterAsset::CreateStatic(&FMaterialPropertyHelpers::OnShouldFilterCurveAsset, AtlasParameter->AtlasData.Atlas))
 				.OnShouldSetAsset(FOnShouldSetAsset::CreateStatic(&FMaterialPropertyHelpers::OnShouldSetCurveAsset, AtlasParameter->AtlasData.Atlas))
 				.OnObjectChanged(FOnSetObject::CreateStatic(&FMaterialPropertyHelpers::SetPositionFromCurveAsset, AtlasParameter->AtlasData.Atlas, AtlasParameter, ParameterProperty, (UObject*)MaterialEditorInstance))
 				.DisplayCompactSize(true)
diff --git a/Engine/Source/Editor/MaterialEditor/Private/MaterialPropertyHelpers.cpp b/Engine/Source/Editor/MaterialEditor/Private/MaterialPropertyHelpers.cpp
--- a/Engine/Source/Editor/MaterialEditor/Private/MaterialPropertyHelpers.cpp
+++ b/Engine/Source/Editor/MaterialEditor/Private/MaterialPropertyHelpers.cpp
@@ -1235,13 +1235,25 @@ TSharedRef<SWidget> FMaterialPropertyHelpers::MakeStackReorderHandle(TSharedPtr<
 bool FMaterialPropertyHelpers::OnShouldSetCurveAsset(const FAssetData& AssetData, TSoftObjectPtr<class UCurveLinearColorAtlas> InAtlas)
 {
 	UCurveLinearColorAtlas* Atlas = Cast<UCurveLinearColorAtlas>(InAtlas.Get());
-	UCurveLinearColor* Curve = Cast<UCurveLinearColor>(AssetData.GetAsset());
-	if (!Atlas || !Curve)
+	if (!Atlas)
 	{
 		return false;
 	}
-	int32 Index = Atlas->GradientCurves.Find(Curve);
-	return Index != INDEX_NONE;
+
+	for (UCurveLinearColor* GradientCurve : Atlas->GradientCurves)
+	{
+		if (GradientCurve->GetOutermost()->GetPathName() == AssetData.PackageName.ToString())
+		{
+			return true;
+		}
+	}
+
+	return false;
+}
+
+bool FMaterialPropertyHelpers::OnShouldFilterCurveAsset(const FAssetData& AssetData, TSoftObjectPtr<class UCurveLinearColorAtlas> InAtlas)
+{
+	return !OnShouldSetCurveAsset(AssetData, InAtlas);
 }
 
 void FMaterialPropertyHelpers::SetPositionFromCurveAsset(const FAssetData& AssetData, TSoftObjectPtr<class UCurveLinearColorAtlas> InAtlas, UDEditorScalarParameterValue* InParameter, TSharedPtr<IPropertyHandle> PropertyHandle, UObject* MaterialEditorInstance)
diff --git a/Engine/Source/Editor/MaterialEditor/Private/SMaterialLayersFunctionsTree.cpp b/Engine/Source/Editor/MaterialEditor/Private/SMaterialLayersFunctionsTree.cpp
--- a/Engine/Source/Editor/MaterialEditor/Private/SMaterialLayersFunctionsTree.cpp
+++ b/Engine/Source/Editor/MaterialEditor/Private/SMaterialLayersFunctionsTree.cpp
@@ -702,6 +702,7 @@ void SMaterialLayersFunctionsInstanceTreeItem::Construct(const FArguments& InArg
 					.DisplayThumbnail(true)
 					.ThumbnailPool(InArgs._InTree->GetTreeThumbnailPool())
 					.OnShouldSetAsset(FOnShouldSetAsset::CreateStatic(&FMaterialPropertyHelpers::OnShouldSetCurveAsset, ScalarParam->AtlasData.Atlas))
+					.OnShouldFilterAsset(FOnShouldFilterAsset::CreateStatic(&FMaterialPropertyHelpers::OnShouldFilterCurveAsset, ScalarParam->AtlasData.Atlas))
 					.OnObjectChanged(FOnSetObject::CreateStatic(&FMaterialPropertyHelpers::SetPositionFromCurveAsset, ScalarParam->AtlasData.Atlas, ScalarParam, StackParameterData->ParameterHandle, (UObject*)MaterialEditorInstance))
 					.DisplayCompactSize(true)
 				];
@@ -1914,6 +1915,7 @@ public:
 						.DisplayThumbnail(true)
 						.ThumbnailPool(InArgs._InTree->GetTreeThumbnailPool())
 						.OnShouldSetAsset(FOnShouldSetAsset::CreateStatic(&FMaterialPropertyHelpers::OnShouldSetCurveAsset, ScalarParam->AtlasData.Atlas))
+						.OnShouldFilterAsset(FOnShouldFilterAsset::CreateStatic(&FMaterialPropertyHelpers::OnShouldFilterCurveAsset, ScalarParam->AtlasData.Atlas))
 						.OnObjectChanged(FOnSetObject::CreateStatic(&FMaterialPropertyHelpers::SetPositionFromCurveAsset, ScalarParam->AtlasData.Atlas, ScalarParam, StackParameterData->ParameterHandle, (UObject*)MaterialEditorInstance))
 						.DisplayCompactSize(true)
 					];
diff --git a/Engine/Source/Editor/MaterialEditor/Private/SMaterialParametersOverviewWidget.cpp b/Engine/Source/Editor/MaterialEditor/Private/SMaterialParametersOverviewWidget.cpp
--- a/Engine/Source/Editor/MaterialEditor/Private/SMaterialParametersOverviewWidget.cpp
+++ b/Engine/Source/Editor/MaterialEditor/Private/SMaterialParametersOverviewWidget.cpp
@@ -166,6 +166,7 @@ void SMaterialParametersOverviewTreeItem::Construct(const FArguments& InArgs, co
 					.NewAssetFactories(TArray<UFactory*>())
 					.DisplayThumbnail(true)
 					.ThumbnailPool(Tree.Pin()->GetTreeThumbnailPool())
+					.OnShouldFilterAsset(FOnShouldFilterAsset::CreateStatic(&FMaterialPropertyHelpers::OnShouldFilterCurveAsset, ScalarParam->AtlasData.Atlas))
 					.OnShouldSetAsset(FOnShouldSetAsset::CreateStatic(&FMaterialPropertyHelpers::OnShouldSetCurveAsset, ScalarParam->AtlasData.Atlas))
 					.OnObjectChanged(FOnSetObject::CreateStatic(&FMaterialPropertyHelpers::SetPositionFromCurveAsset, ScalarParam->AtlasData.Atlas, ScalarParam, StackParameterData->ParameterHandle, (UObject*)MaterialEditorInstance))
 					.DisplayCompactSize(true)
diff --git a/Engine/Source/Editor/MaterialEditor/Public/MaterialPropertyHelpers.h b/Engine/Source/Editor/MaterialEditor/Public/MaterialPropertyHelpers.h
--- a/Engine/Source/Editor/MaterialEditor/Public/MaterialPropertyHelpers.h
+++ b/Engine/Source/Editor/MaterialEditor/Public/MaterialPropertyHelpers.h
@@ -193,6 +193,7 @@ public:
 	static TSharedRef<SWidget> MakeStackReorderHandle(TSharedPtr<SMaterialLayersFunctionsInstanceTreeItem> InOwningStack);
 
 	static bool OnShouldSetCurveAsset(const FAssetData& AssetData, TSoftObjectPtr<class UCurveLinearColorAtlas> InAtlas);
+	static bool OnShouldFilterCurveAsset(const FAssetData& AssetData, TSoftObjectPtr<class UCurveLinearColorAtlas> InAtlas);
 	static void SetPositionFromCurveAsset(const FAssetData& AssetData, TSoftObjectPtr<class UCurveLinearColorAtlas> InAtlas, class UDEditorScalarParameterValue* InParameter, TSharedPtr<IPropertyHandle> PropertyHandle, UObject* MaterialEditorInstance);
 
 	static void ResetCurveToDefault(TSharedPtr<IPropertyHandle> PropertyHandle, class UDEditorParameterValue* Parameter, UMaterialEditorInstanceConstant* MaterialEditorInstance);
diff --git a/Engine/Source/Runtime/Engine/Private/Materials/MaterialExpressions.cpp b/Engine/Source/Runtime/Engine/Private/Materials/MaterialExpressions.cpp
--- a/Engine/Source/Runtime/Engine/Private/Materials/MaterialExpressions.cpp
+++ b/Engine/Source/Runtime/Engine/Private/Materials/MaterialExpressions.cpp
@@ -16425,7 +16425,16 @@ int32 UMaterialExpressionCurveAtlasRowParameter::Compile(class FMaterialCompiler
 			return CompilerError(Compiler, TEXT("The curve is not contained within the atlas."));
 		}
 	}
-	return CompilerError(Compiler, TEXT("Either the curve or atlas is not currently set."));
+	else if (Atlas)
+	{
+		return CompilerError(Compiler, TEXT("The curve is not currently set."));
+	}
+	else if (Curve)
+	{
+		return CompilerError(Compiler, TEXT("The atlas is not currently set."));
+	}
+
+	return CompilerError(Compiler, TEXT("The curve and atlas are not currently set."));
 }
 
 void UMaterialExpressionCurveAtlasRowParameter::PostEditChangeProperty(FPropertyChangedEvent& PropertyChangedEvent)
