Removing stale handle being used when changing a property in a layer.
#rnx

diff --git a/Engine/Source/Editor/MaterialEditor/Private/MaterialPropertyHelpers.cpp b/Engine/Source/Editor/MaterialEditor/Private/MaterialPropertyHelpers.cpp
--- a/Engine/Source/Editor/MaterialEditor/Private/MaterialPropertyHelpers.cpp
+++ b/Engine/Source/Editor/MaterialEditor/Private/MaterialPropertyHelpers.cpp
@@ -1278,8 +1278,6 @@ void FMaterialPropertyHelpers::SetPositionFromCurveAsset(const FAssetData& Asset
 		const FScopedTransaction Transaction(LOCTEXT("SetScalarAtlasPositionValue", "Set Scalar Atlas Position Value"));
 		InParameter->Modify();
 
-		PropertyHandle->NotifyPreChange();
-		
 		InParameter->AtlasData.Curve = TSoftObjectPtr<UCurveLinearColor>(FSoftObjectPath(Curve->GetPathName()));
 		InParameter->ParameterValue = NewValue;
 		UMaterialEditorInstanceConstant* MaterialInstanceEditor = Cast<UMaterialEditorInstanceConstant>(MaterialEditorInstance);
@@ -1287,8 +1285,6 @@ void FMaterialPropertyHelpers::SetPositionFromCurveAsset(const FAssetData& Asset
 		{
 			MaterialInstanceEditor->CopyToSourceInstance();
 		}
-
-		PropertyHandle->NotifyPostChange();
 	}
 }
 
