Added an option to the Material Editor user preferences to set the default size for the material editor preview viewport when viewing a material/material instance with "User Interface" domain

diff --git a/Engine/Source/Editor/MaterialEditor/Private/SMaterialEditorViewport.cpp b/Engine/Source/Editor/MaterialEditor/Private/SMaterialEditorViewport.cpp
--- a/Engine/Source/Editor/MaterialEditor/Private/SMaterialEditorViewport.cpp
+++ b/Engine/Source/Editor/MaterialEditor/Private/SMaterialEditorViewport.cpp
@@ -30,6 +30,7 @@
 #include "AssetViewerSettings.h"
 #include "Engine/PostProcessVolume.h"
 #include "Widgets/Input/SCheckBox.h"
+#include "MaterialEditorSettings.h"
 
 #define LOCTEXT_NAMESPACE "MaterialEditor"
 
@@ -807,7 +808,7 @@ private:
 
 void SMaterialEditorUIPreviewZoomer::Construct( const FArguments& InArgs, UMaterialInterface* InPreviewMaterial )
 {
-	PreviewBrush = MakeShareable( new FSlateMaterialBrush( *InPreviewMaterial, FVector2D(250,250) ) );
+	PreviewBrush = MakeShareable( new FSlateMaterialBrush( *InPreviewMaterial, GetDefault<UMaterialEditorSettings>()->GetPreviewViewportStartingSize() ) );
 
 	ChildSlot
 	[
@@ -987,7 +988,7 @@ void SMaterialEditorUIPreviewViewport::Construct( const FArguments& InArgs, UMat
 		]
 	];
 
-	PreviewSize = FIntPoint(250,250);
+	PreviewSize = GetDefault<UMaterialEditorSettings>()->GetPreviewViewportStartingSize();
 	PreviewZoomer->SetPreviewSize( FVector2D(PreviewSize) );
 }
 
diff --git a/Engine/Source/Editor/MaterialEditor/Public/MaterialEditorSettings.h b/Engine/Source/Editor/MaterialEditor/Public/MaterialEditorSettings.h
--- a/Engine/Source/Editor/MaterialEditor/Public/MaterialEditorSettings.h
+++ b/Engine/Source/Editor/MaterialEditor/Public/MaterialEditorSettings.h
@@ -17,4 +17,19 @@ public:
 	*/
 	UPROPERTY(config, EditAnywhere, Category = "Offline Shader Compilers", meta = (DisplayName = "Mali Offline Compiler"))
 	FFilePath MaliOfflineCompilerPath;
+
+protected:
+	// The width (in pixels) of the preview viewport when a material editor is first opened
+	UPROPERTY(config, EditAnywhere, meta=(ClampMin=1, ClampMax=4096), Category="User Interface Domain")
+	int32 DefaultPreviewWidth = 250;
+
+	// The height (in pixels) of the preview viewport when a material editor is first opened
+	UPROPERTY(config, EditAnywhere, meta=(ClampMin=1, ClampMax=4096), Category="User Interface Domain")
+	int32 DefaultPreviewHeight = 250;
+
+public:
+	FIntPoint GetPreviewViewportStartingSize() const
+	{
+		return FIntPoint(DefaultPreviewWidth, DefaultPreviewHeight);
+	}
 };
\ No newline at end of file
