Remove extra separator in material editor toolbar

diff --git a/Engine/Source/Editor/MaterialEditor/Private/MaterialEditor.cpp b/Engine/Source/Editor/MaterialEditor/Private/MaterialEditor.cpp
--- a/Engine/Source/Editor/MaterialEditor/Private/MaterialEditor.cpp
+++ b/Engine/Source/Editor/MaterialEditor/Private/MaterialEditor.cpp
@@ -1163,8 +1163,6 @@ void FMaterialEditor::FillToolbar(FToolBarBuilder& ToolbarBuilder)
 	}
 	ToolbarBuilder.EndSection();
 
-	ToolbarBuilder.AddSeparator();
-
 	ToolbarBuilder.BeginSection("Stats");
 	{
 		ToolbarBuilder.AddToolBarButton(FMaterialEditorCommands::Get().ToggleMaterialStats);
