Removed OpenGL SM4 menu option.

diff --git a/Engine/Source/Editor/MaterialEditor/Private/MaterialStats.cpp b/Engine/Source/Editor/MaterialEditor/Private/MaterialStats.cpp
--- a/Engine/Source/Editor/MaterialEditor/Private/MaterialStats.cpp
+++ b/Engine/Source/Editor/MaterialEditor/Private/MaterialStats.cpp
@@ -417,7 +417,6 @@ void FMaterialStats::BuildShaderPlatformDB()
 
 	// OpenGL
 	AddShaderPlatform(EPlatformCategoryType::Desktop, SP_OPENGL_SM5, TEXT("OpenGL SM5"), false, true, TEXT("Desktop, OpenGL, Shader Model 5"));
-	AddShaderPlatform(EPlatformCategoryType::Desktop, SP_OPENGL_SM4, TEXT("OpenGL SM4"), false, true, TEXT("Desktop, OpenGL, Shader Model 4"));
 
 	// Android
 	AddShaderPlatform(EPlatformCategoryType::Android, SP_OPENGL_ES3_1_ANDROID, TEXT("Android GLES 3.1"), true, true, TEXT("Android, OpenGLES 3.1"));
