Adding refresh to parameter toggling in material layers so you can update nested children.

diff --git a/Engine/Source/Editor/MaterialEditor/Private/SMaterialLayersFunctionsTree.cpp b/Engine/Source/Editor/MaterialEditor/Private/SMaterialLayersFunctionsTree.cpp
--- a/Engine/Source/Editor/MaterialEditor/Private/SMaterialLayersFunctionsTree.cpp
+++ b/Engine/Source/Editor/MaterialEditor/Private/SMaterialLayersFunctionsTree.cpp
@@ -303,6 +303,12 @@ bool SMaterialLayersFunctionsInstanceTree::IsOverriddenExpression(class UDEditor
 	return FMaterialPropertyHelpers::IsOverriddenExpression(Parameter) && FunctionInstance->LayerStates[InIndex];
 }
 
+void  SMaterialLayersFunctionsInstanceTreeItem::OnOverrideParameter(bool NewValue, class UDEditorParameterValue* Parameter)
+{
+	FMaterialPropertyHelpers::OnOverrideParameter(NewValue, Parameter, MaterialEditorInstance);
+	Tree->GetWrapper()->Refresh();
+}
+
 void SMaterialLayersFunctionsInstanceTreeItem::Construct(const FArguments& InArgs, const TSharedRef<STableViewBase>& InOwnerTableView)
 {
 	StackParameterData = InArgs._StackParameterData;
@@ -646,7 +652,7 @@ void SMaterialLayersFunctionsInstanceTreeItem::Construct(const FArguments& InArg
 		Row
 			.DisplayName(NameOverride)
 			.OverrideResetToDefault(ResetOverride)
-			.EditCondition(IsParamEnabled, FOnBooleanValueChanged::CreateStatic(&FMaterialPropertyHelpers::OnOverrideParameter, StackParameterData->Parameter, MaterialEditorInstance));
+			.EditCondition(IsParamEnabled, FOnBooleanValueChanged::CreateSP(this, &SMaterialLayersFunctionsInstanceTreeItem::OnOverrideParameter, StackParameterData->Parameter));
 
 		if (VectorParam && VectorParam->bIsUsedAsChannelMask)
 		{
diff --git a/Engine/Source/Editor/MaterialEditor/Private/SMaterialLayersFunctionsTree.h b/Engine/Source/Editor/MaterialEditor/Private/SMaterialLayersFunctionsTree.h
--- a/Engine/Source/Editor/MaterialEditor/Private/SMaterialLayersFunctionsTree.h
+++ b/Engine/Source/Editor/MaterialEditor/Private/SMaterialLayersFunctionsTree.h
@@ -73,6 +73,7 @@ public:
 	}
 
 	FReply OnLayerDrop(const FDragDropEvent& DragDropEvent);
+	void OnOverrideParameter(bool NewValue, class UDEditorParameterValue* Parameter);
 	/**
 	* Construct the widget
 	*
