Fix for parameter overview panel not always updating when a parameter was converted back to a constant

diff --git a/Engine/Source/Editor/MaterialEditor/Private/MaterialEditor.cpp b/Engine/Source/Editor/MaterialEditor/Private/MaterialEditor.cpp
--- a/Engine/Source/Editor/MaterialEditor/Private/MaterialEditor.cpp
+++ b/Engine/Source/Editor/MaterialEditor/Private/MaterialEditor.cpp
@@ -3012,6 +3012,11 @@ void FMaterialEditor::OnConvertObjects()
 		{
 			GraphEditor->SetNodeSelection(*NodeIter, true);
 		}
+
+		if (MaterialEditorInstance != nullptr)
+		{
+			MaterialParametersOverviewWidget->UpdateEditorInstance(MaterialEditorInstance);
+		}
 	}
 }
 
diff --git a/Engine/Source/Editor/MaterialEditor/Private/SMaterialParametersOverviewWidget.cpp b/Engine/Source/Editor/MaterialEditor/Private/SMaterialParametersOverviewWidget.cpp
--- a/Engine/Source/Editor/MaterialEditor/Private/SMaterialParametersOverviewWidget.cpp
+++ b/Engine/Source/Editor/MaterialEditor/Private/SMaterialParametersOverviewWidget.cpp
@@ -807,69 +807,70 @@ void SMaterialParametersOverviewPanel::Refresh()
 		OnChildButtonClicked = FOnClicked::CreateStatic(&FMaterialPropertyHelpers::OnClickedSaveNewMaterialInstance, ImplicitConv<UMaterialInterface*>(MaterialEditorInstance->OriginalMaterial), ImplicitConv<UObject*>(MaterialEditorInstance));
 	}
 
+	if (NestedTree->HasAnyParameters())
 	{
 		this->ChildSlot
 			[
 				SNew(SVerticalBox)
-				+SVerticalBox::Slot()
-				[
-					SNew(SBorder)
-					.BorderImage(this, &SMaterialParametersOverviewPanel::GetBackgroundImage)
-					.Padding(FMargin(4.0f))
-					[
-						SNew(SVerticalBox)
-						+ SVerticalBox::Slot()
-						.AutoHeight()
-						[
-							SAssignNew(HeaderBox, SHorizontalBox)
-							+ SHorizontalBox::Slot()
-							.Padding(FMargin(3.0f, 1.0f))
-							.HAlign(HAlign_Left)
-							.AutoWidth()
-							.VAlign(VAlign_Center)
-							[
-								SNew(STextBlock)
-								.Text(LOCTEXT("ParameterDefaults", "Parameter Defaults"))
-								.Font(FEditorStyle::GetFontStyle("DetailsView.CategoryFontStyle"))
-								.ShadowOffset(FVector2D(1.0f, 1.0f))
-							]
-						]
-						+ SVerticalBox::Slot()
-						.Padding(FMargin(3.0f, 2.0f, 3.0f, 3.0f))
-						[
-							SNew(SBorder)
-							.BorderImage(FEditorStyle::GetBrush("DetailsView.CategoryTop"))
-							[
-								SNew(SHorizontalBox)
-								+ SHorizontalBox::Slot()
-								.HAlign(HAlign_Fill)
-								[
-									SNew(SWidgetSwitcher)
-									.WidgetIndex(this, &SMaterialParametersOverviewPanel::GetPanelIndex)
-									+SWidgetSwitcher::Slot()
-									[
-										SNew(STextBlock)
-										.Text(LOCTEXT("AddParams", "Add parameters to see them here."))
-									]
-									+SWidgetSwitcher::Slot()
-									[
-										NestedTree.ToSharedRef()
-									]	
-								]
-								+ SHorizontalBox::Slot()
-								.HAlign(HAlign_Right)
-								.AutoWidth()
-								[
-									SNew(SBox)
-									.WidthOverride(16.0f)
-									[
-										ExternalScrollbar.ToSharedRef()
-									]
-								]
-							]
-						]
-					]
-				]
+				+ SVerticalBox::Slot()
+			[
+				SNew(SBorder)
+				.BorderImage(this, &SMaterialParametersOverviewPanel::GetBackgroundImage)
+			.Padding(FMargin(4.0f))
+			[
+				SNew(SVerticalBox)
+				+ SVerticalBox::Slot()
+			.AutoHeight()
+			[
+				SAssignNew(HeaderBox, SHorizontalBox)
+				+ SHorizontalBox::Slot()
+			.Padding(FMargin(3.0f, 1.0f))
+			.HAlign(HAlign_Left)
+			.AutoWidth()
+			.VAlign(VAlign_Center)
+			[
+				SNew(STextBlock)
+				.Text(LOCTEXT("ParameterDefaults", "Parameter Defaults"))
+			.Font(FEditorStyle::GetFontStyle("DetailsView.CategoryFontStyle"))
+			.ShadowOffset(FVector2D(1.0f, 1.0f))
+			]
+			]
+		+ SVerticalBox::Slot()
+			.Padding(FMargin(3.0f, 2.0f, 3.0f, 3.0f))
+			[
+				SNew(SBorder)
+				.BorderImage(FEditorStyle::GetBrush("DetailsView.CategoryTop"))
+			[
+				SNew(SHorizontalBox)
+				+ SHorizontalBox::Slot()
+			.HAlign(HAlign_Fill)
+			[
+				SNew(SWidgetSwitcher)
+				.WidgetIndex(this, &SMaterialParametersOverviewPanel::GetPanelIndex)
+			+ SWidgetSwitcher::Slot()
+			[
+				SNew(STextBlock)
+				.Text(LOCTEXT("AddParams", "Add parameters to see them here."))
+			]
+		+ SWidgetSwitcher::Slot()
+			[
+				NestedTree.ToSharedRef()
+			]
+			]
+		+ SHorizontalBox::Slot()
+			.HAlign(HAlign_Right)
+			.AutoWidth()
+			[
+				SNew(SBox)
+				.WidthOverride(16.0f)
+			[
+				ExternalScrollbar.ToSharedRef()
+			]
+			]
+			]
+			]
+			]
+			]
 			];
 
 		HeaderBox->AddSlot()
@@ -877,7 +878,7 @@ void SMaterialParametersOverviewPanel::Refresh()
 			[
 				SNullWidget::NullWidget
 			];
-		
+
 		if (NestedTree->HasAnyParameters())
 		{
 			HeaderBox->AddSlot()
@@ -886,31 +887,45 @@ void SMaterialParametersOverviewPanel::Refresh()
 				[
 					SNew(SButton)
 					.ButtonStyle(FEditorStyle::Get(), "FlatButton.DarkGrey")
-					.HAlign(HAlign_Center)
-					.OnClicked(OnChildButtonClicked)
-					.ToolTipText(LOCTEXT("SaveToChildInstance", "Save To Child Instance"))
-					.Content()
-					[
-						SNew(SHorizontalBox)
-						+ SHorizontalBox::Slot()
-						.AutoWidth()
-						[
-							SNew(STextBlock)
-							.Font(FEditorStyle::Get().GetFontStyle("FontAwesome.10"))
-							.TextStyle(FEditorStyle::Get(), "NormalText.Important")
-							.Text(FText::FromString(FString(TEXT("\xf0c7 \xf149"))) /*fa-filter*/)
-						]
-						+ SHorizontalBox::Slot()
-						.AutoWidth()
-						[
-							SNew(STextBlock)
-							.TextStyle(FEditorStyle::Get(), "NormalText.Important")
-							.Text(FText::FromString(FString(TEXT(" Save Child"))) /*fa-filter*/)
-						]
-					]
+				.HAlign(HAlign_Center)
+				.OnClicked(OnChildButtonClicked)
+				.ToolTipText(LOCTEXT("SaveToChildInstance", "Save To Child Instance"))
+				.Content()
+				[
+					SNew(SHorizontalBox)
+					+ SHorizontalBox::Slot()
+				.AutoWidth()
+				[
+					SNew(STextBlock)
+					.Font(FEditorStyle::Get().GetFontStyle("FontAwesome.10"))
+				.TextStyle(FEditorStyle::Get(), "NormalText.Important")
+				.Text(FText::FromString(FString(TEXT("\xf0c7 \xf149"))) /*fa-filter*/)
+				]
+			+ SHorizontalBox::Slot()
+				.AutoWidth()
+				[
+					SNew(STextBlock)
+					.TextStyle(FEditorStyle::Get(), "NormalText.Important")
+				.Text(FText::FromString(FString(TEXT(" Save Child"))) /*fa-filter*/)
+				]
+				]
 				];
 		}
 	}
+	else
+	{
+		this->ChildSlot
+			[
+				SNew(SBorder)
+				.BorderImage(FEditorStyle::GetBrush("MaterialInstanceEditor.StackBody"))
+				.Padding(FMargin(4.0f))
+				[
+					SNew(STextBlock)
+					.Text(LOCTEXT("AddLayerParameterPrompt", "Connect a parameter to see it here."))
+				]
+			];
+	}
+
 	
 }
 
