Fix duplicate localization text key
#rnx

diff --git a/Engine/Source/Editor/MaterialEditor/Private/SMaterialParametersOverviewWidget.cpp b/Engine/Source/Editor/MaterialEditor/Private/SMaterialParametersOverviewWidget.cpp
--- a/Engine/Source/Editor/MaterialEditor/Private/SMaterialParametersOverviewWidget.cpp
+++ b/Engine/Source/Editor/MaterialEditor/Private/SMaterialParametersOverviewWidget.cpp
@@ -921,7 +921,7 @@ void SMaterialParametersOverviewPanel::Refresh()
 				.Padding(FMargin(4.0f))
 				[
 					SNew(STextBlock)
-					.Text(LOCTEXT("AddLayerParameterPrompt", "Connect a parameter to see it here."))
+					.Text(LOCTEXT("ConnectMaterialParametersToFillList", "Connect a parameter to see it here."))
 				]
 			];
 	}
