Parameter to constrain menu height, merged from CL 9763699

diff --git a/Engine/Source/Developer/ToolMenus/Private/ToolMenu.cpp b/Engine/Source/Developer/ToolMenus/Private/ToolMenu.cpp
--- a/Engine/Source/Developer/ToolMenus/Private/ToolMenu.cpp
+++ b/Engine/Source/Developer/ToolMenus/Private/ToolMenu.cpp
@@ -52,6 +52,7 @@ void UToolMenu::InitGeneratedCopy(const UToolMenu* Source, const FName InMenuNam
 
 	SubMenuParent = Source->SubMenuParent;
 	SubMenuSourceEntryName = Source->SubMenuSourceEntryName;
+	MaxHeight = Source->MaxHeight;
 	bExtendersEnabled = Source->bExtendersEnabled;
 
 	MaxHeight = Source->MaxHeight;
diff --git a/Engine/Source/Editor/Kismet/Private/SBlueprintRevisionMenu.cpp b/Engine/Source/Editor/Kismet/Private/SBlueprintRevisionMenu.cpp
--- a/Engine/Source/Editor/Kismet/Private/SBlueprintRevisionMenu.cpp
+++ b/Engine/Source/Editor/Kismet/Private/SBlueprintRevisionMenu.cpp
@@ -242,12 +242,7 @@ void SBlueprintRevisionMenu::OnSourceControlQueryComplete(const FSourceControlOp
 	MenuBuilder.EndSection();
 	MenuBox->AddSlot() 
 	[
-		SNew(SBorder)
-		.BorderImage(FEditorStyle::GetBrush("Menu.Background"))
-		.Content()
-		[
-			MenuBuilder.MakeWidget(nullptr, 500)
-		]
+		MenuBuilder.MakeWidget(nullptr, 500)
 	];
 
 	SourceControlQueryOp.Reset();
diff --git a/Engine/Source/Editor/MaterialEditor/Private/MaterialInstanceEditor.cpp b/Engine/Source/Editor/MaterialEditor/Private/MaterialInstanceEditor.cpp
--- a/Engine/Source/Editor/MaterialEditor/Private/MaterialInstanceEditor.cpp
+++ b/Engine/Source/Editor/MaterialEditor/Private/MaterialInstanceEditor.cpp
@@ -935,6 +935,7 @@ void FMaterialInstanceEditor::GenerateInheritanceMenu(UToolMenu* Menu)
 	RebuildInheritanceList();
 	Menu->bShouldCloseWindowAfterMenuSelection = true;
 	Menu->bSearchable = true;
+	Menu->SetMaxHeight(500);
 	const FName ParentName = TEXT("ParentChain");
 	FToolMenuSection& Section = Menu->AddSection(ParentName, LOCTEXT("ParentChain", "Parent Chain"));
 	if (bIsFunctionPreviewMaterial)
