[NO NON-PII COMMIT MESSAGE]

diff --git a/Engine/Source/Editor/MaterialEditor/Private/SMaterialParametersOverviewWidget.cpp b/Engine/Source/Editor/MaterialEditor/Private/SMaterialParametersOverviewWidget.cpp
--- a/Engine/Source/Editor/MaterialEditor/Private/SMaterialParametersOverviewWidget.cpp
+++ b/Engine/Source/Editor/MaterialEditor/Private/SMaterialParametersOverviewWidget.cpp
@@ -650,8 +650,9 @@ void SMaterialParametersOverviewPanel::Refresh()
 							SNew(SBorder)
 							.BorderImage(FEditorStyle::GetBrush("DetailsView.CategoryTop"))
 							[
-								SNew(SOverlay)
-								+ SOverlay::Slot()
+								SNew(SHorizontalBox)
+								+ SHorizontalBox::Slot()
+								.HAlign(HAlign_Fill)
 								[
 									SNew(SWidgetSwitcher)
 									.WidgetIndex(this, &SMaterialParametersOverviewPanel::GetPanelIndex)
@@ -665,8 +666,9 @@ void SMaterialParametersOverviewPanel::Refresh()
 										NestedTree.ToSharedRef()
 									]	
 								]
-								+ SOverlay::Slot()
+								+ SHorizontalBox::Slot()
 								.HAlign(HAlign_Right)
+								.AutoWidth()
 								[
 									SNew(SBox)
 									.WidthOverride(16.0f)
@@ -725,8 +727,7 @@ void SMaterialParametersOverviewPanel::Refresh()
 
 void SMaterialParametersOverviewPanel::Construct(const FArguments& InArgs)
 {
-	ExternalScrollbar = SNew(SScrollBar)
-		.AlwaysShowScrollbar(true);
+	ExternalScrollbar = SNew(SScrollBar);
 
 	NestedTree = SNew(SMaterialParametersOverviewTree)
 		.InMaterialEditorInstance(InArgs._InMaterialEditorInstance)
@@ -735,7 +736,6 @@ void SMaterialParametersOverviewPanel::Construct(const FArguments& InArgs)
 
 	MaterialEditorInstance = InArgs._InMaterialEditorInstance;
 	FEditorSupportDelegates::UpdateUI.AddSP(this, &SMaterialParametersOverviewPanel::Refresh);
-
 }
 
 void SMaterialParametersOverviewPanel::UpdateEditorInstance(UMaterialEditorPreviewParameters* InMaterialEditorInstance)
