You can now use the advanced property Channel Names to label the R, G, B and A outputs of texture and vector parameters. These names will be visible in the material instance editor for instances of this material, as well as next to the pins for this expression.

diff --git a/Engine/Source/Editor/MaterialEditor/Private/MaterialEditor.cpp b/Engine/Source/Editor/MaterialEditor/Private/MaterialEditor.cpp
--- a/Engine/Source/Editor/MaterialEditor/Private/MaterialEditor.cpp
+++ b/Engine/Source/Editor/MaterialEditor/Private/MaterialEditor.cpp
@@ -4387,6 +4387,10 @@ void FMaterialEditor::NotifyPostChange( const FPropertyChangedEvent& PropertyCha
 				{
 					Material->UpdateExpressionDynamicParameters(SelectedNode->MaterialExpression);
 				}
+				else if (PropertyThatChanged->IsA<UTextProperty>())
+				{
+					// Do nothing to the expression if we are just changing the label
+				}
 				else
 				{
 					Material->PropagateExpressionParameterChanges(SelectedNode->MaterialExpression);
@@ -4402,7 +4406,9 @@ void FMaterialEditor::NotifyPostChange( const FPropertyChangedEvent& PropertyCha
 		const FName PropertyName = PropertyChangedEvent.MemberProperty ? PropertyChangedEvent.MemberProperty->GetFName() : NAME_None;
 		if (PropertyName != GET_MEMBER_NAME_CHECKED(UMaterialExpressionComment, Text)
 		&& PropertyName != GET_MEMBER_NAME_CHECKED(UMaterialExpressionComment, CommentColor)
-		&& PropertyName != GET_MEMBER_NAME_CHECKED(UMaterialExpression, Desc))
+		&& PropertyName != GET_MEMBER_NAME_CHECKED(UMaterialExpression, Desc)
+		&& PropertyName != GET_MEMBER_NAME_CHECKED(UMaterialExpressionVectorParameter, ChannelNames)
+		&& PropertyName != GET_MEMBER_NAME_CHECKED(UMaterialExpressionTextureSampleParameter, ChannelNames))
 		{
 			// Update the current preview material.
 			UpdatePreviewMaterial();
diff --git a/Engine/Source/Editor/MaterialEditor/Private/MaterialEditorDetailCustomization.cpp b/Engine/Source/Editor/MaterialEditor/Private/MaterialEditorDetailCustomization.cpp
--- a/Engine/Source/Editor/MaterialEditor/Private/MaterialEditorDetailCustomization.cpp
+++ b/Engine/Source/Editor/MaterialEditor/Private/MaterialEditorDetailCustomization.cpp
@@ -35,401 +35,9 @@
 #include "Curves/CurveLinearColor.h"
 #include "IPropertyUtilities.h"
 
-#define LOCTEXT_NAMESPACE "MaterialEditor"
-
-TSharedRef<IDetailCustomization> FMaterialEditorParameterDetails::MakeInstance(UMaterialEditorPreviewParameters* Material)
-{
-	return MakeShareable(new FMaterialEditorParameterDetails(Material));
-}
-
-FMaterialEditorParameterDetails::FMaterialEditorParameterDetails(UMaterialEditorPreviewParameters* Material)
-	: MaterialEditorInstance(Material)
-{
-}
-
-
-TOptional<float> FMaterialEditorParameterDetails::OnGetValue(TSharedRef<IPropertyHandle> PropertyHandle)
-{
-	float Value = 0.0f;
-	if (PropertyHandle->GetValue(Value) == FPropertyAccess::Success)
-	{
-		return TOptional<float>(Value);
-	}
-
-	// Value couldn't be accessed. Return an unset value
-	return TOptional<float>();
-}
-
-void FMaterialEditorParameterDetails::OnValueCommitted(float NewValue, ETextCommit::Type CommitType, TSharedRef<IPropertyHandle> PropertyHandle)
-{
-	// Try setting as float, if that fails then set as int
-	ensure(PropertyHandle->SetValue(NewValue) == FPropertyAccess::Success);
-}
-
-void FMaterialEditorParameterDetails::CustomizeDetails(IDetailLayoutBuilder& DetailLayout)
-{
-	PropertyUtilities = DetailLayout.GetPropertyUtilities();
-	// Create a new category for a custom layout for the MIC parameters at the very top
-	FName GroupsCategoryName = TEXT("ParameterGroups");
-	IDetailCategoryBuilder& GroupsCategory = DetailLayout.EditCategory(GroupsCategoryName, LOCTEXT("MICParamGroupsTitle", "Parameter Groups"));
-	TSharedRef<IPropertyHandle> ParameterGroupsProperty = DetailLayout.GetProperty("ParameterGroups");
-
-	CreateGroupsWidget(ParameterGroupsProperty, GroupsCategory);
-
-	DetailLayout.HideProperty("ParameterGroups");
-
-	TArray<UObject*> ExternalObjects;
-	ExternalObjects.Add(MaterialEditorInstance->OriginalMaterial);
-}
-
-
-void FMaterialEditorParameterDetails::CreateGroupsWidget(TSharedRef<IPropertyHandle> ParameterGroupsProperty, IDetailCategoryBuilder& GroupsCategory)
-{
-	check(MaterialEditorInstance);
-	bool bShowSaveButtons = false;
-	for (int32 GroupIdx = 0; GroupIdx < MaterialEditorInstance->ParameterGroups.Num(); ++GroupIdx)
-	{
-		FEditorParameterGroup& ParameterGroup = MaterialEditorInstance->ParameterGroups[GroupIdx];
-		if (ParameterGroup.GroupAssociation == EMaterialParameterAssociation::GlobalParameter
-			&& ParameterGroup.GroupName != FMaterialPropertyHelpers::LayerParamName)
-		{
-			IDetailGroup& DetailGroup = GroupsCategory.AddGroup(ParameterGroup.GroupName, FText::FromName(ParameterGroup.GroupName), false, true);
-			CreateSingleGroupWidget(ParameterGroup, ParameterGroupsProperty->GetChildHandle(GroupIdx), DetailGroup);
-			bShowSaveButtons = true;
-		}
-	}
-	if (bShowSaveButtons)
-	{
-		FOnClicked ChildButtonClicked;
-		UObject* LocalEditorInstance = MaterialEditorInstance;
-		if (MaterialEditorInstance->OriginalMaterial)
-		{
-			ChildButtonClicked = FOnClicked::CreateStatic(FMaterialPropertyHelpers::OnClickedSaveNewMaterialInstance,
-				ImplicitConv<UMaterialInterface*>(MaterialEditorInstance->OriginalMaterial), LocalEditorInstance);
-		}
-		if (MaterialEditorInstance->OriginalFunction)
-		{
-			ChildButtonClicked = FOnClicked::CreateStatic(FMaterialPropertyHelpers::OnClickedSaveNewFunctionInstance,
-				ImplicitConv<UMaterialFunctionInterface*>(MaterialEditorInstance->OriginalFunction), ImplicitConv<UMaterialInterface*>(MaterialEditorInstance->PreviewMaterial), LocalEditorInstance);
-		}
-
-		FDetailWidgetRow& SaveInstanceRow = GroupsCategory.AddCustomRow(LOCTEXT("SaveInstances", "Save Instances"));
-		SaveInstanceRow.ValueContent()
-			.HAlign(HAlign_Fill)
-			[
-				SNew(SHorizontalBox)
-				+ SHorizontalBox::Slot()
-				.FillWidth(1.0f)
-				[
-					SNullWidget::NullWidget
-				]
-				+ SHorizontalBox::Slot()
-				.AutoWidth()
-				.Padding(2.0f)
-				[
-					SNew(SButton)
-					.ButtonStyle(FEditorStyle::Get(), "FlatButton.Dark")
-					.HAlign(HAlign_Center)
-					.OnClicked(ChildButtonClicked)
-					.ToolTipText(LOCTEXT("SaveToChildInstance", "Save To Child Instance"))
-					.Content()
-					[
-						SNew(SHorizontalBox)
-						+ SHorizontalBox::Slot()
-						.AutoWidth()
-						[
-							SNew(STextBlock)
-							.Font(FEditorStyle::Get().GetFontStyle("FontAwesome.10"))
-							.TextStyle(FEditorStyle::Get(), "NormalText.Important")
-							.Text(FText::FromString(FString(TEXT("\xf0c7 \xf149"))) /*fa-filter*/)
-						]
-						+ SHorizontalBox::Slot()
-						.AutoWidth()
-						[
-							SNew(STextBlock)
-							.TextStyle(FEditorStyle::Get(), "NormalText.Important")
-							.Text(FText::FromString(FString(TEXT(" Save Child"))) /*fa-filter*/)
-						]
-					]
-				]
-			];
-	}
-	else
-	{
-		FDetailWidgetRow& NoParameters = GroupsCategory.AddCustomRow(LOCTEXT("NoParameters", "No Parameters"));
-		NoParameters.WholeRowContent()
-			[
-				SNew(STextBlock)
-				.Text(LOCTEXT("AddGlobalParameterPrompt", "Add a global parameter to see it here. \nMaterial Attribute Layers parameters are found in the Layer Parameter Preview panel."))
-				.AutoWrapText(true)
-			];
-	}
-}
-
-
-void FMaterialEditorParameterDetails::CreateSingleGroupWidget(FEditorParameterGroup& ParameterGroup, TSharedPtr<IPropertyHandle> ParameterGroupProperty, IDetailGroup& DetailGroup)
-{
-	TSharedPtr<IPropertyHandle> ParametersArrayProperty = ParameterGroupProperty->GetChildHandle("Parameters");
-
-	// Create a custom widget for each parameter in the group
-	for (int32 ParamIdx = 0; ParamIdx < ParameterGroup.Parameters.Num(); ++ParamIdx)
-	{
-		TSharedPtr<IPropertyHandle> ParameterProperty = ParametersArrayProperty->GetChildHandle(ParamIdx);
-
-		FString ParameterName = ParameterGroup.Parameters[ParamIdx]->ParameterInfo.Name.ToString();
-
-		UDEditorParameterValue* Parameter = ParameterGroup.Parameters[ParamIdx];
-		UDEditorFontParameterValue* FontParam = Cast<UDEditorFontParameterValue>(Parameter);
-		UDEditorScalarParameterValue* ScalarParam = Cast<UDEditorScalarParameterValue>(Parameter);
-		UDEditorStaticComponentMaskParameterValue* CompMaskParam = Cast<UDEditorStaticComponentMaskParameterValue>(Parameter);
-		UDEditorStaticSwitchParameterValue* SwitchParam = Cast<UDEditorStaticSwitchParameterValue>(Parameter);
-		UDEditorTextureParameterValue* TextureParam = Cast<UDEditorTextureParameterValue>(Parameter);
-		UDEditorVectorParameterValue* VectorParam = Cast<UDEditorVectorParameterValue>(Parameter);
-		UDEditorMaterialLayersParameterValue* LayersParam = Cast<UDEditorMaterialLayersParameterValue>(Parameter);
-
-		if (VectorParam && VectorParam->bIsUsedAsChannelMask)
-		{
-			CreateVectorChannelMaskParameterValueWidget(Parameter, ParameterProperty, DetailGroup);
-		}
-		if (ScalarParam && ScalarParam->AtlasData.bIsUsedAsAtlasPosition)
-		{
-			CreateScalarAtlasPositionParameterValueWidget(Parameter, ParameterProperty, DetailGroup);
-		}
-		else if (ScalarParam || SwitchParam || TextureParam || VectorParam || FontParam)
-		{
-			if (ScalarParam && ScalarParam->SliderMax > ScalarParam->SliderMin)
-			{
-				TSharedPtr<IPropertyHandle> ParameterValueProperty = ParameterProperty->GetChildHandle("ParameterValue");
-				ParameterValueProperty->SetInstanceMetaData("UIMin", FString::Printf(TEXT("%f"), ScalarParam->SliderMin));
-				ParameterValueProperty->SetInstanceMetaData("UIMax", FString::Printf(TEXT("%f"), ScalarParam->SliderMax));
-			}
-
-			CreateParameterValueWidget(Parameter, ParameterProperty, DetailGroup);
-		}
-		else if (LayersParam)
-		{
-		}
-		else if (CompMaskParam)
-		{
-			CreateMaskParameterValueWidget(Parameter, ParameterProperty, DetailGroup);
-		}
-		else
-		{
-			// Unsupported parameter type
-			check(false);
-		}
-	}
-}
-
-
-void FMaterialEditorParameterDetails::CreateParameterValueWidget(UDEditorParameterValue* Parameter, TSharedPtr<IPropertyHandle> ParameterProperty, IDetailGroup& DetailGroup)
-{
-	TSharedPtr<IPropertyHandle> ParameterValueProperty = ParameterProperty->GetChildHandle("ParameterValue");
-
-	if (ParameterValueProperty->IsValidHandle())
-	{
-
-		IDetailPropertyRow& PropertyRow = DetailGroup.AddPropertyRow(ParameterValueProperty.ToSharedRef());
-		PropertyRow
-			.DisplayName(FText::FromName(Parameter->ParameterInfo.Name))
-			.ToolTip(GetParameterExpressionDescription(Parameter));
-	}
-}
-
-void FMaterialEditorParameterDetails::CreateMaskParameterValueWidget(UDEditorParameterValue* Parameter, TSharedPtr<IPropertyHandle> ParameterProperty, IDetailGroup& DetailGroup)
-{
-	TSharedPtr<IPropertyHandle> ParameterValueProperty = ParameterProperty->GetChildHandle("ParameterValue");
-	TSharedPtr<IPropertyHandle> RMaskProperty = ParameterValueProperty->GetChildHandle("R");
-	TSharedPtr<IPropertyHandle> GMaskProperty = ParameterValueProperty->GetChildHandle("G");
-	TSharedPtr<IPropertyHandle> BMaskProperty = ParameterValueProperty->GetChildHandle("B");
-	TSharedPtr<IPropertyHandle> AMaskProperty = ParameterValueProperty->GetChildHandle("A");
-
-	if (ParameterValueProperty->IsValidHandle())
-	{
-		IDetailPropertyRow& PropertyRow = DetailGroup.AddPropertyRow(ParameterValueProperty.ToSharedRef());
-		const FText ParameterName = FText::FromName(Parameter->ParameterInfo.Name);
-
-		FDetailWidgetRow& CustomWidget = PropertyRow.CustomWidget();
-		CustomWidget
-			.FilterString(ParameterName)
-			.NameContent()
-			[
-				SNew(STextBlock)
-				.Text(ParameterName)
-				.ToolTipText(GetParameterExpressionDescription(Parameter))
-				.Font(FEditorStyle::GetFontStyle(TEXT("PropertyWindow.NormalFont")))
-			]
-			.ValueContent()
-			.MaxDesiredWidth(200.0f)
-			[
-				SNew(SHorizontalBox)
-				+ SHorizontalBox::Slot()
-				.FillWidth(1.0f)
-				[
-					SNew(SHorizontalBox)
-					+ SHorizontalBox::Slot()
-					.HAlign(HAlign_Left)
-					.AutoWidth()
-					[
-						RMaskProperty->CreatePropertyNameWidget(FText::GetEmpty(), FText::GetEmpty(), false)
-					]
-					+ SHorizontalBox::Slot()
-					.HAlign(HAlign_Left)
-					.AutoWidth()
-					[
-						RMaskProperty->CreatePropertyValueWidget()
-					]
-					+ SHorizontalBox::Slot()
-					.HAlign(HAlign_Left)
-					.Padding(FMargin(10.0f, 0.0f, 0.0f, 0.0f))
-					.AutoWidth()
-					[
-						GMaskProperty->CreatePropertyNameWidget(FText::GetEmpty(), FText::GetEmpty(), false)
-					]
-					+ SHorizontalBox::Slot()
-					.HAlign(HAlign_Left)
-					.AutoWidth()
-					[
-						GMaskProperty->CreatePropertyValueWidget()
-					]
-					+ SHorizontalBox::Slot()
-					.HAlign(HAlign_Left)
-					.Padding(FMargin(10.0f, 0.0f, 0.0f, 0.0f))
-					.AutoWidth()
-					[
-						BMaskProperty->CreatePropertyNameWidget(FText::GetEmpty(), FText::GetEmpty(), false)
-					]
-					+ SHorizontalBox::Slot()
-					.HAlign(HAlign_Left)
-					.AutoWidth()
-					[
-						BMaskProperty->CreatePropertyValueWidget()
-					]
-					+ SHorizontalBox::Slot()
-					.HAlign(HAlign_Left)
-					.Padding(FMargin(10.0f, 0.0f, 0.0f, 0.0f))
-					.AutoWidth()
-					[
-						AMaskProperty->CreatePropertyNameWidget(FText::GetEmpty(), FText::GetEmpty(), false)
-					]
-					+ SHorizontalBox::Slot()
-					.HAlign(HAlign_Left)
-					.AutoWidth()
-					[
-						AMaskProperty->CreatePropertyValueWidget()
-					]
-				]
-			];
-	}
-}
-
-void FMaterialEditorParameterDetails::CreateVectorChannelMaskParameterValueWidget(UDEditorParameterValue* Parameter, TSharedPtr<IPropertyHandle> ParameterProperty, IDetailGroup& DetailGroup)
-{
-	TSharedPtr<IPropertyHandle> ParameterValueProperty = ParameterProperty->GetChildHandle("ParameterValue");
-
-	if (ParameterValueProperty->IsValidHandle())
-	{
-		TAttribute<bool> IsParamEnabled = TAttribute<bool>::Create(TAttribute<bool>::FGetter::CreateStatic(&FMaterialPropertyHelpers::IsOverriddenExpression, Parameter));
-
-		IDetailPropertyRow& PropertyRow = DetailGroup.AddPropertyRow(ParameterValueProperty.ToSharedRef());
-		const FText ParameterName = FText::FromName(Parameter->ParameterInfo.Name);
-
-		FOnGetPropertyComboBoxStrings GetMaskStrings = FOnGetPropertyComboBoxStrings::CreateStatic(&FMaterialPropertyHelpers::GetVectorChannelMaskComboBoxStrings);
-		FOnGetPropertyComboBoxValue GetMaskValue = FOnGetPropertyComboBoxValue::CreateStatic(&FMaterialPropertyHelpers::GetVectorChannelMaskValue, Parameter);
-		FOnPropertyComboBoxValueSelected SetMaskValue = FOnPropertyComboBoxValueSelected::CreateStatic(&FMaterialPropertyHelpers::SetVectorChannelMaskValue, ParameterValueProperty, Parameter, (UObject*)MaterialEditorInstance);
-
-		FDetailWidgetRow& CustomWidget = PropertyRow.CustomWidget();
-		CustomWidget
-		.FilterString(ParameterName)
-		.NameContent()
-		[
-			SNew(STextBlock)
-			.Text(ParameterName)
-			.ToolTipText(FMaterialPropertyHelpers::GetParameterExpressionDescription(Parameter, MaterialEditorInstance))
-			.Font(FEditorStyle::GetFontStyle(TEXT("PropertyWindow.NormalFont")))
-		]
-		.ValueContent()
-		.MaxDesiredWidth(200.0f)
-		[
-			SNew(SHorizontalBox)
-			+ SHorizontalBox::Slot()
-			.FillWidth(1.0f)
-			[
-				SNew(SHorizontalBox)
-				+ SHorizontalBox::Slot()
-				.HAlign(HAlign_Left)
-				.AutoWidth()
-				[
-					PropertyCustomizationHelpers::MakePropertyComboBox(ParameterValueProperty, GetMaskStrings, GetMaskValue, SetMaskValue)
-				]
-			]
-		];
-	}
-}
-
-void FMaterialEditorParameterDetails::CreateScalarAtlasPositionParameterValueWidget(class UDEditorParameterValue* Parameter, TSharedPtr<IPropertyHandle> ParameterProperty, IDetailGroup& DetailGroup)
-{
-	TSharedPtr<IPropertyHandle> ParameterValueProperty = ParameterProperty->GetChildHandle("ParameterValue");
-
-	if (ParameterValueProperty->IsValidHandle())
-	{
-		TAttribute<bool> IsParamEnabled = TAttribute<bool>::Create(TAttribute<bool>::FGetter::CreateStatic(&FMaterialPropertyHelpers::IsOverriddenExpression, Parameter));
-
-		IDetailPropertyRow& PropertyRow = DetailGroup.AddPropertyRow(ParameterValueProperty.ToSharedRef());
-	
-		const FText ParameterName = FText::FromName(Parameter->ParameterInfo.Name);
-		UDEditorScalarParameterValue* AtlasParameter = Cast<UDEditorScalarParameterValue>(Parameter);
 
-		FDetailWidgetRow& CustomWidget = PropertyRow.CustomWidget();
-		CustomWidget
-			.FilterString(ParameterName)
-			.NameContent()
-			[
-				SNew(STextBlock)
-				.Text(ParameterName)
-				.ToolTipText(FMaterialPropertyHelpers::GetParameterExpressionDescription(Parameter, MaterialEditorInstance))
-				.Font(FEditorStyle::GetFontStyle(TEXT("PropertyWindow.NormalFont")))
-			]
-			.ValueContent()
-			.HAlign(HAlign_Fill)
-			.MaxDesiredWidth(400.0f)
-			[
-				SNew(SObjectPropertyEntryBox)
-				.ObjectPath(this, &FMaterialEditorParameterDetails::GetCurvePath, AtlasParameter)
-				.AllowedClass(UCurveLinearColor::StaticClass())
-				.NewAssetFactories(TArray<UFactory*>())
-				.DisplayThumbnail(true)
-				.ThumbnailPool(PropertyUtilities.Pin()->GetThumbnailPool())
-				.OnShouldSetAsset(FOnShouldSetAsset::CreateStatic(&FMaterialPropertyHelpers::OnShouldSetCurveAsset, AtlasParameter->AtlasData.Atlas))
-				.OnObjectChanged(FOnSetObject::CreateStatic(&FMaterialPropertyHelpers::SetPositionFromCurveAsset, AtlasParameter->AtlasData.Atlas, AtlasParameter, ParameterProperty, (UObject*)MaterialEditorInstance))
-				.DisplayCompactSize(true)
-			];
-	}
-}
-
-FString FMaterialEditorParameterDetails::GetCurvePath(UDEditorScalarParameterValue* Parameter) const
-{
-	FString Path = Parameter->AtlasData.Curve->GetPathName();
-	return Path;
-}
+#define LOCTEXT_NAMESPACE "MaterialEditor"
 
-FText FMaterialEditorParameterDetails::GetParameterExpressionDescription(UDEditorParameterValue* Parameter) const
-{
-	if (MaterialEditorInstance && MaterialEditorInstance->OriginalMaterial)
-	{
-		UMaterial* BaseMaterial = MaterialEditorInstance->OriginalMaterial->GetMaterial();
-		if (BaseMaterial)
-		{
-			UMaterialExpression* MaterialExpression = BaseMaterial->FindExpressionByGUID<UMaterialExpression>(Parameter->ExpressionId);
-
-			if (MaterialExpression)
-			{
-				return FText::FromString(MaterialExpression->Desc);
-			}
-		}
-	}
-	return FText::GetEmpty();
-}
 
 TSharedRef<IDetailCustomization> FMaterialExpressionParameterDetails::MakeInstance(FOnCollectParameterGroups InCollectGroupsDelegate)
 {
@@ -520,12 +128,234 @@ void FMaterialExpressionParameterDetails::CustomizeDetails( IDetailLayoutBuilder
 					DetailLayout.ForceRefreshDetails();
 				}));
 			}
-			
+
+			TSharedPtr<IPropertyHandle> ChannelHandle = DetailLayout.GetProperty("ChannelNames", UMaterialExpressionVectorParameter::StaticClass());
+			TSharedPtr<IPropertyHandle> ValueHandle = DetailLayout.GetProperty("DefaultValue", UMaterialExpressionVectorParameter::StaticClass());
+			if (ChannelHandle.IsValid() && ChannelHandle->IsValidHandle())
+			{
+				static const FName Red("R");
+				static const FName Green("G");
+				static const FName Blue("B");
+				static const FName Alpha("A");
+				// Rebuild the layout when the ChannelNames property changes
+				ChannelHandle->GetChildHandle(Red)->SetOnPropertyValueChanged(FSimpleDelegate::CreateLambda([&DetailLayout]()
+				{
+	
+					DetailLayout.ForceRefreshDetails();
+				}));
+				ChannelHandle->GetChildHandle(Green)->SetOnPropertyValueChanged(FSimpleDelegate::CreateLambda([&DetailLayout]()
+				{
+
+					DetailLayout.ForceRefreshDetails();
+				}));
+				ChannelHandle->GetChildHandle(Blue)->SetOnPropertyValueChanged(FSimpleDelegate::CreateLambda([&DetailLayout]()
+				{
+
+					DetailLayout.ForceRefreshDetails();
+				}));
+				ChannelHandle->GetChildHandle(Alpha)->SetOnPropertyValueChanged(FSimpleDelegate::CreateLambda([&DetailLayout]()
+				{
+
+					DetailLayout.ForceRefreshDetails();
+				}));
+			}
+
 			if (VectorParameter->bUseCustomPrimitiveData)
 			{
 				DetailLayout.HideCategory(TEXT("MaterialExpressionVectorParameter"));
 				DetailLayout.HideCategory(MaterialExpressionCategory);
 			}
+
+			if (ValueHandle.IsValid() && ValueHandle->IsValidHandle())
+			{
+				static const FName Red("R");
+				static const FName Green("G");
+				static const FName Blue("B");
+				static const FName Alpha("A");
+				if (!VectorParameter->ChannelNames.R.IsEmpty())
+				{
+					ValueHandle->GetChildHandle(Red)->SetPropertyDisplayName(VectorParameter->ChannelNames.R);
+				}
+				if (!VectorParameter->ChannelNames.G.IsEmpty())
+				{
+					ValueHandle->GetChildHandle(Green)->SetPropertyDisplayName(VectorParameter->ChannelNames.G);
+				}
+				if (!VectorParameter->ChannelNames.B.IsEmpty())
+				{
+					ValueHandle->GetChildHandle(Blue)->SetPropertyDisplayName(VectorParameter->ChannelNames.B);
+				}
+				if (!VectorParameter->ChannelNames.A.IsEmpty())
+				{
+					ValueHandle->GetChildHandle(Alpha)->SetPropertyDisplayName(VectorParameter->ChannelNames.A);
+				}
+			}
+		}
+
+		UMaterialExpressionTextureSampleParameter* TextureParameter = Cast<UMaterialExpressionTextureSampleParameter>(Object);
+
+		if (TextureParameter)
+		{
+			TSharedPtr<IPropertyHandle> ChannelHandle = DetailLayout.GetProperty("ChannelNames", UMaterialExpressionTextureSampleParameter::StaticClass());
+			TSharedPtr<IPropertyHandle> ValueHandle = DetailLayout.GetProperty("Texture", UMaterialExpressionTextureBase::StaticClass());
+			if (TextureParameter->GetOutputType(0) != MCT_Texture)
+			{
+				if (ChannelHandle.IsValid() && ChannelHandle->IsValidHandle())
+				{
+					static const FName Red("R");
+					static const FName Green("G");
+					static const FName Blue("B");
+					static const FName Alpha("A");
+					// Rebuild the layout when the ChannelNames property changes
+					ChannelHandle->GetChildHandle(Red)->SetOnPropertyValueChanged(FSimpleDelegate::CreateLambda([&DetailLayout]()
+					{
+						DetailLayout.ForceRefreshDetails();
+					}));
+					ChannelHandle->GetChildHandle(Green)->SetOnPropertyValueChanged(FSimpleDelegate::CreateLambda([&DetailLayout]()
+					{
+						DetailLayout.ForceRefreshDetails();
+					}));
+					ChannelHandle->GetChildHandle(Blue)->SetOnPropertyValueChanged(FSimpleDelegate::CreateLambda([&DetailLayout]()
+					{
+						DetailLayout.ForceRefreshDetails();
+					}));
+					ChannelHandle->GetChildHandle(Alpha)->SetOnPropertyValueChanged(FSimpleDelegate::CreateLambda([&DetailLayout]()
+					{
+						DetailLayout.ForceRefreshDetails();
+					}));
+				}
+
+				if (ValueHandle.IsValid() && ValueHandle->IsValidHandle())
+				{
+					IDetailPropertyRow& PropertyRow = *DetailLayout.EditDefaultProperty(ValueHandle);
+					TSharedPtr<SWidget> NameWidget;
+					TSharedPtr<SWidget> ValueWidget;
+					FDetailWidgetRow DefaultRow;
+					PropertyRow.GetDefaultWidgets(NameWidget, ValueWidget, DefaultRow);
+
+					FDetailWidgetRow &DetailWidgetRow = PropertyRow.CustomWidget();
+					TSharedPtr<SVerticalBox> NameVerticalBox;
+					DetailWidgetRow.NameContent()
+						[
+							SAssignNew(NameVerticalBox, SVerticalBox)
+							+ SVerticalBox::Slot()
+							.AutoHeight()
+							[
+								SNew(STextBlock)
+								.Text(FText::FromName(TextureParameter->ParameterName))
+								.Font(FEditorStyle::GetFontStyle(TEXT("PropertyWindow.NormalFont")))
+							]
+						];
+
+					DetailWidgetRow.ValueContent()
+						.MinDesiredWidth(DefaultRow.ValueWidget.MinWidth)
+						.MaxDesiredWidth(DefaultRow.ValueWidget.MaxWidth)
+						[
+							ValueWidget.ToSharedRef()
+						];
+
+					static const FName Red("R");
+					static const FName Green("G");
+					static const FName Blue("B");
+					static const FName Alpha("A");
+
+					if (!TextureParameter->ChannelNames.R.IsEmpty())
+					{
+						NameVerticalBox->AddSlot()
+							[
+								SNew(SHorizontalBox)
+								+ SHorizontalBox::Slot()
+								.AutoWidth()
+								.Padding(20.0, 2.0, 4.0, 2.0)
+								[
+									SNew(STextBlock)
+									.Text(FText::FromName(Red))
+									.Font(FEditorStyle::GetFontStyle(TEXT("PropertyWindow.BoldFont")))
+								]
+								+ SHorizontalBox::Slot()
+								.HAlign(HAlign_Left)
+								.Padding(4.0, 2.0)
+								[
+									SNew(STextBlock)
+									.Text(TextureParameter->ChannelNames.R)
+									.Font(FEditorStyle::GetFontStyle(TEXT("PropertyWindow.NormalFont")))
+								]
+							];
+					}
+					if (!TextureParameter->ChannelNames.G.IsEmpty())
+					{
+						NameVerticalBox->AddSlot()
+							[
+								SNew(SHorizontalBox)
+								+ SHorizontalBox::Slot()
+								.Padding(20.0, 2.0, 4.0, 2.0)
+								.AutoWidth()
+								[
+									SNew(STextBlock)
+									.Text(FText::FromName(Green))
+									.Font(FEditorStyle::GetFontStyle(TEXT("PropertyWindow.BoldFont")))
+								]
+								+ SHorizontalBox::Slot()
+								.HAlign(HAlign_Left)
+								.Padding(4.0, 2.0)
+								[
+									SNew(STextBlock)
+									.Text(TextureParameter->ChannelNames.G)
+									.Font(FEditorStyle::GetFontStyle(TEXT("PropertyWindow.NormalFont")))
+								]
+							];
+					}
+					if (!TextureParameter->ChannelNames.B.IsEmpty())
+					{
+							NameVerticalBox->AddSlot()
+							[
+								SNew(SHorizontalBox)
+								+ SHorizontalBox::Slot()
+								.Padding(20.0, 2.0, 4.0, 2.0)
+								.AutoWidth()
+								[
+									SNew(STextBlock)
+									.Text(FText::FromName(Blue))
+									.Font(FEditorStyle::GetFontStyle(TEXT("PropertyWindow.BoldFont")))
+								]
+								+ SHorizontalBox::Slot()
+								.HAlign(HAlign_Left)
+								.Padding(4.0, 2.0)
+								[
+									SNew(STextBlock)
+									.Text(TextureParameter->ChannelNames.B)
+									.Font(FEditorStyle::GetFontStyle(TEXT("PropertyWindow.NormalFont")))
+								]
+							];
+					}
+					if (!TextureParameter->ChannelNames.A.IsEmpty())
+					{
+						NameVerticalBox->AddSlot()
+							[
+								SNew(SHorizontalBox)
+								+ SHorizontalBox::Slot()
+								.Padding(20.0, 2.0, 4.0, 2.0)
+								.AutoWidth()
+								[
+									SNew(STextBlock)
+									.Text(FText::FromName(Alpha))
+									.Font(FEditorStyle::GetFontStyle(TEXT("PropertyWindow.BoldFont")))
+								]
+								+ SHorizontalBox::Slot()
+								.HAlign(HAlign_Left)
+								.Padding(4.0, 2.0)
+								[
+									SNew(STextBlock)
+									.Text(TextureParameter->ChannelNames.A)
+									.Font(FEditorStyle::GetFontStyle(TEXT("PropertyWindow.NormalFont")))
+								]
+							];
+					}
+				}
+			}
+			else
+			{
+				DetailLayout.HideProperty(ChannelHandle);
+			}
 		}
 
 	}
diff --git a/Engine/Source/Editor/MaterialEditor/Private/MaterialEditorDetailCustomization.h b/Engine/Source/Editor/MaterialEditor/Private/MaterialEditorDetailCustomization.h
--- a/Engine/Source/Editor/MaterialEditor/Private/MaterialEditorDetailCustomization.h
+++ b/Engine/Source/Editor/MaterialEditor/Private/MaterialEditorDetailCustomization.h
@@ -14,46 +14,7 @@ class SEditableText;
 
 DECLARE_DELEGATE_OneParam(FOnCollectParameterGroups, TArray<FString>*);
 
-class FMaterialEditorParameterDetails : public IDetailCustomization
-{
-
-public:
-	/** Makes a new instance of this detail layout class for a specific detail view requesting it */
-	static TSharedRef<class IDetailCustomization> MakeInstance(class UMaterialEditorPreviewParameters* Material);
-
-	/** Constructor */
-	FMaterialEditorParameterDetails(class UMaterialEditorPreviewParameters* Material);
-
-	/** IDetailCustomization interface */
-	virtual void CustomizeDetails(IDetailLayoutBuilder& DetailLayout) override;
-
-	static TOptional<float> OnGetValue(TSharedRef<IPropertyHandle> PropertyHandle);
-	static void OnValueCommitted(float NewValue, ETextCommit::Type CommitType, TSharedRef<IPropertyHandle> PropertyHandle);
 
-private:
-	/** Builds the custom parameter groups category */
-	void CreateGroupsWidget(TSharedRef<IPropertyHandle> ParameterGroupsProperty, class IDetailCategoryBuilder& GroupsCategory);
-
-	/** Builds the widget for an individual parameter group */
-	void CreateSingleGroupWidget(struct FEditorParameterGroup& ParameterGroup, TSharedPtr<IPropertyHandle> ParameterGroupProperty, class IDetailGroup& DetailGroup);
-
-	/** These methods generate the custom widgets for the various parameter types */
-	void CreateParameterValueWidget(class UDEditorParameterValue* Parameter, TSharedPtr<IPropertyHandle> ParameterProperty, IDetailGroup& DetailGroup);
-	void CreateMaskParameterValueWidget(class UDEditorParameterValue* Parameter, TSharedPtr<IPropertyHandle> ParameterProperty, IDetailGroup& DetailGroup);
-	void CreateVectorChannelMaskParameterValueWidget(class UDEditorParameterValue* Parameter, TSharedPtr<IPropertyHandle> ParameterProperty, IDetailGroup& DetailGroup);
-	void CreateScalarAtlasPositionParameterValueWidget(class UDEditorParameterValue* Parameter, TSharedPtr<IPropertyHandle> ParameterProperty, IDetailGroup& DetailGroup);
-	FString GetCurvePath(class UDEditorScalarParameterValue* Parameter) const;
-	/** Gets the expression description of this parameter from the the base material */
-	FText GetParameterExpressionDescription(class UDEditorParameterValue* Parameter) const;
-
-private:
-	/** Object that stores all of the possible parameters we can edit */
-	UMaterialEditorPreviewParameters* MaterialEditorInstance;
-
-	/** Associated UMaterialEditorPreviewParameters utilities */
-	TWeakPtr<class IPropertyUtilities> PropertyUtilities;
-
-};
 
 class FMaterialExpressionParameterDetails : public IDetailCustomization
 {
diff --git a/Engine/Source/Editor/MaterialEditor/Private/MaterialEditorInstanceDetailCustomization.cpp b/Engine/Source/Editor/MaterialEditor/Private/MaterialEditorInstanceDetailCustomization.cpp
--- a/Engine/Source/Editor/MaterialEditor/Private/MaterialEditorInstanceDetailCustomization.cpp
+++ b/Engine/Source/Editor/MaterialEditor/Private/MaterialEditorInstanceDetailCustomization.cpp
@@ -49,6 +49,8 @@
 
 #define LOCTEXT_NAMESPACE "MaterialInstanceEditor"
 
+
+
 TSharedRef<IDetailCustomization> FMaterialInstanceParameterDetails::MakeInstance(UMaterialEditorInstanceConstant* MaterialInstance, FGetShowHiddenParameters InShowHiddenDelegate)
 {
 	return MakeShareable(new FMaterialInstanceParameterDetails(MaterialInstance, InShowHiddenDelegate));
@@ -427,6 +429,14 @@ void FMaterialInstanceParameterDetails::CreateSingleGroupWidget(FEditorParameter
 			{
 				CreateScalarAtlasPositionParameterValueWidget(Parameter, ParameterProperty, DetailGroup);
 			}
+			if (TextureParam && 
+				( !TextureParam->ChannelNames.R.IsEmpty()
+				|| !TextureParam->ChannelNames.G.IsEmpty()
+				|| !TextureParam->ChannelNames.B.IsEmpty()
+				|| !TextureParam->ChannelNames.A.IsEmpty()))
+			{
+				CreateLabeledTextureParameterValueWidget(Parameter, ParameterProperty, DetailGroup);
+			}
 			else if (ScalarParam || SwitchParam || TextureParam || VectorParam || FontParam)
 			{
 				if (ScalarParam && ScalarParam->SliderMax > ScalarParam->SliderMin)
@@ -436,6 +446,30 @@ void FMaterialInstanceParameterDetails::CreateSingleGroupWidget(FEditorParameter
 					ParameterValueProperty->SetInstanceMetaData("UIMax", FString::Printf(TEXT("%f"), ScalarParam->SliderMax));
 				}
 
+				if (VectorParam)
+				{
+					static const FName Red("R");
+					static const FName Green("G");
+					static const FName Blue("B");
+					static const FName Alpha("A");
+					if (!VectorParam->ChannelNames.R.IsEmpty())
+					{
+						ParameterProperty->GetChildHandle(Red)->SetPropertyDisplayName(VectorParam->ChannelNames.R);
+					}
+					if (!VectorParam->ChannelNames.G.IsEmpty())
+					{
+						ParameterProperty->GetChildHandle(Green)->SetPropertyDisplayName(VectorParam->ChannelNames.G);
+					}
+					if (!VectorParam->ChannelNames.B.IsEmpty())
+					{
+						ParameterProperty->GetChildHandle(Blue)->SetPropertyDisplayName(VectorParam->ChannelNames.B);
+					}
+					if (!VectorParam->ChannelNames.A.IsEmpty())
+					{
+						ParameterProperty->GetChildHandle(Alpha)->SetPropertyDisplayName(VectorParam->ChannelNames.A);
+					}
+				}
+
 				CreateParameterValueWidget(Parameter, ParameterProperty, DetailGroup);
 			}
 			else if (LayersParam)
@@ -671,6 +705,159 @@ void FMaterialInstanceParameterDetails::CreateScalarAtlasPositionParameterValueW
 	}
 }
 
+void FMaterialInstanceParameterDetails::CreateLabeledTextureParameterValueWidget(class UDEditorParameterValue* Parameter, TSharedPtr<IPropertyHandle> ParameterProperty, IDetailGroup& DetailGroup)
+{
+	TSharedPtr<IPropertyHandle> ParameterValueProperty = ParameterProperty->GetChildHandle("ParameterValue");
+
+	if (ParameterValueProperty->IsValidHandle())
+	{
+		UDEditorTextureParameterValue* TextureParam = Cast<UDEditorTextureParameterValue>(Parameter);
+		if (TextureParam)
+		{
+			TAttribute<bool> IsParamEnabled = TAttribute<bool>::Create(TAttribute<bool>::FGetter::CreateStatic(&FMaterialPropertyHelpers::IsOverriddenExpression, Parameter));
+
+			IDetailPropertyRow& PropertyRow = DetailGroup.AddPropertyRow(ParameterValueProperty.ToSharedRef());
+
+			FIsResetToDefaultVisible IsResetVisible = FIsResetToDefaultVisible::CreateStatic(&FMaterialPropertyHelpers::ShouldShowResetToDefault, Parameter, MaterialEditorInstance);
+			FResetToDefaultHandler ResetHandler = FResetToDefaultHandler::CreateStatic(&FMaterialPropertyHelpers::ResetToDefault, Parameter, MaterialEditorInstance);
+			FResetToDefaultOverride ResetOverride = FResetToDefaultOverride::Create(IsResetVisible, ResetHandler);
+
+			PropertyRow
+				.DisplayName(FText::FromName(Parameter->ParameterInfo.Name))
+				.EditCondition(IsParamEnabled, FOnBooleanValueChanged::CreateStatic(&FMaterialPropertyHelpers::OnOverrideParameter, Parameter, MaterialEditorInstance))
+				.ToolTip(FMaterialPropertyHelpers::GetParameterTooltip(Parameter, MaterialEditorInstance))
+				.Visibility(TAttribute<EVisibility>::Create(TAttribute<EVisibility>::FGetter::CreateStatic(&FMaterialPropertyHelpers::ShouldShowExpression, Parameter, MaterialEditorInstance, ShowHiddenDelegate)))
+				.OverrideResetToDefault(ResetOverride);
+
+			
+			TSharedPtr<SWidget> NameWidget;
+			TSharedPtr<SWidget> ValueWidget;
+			FDetailWidgetRow Row;
+			PropertyRow.GetDefaultWidgets(NameWidget, ValueWidget, Row);
+
+			FDetailWidgetRow &DetailWidgetRow = PropertyRow.CustomWidget();
+			TSharedPtr<SVerticalBox> NameVerticalBox;
+			DetailWidgetRow.NameContent()
+			[
+				SAssignNew(NameVerticalBox, SVerticalBox)
+				+ SVerticalBox::Slot()
+				.AutoHeight()
+				[
+					SNew(STextBlock)
+					.Text(FText::FromName(Parameter->ParameterInfo.Name))
+					.ToolTipText(FMaterialPropertyHelpers::GetParameterTooltip(Parameter, MaterialEditorInstance))
+					.Font(FEditorStyle::GetFontStyle(TEXT("PropertyWindow.NormalFont")))
+				]
+			];
+			DetailWidgetRow.ValueContent()
+			.MinDesiredWidth(Row.ValueWidget.MinWidth)
+			.MaxDesiredWidth(Row.ValueWidget.MaxWidth)
+			[
+				ValueWidget.ToSharedRef()
+			];
+
+			static const FName Red("R");
+			static const FName Green("G");
+			static const FName Blue("B");
+			static const FName Alpha("A");
+
+			if(!TextureParam->ChannelNames.R.IsEmpty())
+			{
+				NameVerticalBox->AddSlot()
+				[
+					SNew(SHorizontalBox)
+					+SHorizontalBox::Slot()
+					.AutoWidth()
+					.Padding(20.0, 2.0, 4.0, 2.0)
+					[
+						SNew(STextBlock)
+						.Text(FText::FromName(Red))
+						.Font(FEditorStyle::GetFontStyle(TEXT("PropertyWindow.BoldFont")))
+					]
+					+SHorizontalBox::Slot()
+					.HAlign(HAlign_Left)
+					.Padding(4.0, 2.0)
+					[
+						SNew(STextBlock)
+						.Text(TextureParam->ChannelNames.R)
+						.Font(FEditorStyle::GetFontStyle(TEXT("PropertyWindow.NormalFont")))
+					]
+				];
+			}
+			if (!TextureParam->ChannelNames.G.IsEmpty())
+			{
+				NameVerticalBox->AddSlot()
+					[
+						SNew(SHorizontalBox)
+						+ SHorizontalBox::Slot()
+						.Padding(20.0, 2.0, 4.0, 2.0)
+						.AutoWidth()
+						[
+							SNew(STextBlock)
+							.Text(FText::FromName(Green))
+							.Font(FEditorStyle::GetFontStyle(TEXT("PropertyWindow.BoldFont")))
+						]
+						+ SHorizontalBox::Slot()
+						.HAlign(HAlign_Left)
+						.Padding(4.0, 2.0)
+						[
+							SNew(STextBlock)
+							.Text(TextureParam->ChannelNames.G)
+							.Font(FEditorStyle::GetFontStyle(TEXT("PropertyWindow.NormalFont")))
+						]
+					];
+			}
+			if (!TextureParam->ChannelNames.B.IsEmpty())
+			{
+				NameVerticalBox->AddSlot()
+					[
+						SNew(SHorizontalBox)
+						+ SHorizontalBox::Slot()
+						.Padding(20.0, 2.0, 4.0, 2.0)
+						.AutoWidth()
+						[
+							SNew(STextBlock)
+							.Text(FText::FromName(Blue))
+							.Font(FEditorStyle::GetFontStyle(TEXT("PropertyWindow.BoldFont")))
+						]
+						+ SHorizontalBox::Slot()
+						.HAlign(HAlign_Left)
+						.Padding(4.0, 2.0)
+						[
+							SNew(STextBlock)
+							.Text(TextureParam->ChannelNames.B)
+							.Font(FEditorStyle::GetFontStyle(TEXT("PropertyWindow.NormalFont")))
+						]
+					];
+			}
+			if (!TextureParam->ChannelNames.A.IsEmpty())
+			{
+				NameVerticalBox->AddSlot()
+				[
+					SNew(SHorizontalBox)
+					+ SHorizontalBox::Slot()
+					.Padding(20.0, 2.0, 4.0, 2.0)
+					.AutoWidth()
+					[
+						SNew(STextBlock)
+						.Text(FText::FromName(Alpha))
+						.Font(FEditorStyle::GetFontStyle(TEXT("PropertyWindow.BoldFont")))
+					]
+					+ SHorizontalBox::Slot()
+					.HAlign(HAlign_Left)
+					.Padding(4.0, 2.0)
+					[
+						SNew(STextBlock)
+						.Text(TextureParam->ChannelNames.A)
+						.Font(FEditorStyle::GetFontStyle(TEXT("PropertyWindow.NormalFont")))
+					]
+				];
+			}
+			
+		}
+	}
+}
+
 FString FMaterialInstanceParameterDetails::GetCurvePath(UDEditorScalarParameterValue* Parameter) const
 {
 	FString Path = Parameter->AtlasData.Curve->GetPathName();
diff --git a/Engine/Source/Editor/MaterialEditor/Private/MaterialEditorInstanceDetailCustomization.h b/Engine/Source/Editor/MaterialEditor/Private/MaterialEditorInstanceDetailCustomization.h
--- a/Engine/Source/Editor/MaterialEditor/Private/MaterialEditorInstanceDetailCustomization.h
+++ b/Engine/Source/Editor/MaterialEditor/Private/MaterialEditorInstanceDetailCustomization.h
@@ -7,6 +7,7 @@
 #include "Layout/Visibility.h"
 #include "IDetailCustomization.h"
 #include "Input/Reply.h"
+#include "Customizations/ColorStructCustomization.h"
 
 
 struct FAssetData;
@@ -17,7 +18,6 @@ class UMaterialEditorInstanceConstant;
 
 DECLARE_DELEGATE_OneParam(FGetShowHiddenParameters, bool&);
 
-
 /*-----------------------------------------------------------------------------
    FMaterialInstanceParameterDetails
 -----------------------------------------------------------------------------*/
@@ -52,6 +52,7 @@ private:
 	void CreateMaskParameterValueWidget(class UDEditorParameterValue* Parameter, TSharedPtr<IPropertyHandle> ParameterProperty, IDetailGroup& DetailGroup);
 	void CreateVectorChannelMaskParameterValueWidget(class UDEditorParameterValue* Parameter, TSharedPtr<IPropertyHandle> ParameterProperty, IDetailGroup& DetailGroup);
 	void CreateScalarAtlasPositionParameterValueWidget(class UDEditorParameterValue* Parameter, TSharedPtr<IPropertyHandle> ParameterProperty, IDetailGroup& DetailGroup);
+	void CreateLabeledTextureParameterValueWidget(class UDEditorParameterValue* Parameter, TSharedPtr<IPropertyHandle> ParameterProperty, IDetailGroup& DetailGroup);
 
 	FString GetCurvePath(class UDEditorScalarParameterValue * Parameter) const;
 
diff --git a/Engine/Source/Editor/MaterialEditor/Private/SMaterialParametersOverviewWidget.cpp b/Engine/Source/Editor/MaterialEditor/Private/SMaterialParametersOverviewWidget.cpp
--- a/Engine/Source/Editor/MaterialEditor/Private/SMaterialParametersOverviewWidget.cpp
+++ b/Engine/Source/Editor/MaterialEditor/Private/SMaterialParametersOverviewWidget.cpp
@@ -99,6 +99,7 @@ void SMaterialParametersOverviewTreeItem::Construct(const FArguments& InArgs, co
 		UDEditorStaticComponentMaskParameterValue* CompMaskParam = Cast<UDEditorStaticComponentMaskParameterValue>(StackParameterData->Parameter);
 		UDEditorVectorParameterValue* VectorParam = Cast<UDEditorVectorParameterValue>(StackParameterData->Parameter);
 		UDEditorScalarParameterValue* ScalarParam = Cast<UDEditorScalarParameterValue>(StackParameterData->Parameter);
+		UDEditorTextureParameterValue* TextureParam = Cast<UDEditorTextureParameterValue>(StackParameterData->Parameter);
 
 		TAttribute<bool> IsParamEnabled = TAttribute<bool>::Create(TAttribute<bool>::FGetter::CreateStatic(&FMaterialPropertyHelpers::IsOverriddenExpression, StackParameterData->Parameter));
 		NameOverride = FText::FromName(StackParameterData->Parameter->ParameterInfo.Name);
@@ -171,6 +172,133 @@ void SMaterialParametersOverviewTreeItem::Construct(const FArguments& InArgs, co
 				];
 			
 		}
+		else if (TextureParam)
+		{
+
+			TSharedPtr<SWidget> NameWidget;
+			TSharedPtr<SWidget> ValueWidget;
+			FDetailWidgetRow DefaultRow;
+			Row.GetDefaultWidgets(NameWidget, ValueWidget, DefaultRow);
+
+			FDetailWidgetRow &DetailWidgetRow = Row.CustomWidget();
+			TSharedPtr<SVerticalBox> NameVerticalBox;
+			DetailWidgetRow.NameContent()
+				[
+					SAssignNew(NameVerticalBox, SVerticalBox)
+					+ SVerticalBox::Slot()
+					.AutoHeight()
+					[
+						SNew(STextBlock)
+						.Text(FText::FromName(StackParameterData->Parameter->ParameterInfo.Name))
+						.Font(FEditorStyle::GetFontStyle(TEXT("PropertyWindow.NormalFont")))
+					]
+				];
+
+			DetailWidgetRow.ValueContent()
+				.MinDesiredWidth(DefaultRow.ValueWidget.MinWidth)
+				.MaxDesiredWidth(DefaultRow.ValueWidget.MaxWidth)
+				[
+					ValueWidget.ToSharedRef()
+				];
+
+			static const FName Red("R");
+			static const FName Green("G");
+			static const FName Blue("B");
+			static const FName Alpha("A");
+
+			if (!TextureParam->ChannelNames.R.IsEmpty())
+			{
+				NameVerticalBox->AddSlot()
+					[
+						SNew(SHorizontalBox)
+						+ SHorizontalBox::Slot()
+						.AutoWidth()
+						.Padding(20.0, 2.0, 4.0, 2.0)
+						[
+							SNew(STextBlock)
+							.Text(FText::FromName(Red))
+							.Font(FEditorStyle::GetFontStyle(TEXT("PropertyWindow.BoldFont")))
+						]
+						+ SHorizontalBox::Slot()
+						.HAlign(HAlign_Left)
+						.Padding(4.0, 2.0)
+						[
+							SNew(STextBlock)
+							.Text(TextureParam->ChannelNames.R)
+							.Font(FEditorStyle::GetFontStyle(TEXT("PropertyWindow.NormalFont")))
+						]
+					];
+			}
+			if (!TextureParam->ChannelNames.G.IsEmpty())
+			{
+				NameVerticalBox->AddSlot()
+					[
+						SNew(SHorizontalBox)
+						+ SHorizontalBox::Slot()
+						.Padding(20.0, 2.0, 4.0, 2.0)
+						.AutoWidth()
+						[
+							SNew(STextBlock)
+							.Text(FText::FromName(Green))
+							.Font(FEditorStyle::GetFontStyle(TEXT("PropertyWindow.BoldFont")))
+						]
+						+ SHorizontalBox::Slot()
+						.HAlign(HAlign_Left)
+						.Padding(4.0, 2.0)
+						[
+							SNew(STextBlock)
+							.Text(TextureParam->ChannelNames.G)
+							.Font(FEditorStyle::GetFontStyle(TEXT("PropertyWindow.NormalFont")))
+						]
+					];
+			}
+			if (!TextureParam->ChannelNames.B.IsEmpty())
+			{
+				NameVerticalBox->AddSlot()
+					[
+						SNew(SHorizontalBox)
+						+ SHorizontalBox::Slot()
+						.Padding(20.0, 2.0, 4.0, 2.0)
+						.AutoWidth()
+						[
+							SNew(STextBlock)
+							.Text(FText::FromName(Blue))
+							.Font(FEditorStyle::GetFontStyle(TEXT("PropertyWindow.BoldFont")))
+						]
+						+ SHorizontalBox::Slot()
+						.HAlign(HAlign_Left)
+						.Padding(4.0, 2.0)
+						[
+							SNew(STextBlock)
+							.Text(TextureParam->ChannelNames.B)
+						.Font(FEditorStyle::GetFontStyle(TEXT("PropertyWindow.NormalFont")))
+						]
+					];
+			}
+			if (!TextureParam->ChannelNames.A.IsEmpty())
+			{
+				NameVerticalBox->AddSlot()
+					[
+						SNew(SHorizontalBox)
+						+ SHorizontalBox::Slot()
+						.Padding(20.0, 2.0, 4.0, 2.0)
+						.AutoWidth()
+						[
+							SNew(STextBlock)
+							.Text(FText::FromName(Alpha))
+							.Font(FEditorStyle::GetFontStyle(TEXT("PropertyWindow.BoldFont")))
+						]
+						+ SHorizontalBox::Slot()
+						.HAlign(HAlign_Left)
+						.Padding(4.0, 2.0)
+						[
+							SNew(STextBlock)
+							.Text(TextureParam->ChannelNames.A)
+							.Font(FEditorStyle::GetFontStyle(TEXT("PropertyWindow.NormalFont")))
+						]
+					];
+			}
+		}
 		else if (!CompMaskParam)
 		{
 			FNodeWidgets StoredNodeWidgets = Node.CreateNodeWidgets();
@@ -501,6 +629,7 @@ void SMaterialParametersOverviewTree::CreateGroupsWidget()
 			{
 				FUnsortedParamData NonLayerProperty;
 				UDEditorScalarParameterValue* ScalarParam = Cast<UDEditorScalarParameterValue>(Parameter);
+				UDEditorVectorParameterValue* VectorParam = Cast<UDEditorVectorParameterValue>(Parameter);
 
 				if (ScalarParam && ScalarParam->SliderMax > ScalarParam->SliderMin)
 				{
@@ -508,6 +637,30 @@ void SMaterialParametersOverviewTree::CreateGroupsWidget()
 					ParameterValueProperty->SetInstanceMetaData("UIMax", FString::Printf(TEXT("%f"), ScalarParam->SliderMax));
 				}
 
+				if (VectorParam)
+				{
+					static const FName Red("R");
+					static const FName Green("G");
+					static const FName Blue("B");
+					static const FName Alpha("A");
+					if (!VectorParam->ChannelNames.R.IsEmpty())
+					{
+						ParameterProperty->GetChildHandle(Red)->SetPropertyDisplayName(VectorParam->ChannelNames.R);
+					}
+					if (!VectorParam->ChannelNames.G.IsEmpty())
+					{
+						ParameterProperty->GetChildHandle(Green)->SetPropertyDisplayName(VectorParam->ChannelNames.G);
+					}
+					if (!VectorParam->ChannelNames.B.IsEmpty())
+					{
+						ParameterProperty->GetChildHandle(Blue)->SetPropertyDisplayName(VectorParam->ChannelNames.B);
+					}
+					if (!VectorParam->ChannelNames.A.IsEmpty())
+					{
+						ParameterProperty->GetChildHandle(Alpha)->SetPropertyDisplayName(VectorParam->ChannelNames.A);
+					}
+				}
+
 				NonLayerProperty.Parameter = Parameter;
 				NonLayerProperty.ParameterGroup = ParameterGroup;
 				NonLayerProperty.ParameterNode = Generator->FindTreeNode(ParameterValueProperty);
diff --git a/Engine/Source/Editor/UnrealEd/Classes/MaterialEditor/DEditorTextureParameterValue.h b/Engine/Source/Editor/UnrealEd/Classes/MaterialEditor/DEditorTextureParameterValue.h
--- a/Engine/Source/Editor/UnrealEd/Classes/MaterialEditor/DEditorTextureParameterValue.h
+++ b/Engine/Source/Editor/UnrealEd/Classes/MaterialEditor/DEditorTextureParameterValue.h
@@ -16,5 +16,7 @@ class UNREALED_API UDEditorTextureParameterValue : public UDEditorParameterValue
 	UPROPERTY(EditAnywhere, Category=DEditorTextureParameterValue)
 	class UTexture* ParameterValue;
 
+	UPROPERTY(Transient)
+	FParameterChannelNames ChannelNames;
 };
 
diff --git a/Engine/Source/Editor/UnrealEd/Classes/MaterialEditor/DEditorVectorParameterValue.h b/Engine/Source/Editor/UnrealEd/Classes/MaterialEditor/DEditorVectorParameterValue.h
--- a/Engine/Source/Editor/UnrealEd/Classes/MaterialEditor/DEditorVectorParameterValue.h
+++ b/Engine/Source/Editor/UnrealEd/Classes/MaterialEditor/DEditorVectorParameterValue.h
@@ -5,7 +5,7 @@
 #include "CoreMinimal.h"
 #include "UObject/ObjectMacros.h"
 #include "MaterialEditor/DEditorParameterValue.h"
-
+#include "Materials/MaterialExpressionVectorParameter.h"
 #include "DEditorVectorParameterValue.generated.h"
 
 UCLASS(hidecategories=Object, collapsecategories, editinlinenew)
@@ -18,5 +18,8 @@ class UNREALED_API UDEditorVectorParameterValue : public UDEditorParameterValue
 
 	UPROPERTY(Transient)
 	bool bIsUsedAsChannelMask;
+
+	UPROPERTY(Transient)
+	FParameterChannelNames ChannelNames;
 };
 
diff --git a/Engine/Source/Editor/UnrealEd/Private/MaterialGraphNode.cpp b/Engine/Source/Editor/UnrealEd/Private/MaterialGraphNode.cpp
--- a/Engine/Source/Editor/UnrealEd/Private/MaterialGraphNode.cpp
+++ b/Engine/Source/Editor/UnrealEd/Private/MaterialGraphNode.cpp
@@ -42,6 +42,8 @@
 
 #define LOCTEXT_NAMESPACE "MaterialGraphNode"
 
+static FText SpaceText = LOCTEXT("Space", " ");
+
 /////////////////////////////////////////////////////
 // UMaterialGraphNode
 
@@ -553,7 +555,7 @@ void UMaterialGraphNode::CreateInputPins()
 		{
 			// Makes sure pin has a name for lookup purposes but user will never see it
 			NewPin->PinName = CreateUniquePinName(TEXT("Input"));
-			NewPin->PinFriendlyName = FText::GetEmpty();
+			NewPin->PinFriendlyName = SpaceText;
 		}
 	}
 }
@@ -607,7 +609,7 @@ void UMaterialGraphNode::CreateOutputPins()
 		{
 			// Makes sure pin has a name for lookup purposes but user will never see it
 			NewPin->PinName = CreateUniquePinName(TEXT("Output"));
-			NewPin->PinFriendlyName = FText::GetEmpty();
+			NewPin->PinFriendlyName = SpaceText;
 		}
 	}
 }
diff --git a/Engine/Source/Editor/UnrealEd/Private/PreviewMaterial.cpp b/Engine/Source/Editor/UnrealEd/Private/PreviewMaterial.cpp
--- a/Engine/Source/Editor/UnrealEd/Private/PreviewMaterial.cpp
+++ b/Engine/Source/Editor/UnrealEd/Private/PreviewMaterial.cpp
@@ -341,7 +341,8 @@ void UMaterialEditorPreviewParameters::RegenerateArrays()
 			if (PreviewMaterial->GetVectorParameterValue(ParameterValue.ParameterInfo, Value))
 			{
 				ParameterValue.ParameterValue = Value;
-				PreviewMaterial->IsVectorParameterUsedAsChannelMask(ParameterValue.ParameterInfo, ParameterValue.bIsUsedAsChannelMask);			
+				PreviewMaterial->IsVectorParameterUsedAsChannelMask(ParameterValue.ParameterInfo, ParameterValue.bIsUsedAsChannelMask);		
+				PreviewMaterial->GetVectorParameterChannelNames(ParameterValue.ParameterInfo, ParameterValue.ChannelNames);
 			}
 			if (PreviewMaterial->GetParameterSortPriority(ParameterName, SortPriority))
 			{
@@ -400,6 +401,7 @@ void UMaterialEditorPreviewParameters::RegenerateArrays()
 			if (PreviewMaterial->GetTextureParameterValue(ParameterValue.ParameterInfo, Value))
 			{
 				ParameterValue.ParameterValue = Value;
+				PreviewMaterial->GetTextureParameterChannelNames(ParameterValue.ParameterInfo, ParameterValue.ChannelNames);
 			}
 			if (ParentMaterial->GetParameterSortPriority(ParameterName, SortPriority))
 			{
@@ -1007,6 +1009,7 @@ void UMaterialEditorInstanceConstant::RegenerateArrays()
 
 			SourceInstance->GetVectorParameterValue(ParameterInfo, ParameterValue.ParameterValue);
 			SourceInstance->IsVectorParameterUsedAsChannelMask(ParameterInfo, ParameterValue.bIsUsedAsChannelMask);
+			SourceInstance->GetVectorParameterChannelNames(ParameterInfo, ParameterValue.ChannelNames);
 
 			// @todo: This is kind of slow, maybe store these in a map for lookup?
 			// See if this keyname exists in the source instance.
@@ -1039,6 +1042,7 @@ void UMaterialEditorInstanceConstant::RegenerateArrays()
 
 			ParameterValue.ParameterValue = nullptr;
 			SourceInstance->GetTextureParameterValue(ParameterInfo, ParameterValue.ParameterValue);
+			SourceInstance->GetTextureParameterChannelNames(ParameterInfo, ParameterValue.ChannelNames);
 
 			// @todo: This is kind of slow, maybe store these in a map for lookup?
 			// See if this keyname exists in the source instance.
diff --git a/Engine/Source/Runtime/Engine/Classes/Materials/Material.h b/Engine/Source/Runtime/Engine/Classes/Materials/Material.h
--- a/Engine/Source/Runtime/Engine/Classes/Materials/Material.h
+++ b/Engine/Source/Runtime/Engine/Classes/Materials/Material.h
@@ -969,7 +969,13 @@ public:
 #endif
 	ENGINE_API virtual bool GetVectorParameterValue(const FMaterialParameterInfo& ParameterInfo,FLinearColor& OutValue, bool bOveriddenOnly = false) const override;
 	ENGINE_API virtual bool IsVectorParameterUsedAsChannelMask(const FMaterialParameterInfo& ParameterInfo, bool& OutValue) const override;
+#if WITH_EDITOR
+	ENGINE_API virtual bool GetVectorParameterChannelNames(const FMaterialParameterInfo& ParameterInfo, FParameterChannelNames& OutValue) const override;
+#endif
 	ENGINE_API virtual bool GetTextureParameterValue(const FMaterialParameterInfo& ParameterInfo,class UTexture*& OutValue, bool bOveriddenOnly = false) const override;
+#if WITH_EDITOR
+	ENGINE_API virtual bool GetTextureParameterChannelNames(const FMaterialParameterInfo& ParameterInfo, FParameterChannelNames& OutValue) const override;
+#endif
 	ENGINE_API virtual bool GetFontParameterValue(const FMaterialParameterInfo& ParameterInfo,class UFont*& OutFontValue,int32& OutFontPage, bool bOveriddenOnly = false) const override;
 	ENGINE_API virtual bool GetRefractionSettings(float& OutBiasValue) const override;
 	ENGINE_API virtual FMaterialRenderProxy* GetRenderProxy() const override;
diff --git a/Engine/Source/Runtime/Engine/Classes/Materials/MaterialExpression.h b/Engine/Source/Runtime/Engine/Classes/Materials/MaterialExpression.h
--- a/Engine/Source/Runtime/Engine/Classes/Materials/MaterialExpression.h
+++ b/Engine/Source/Runtime/Engine/Classes/Materials/MaterialExpression.h
@@ -96,6 +96,29 @@ struct FExpressionOutput
 };
 #endif
 
+USTRUCT()
+struct FParameterChannelNames
+{
+	GENERATED_USTRUCT_BODY()
+
+		UPROPERTY(EditAnywhere, Category = MaterialExpressionVectorParameter)
+		FText R;
+
+	UPROPERTY(EditAnywhere, Category = MaterialExpressionVectorParameter)
+		FText G;
+
+	UPROPERTY(EditAnywhere, Category = MaterialExpressionVectorParameter)
+		FText B;
+
+	UPROPERTY(EditAnywhere, Category = MaterialExpressionVectorParameter)
+		FText A;
+
+	FParameterChannelNames()
+	{
+
+	};
+};
+
 UCLASS(abstract, BlueprintType, hidecategories=Object)
 class ENGINE_API UMaterialExpression : public UObject
 {
diff --git a/Engine/Source/Runtime/Engine/Classes/Materials/MaterialExpressionTextureSampleParameter.h b/Engine/Source/Runtime/Engine/Classes/Materials/MaterialExpressionTextureSampleParameter.h
--- a/Engine/Source/Runtime/Engine/Classes/Materials/MaterialExpressionTextureSampleParameter.h
+++ b/Engine/Source/Runtime/Engine/Classes/Materials/MaterialExpressionTextureSampleParameter.h
@@ -32,6 +32,9 @@ class ENGINE_API UMaterialExpressionTextureSampleParameter : public UMaterialExp
 	/** Controls where the this parameter is displayed in a material instance parameter list.  The lower the number the higher up in the parameter list. */
 	UPROPERTY(EditAnywhere, Category = MaterialExpressionTextureSampleParameter)
 	int32 SortPriority;
+
+	UPROPERTY(EditAnywhere, Category = ParameterCustomization)
+	FParameterChannelNames ChannelNames;
 #endif
 
 	//~ Begin UMaterialExpression Interface
@@ -56,6 +59,8 @@ class ENGINE_API UMaterialExpressionTextureSampleParameter : public UMaterialExp
 
 #if WITH_EDITOR
 	bool SetParameterValue(FName InParameterName, UTexture* InValue);
+	virtual void PostEditChangeProperty(FPropertyChangedEvent& PropertyChangedEvent) override;
+	void ApplyChannelNames();
 #endif
 
 	/**
@@ -83,5 +88,12 @@ class ENGINE_API UMaterialExpressionTextureSampleParameter : public UMaterialExp
 		return ExpressionGUID;
 	}
 
+#if WITH_EDITOR
+	FParameterChannelNames GetTextureChannelNames() const
+	{
+		return ChannelNames;
+	}
+#endif
+
 	void GetAllParameterInfo(TArray<FMaterialParameterInfo> &OutParameterInfo, TArray<FGuid> &OutParameterIds, const FMaterialParameterInfo& InBaseParameterInfo) const;
 };
diff --git a/Engine/Source/Runtime/Engine/Classes/Materials/MaterialExpressionVectorParameter.h b/Engine/Source/Runtime/Engine/Classes/Materials/MaterialExpressionVectorParameter.h
--- a/Engine/Source/Runtime/Engine/Classes/Materials/MaterialExpressionVectorParameter.h
+++ b/Engine/Source/Runtime/Engine/Classes/Materials/MaterialExpressionVectorParameter.h
@@ -23,6 +23,11 @@ class UMaterialExpressionVectorParameter : public UMaterialExpressionParameter
 	UPROPERTY(EditAnywhere, Category=CustomPrimitiveData, meta=(ClampMin="0"))
 	uint8 PrimitiveDataIndex = 0;
 
+#if WITH_EDITORONLY_DATA
+	UPROPERTY(EditAnywhere, Category = ParameterCustomization)
+	FParameterChannelNames ChannelNames;
+#endif
+
 	//~ Begin UMaterialExpression Interface
 #if WITH_EDITOR
 	virtual int32 Compile(class FMaterialCompiler* Compiler, int32 OutputIndex) override;
@@ -37,6 +42,10 @@ class UMaterialExpressionVectorParameter : public UMaterialExpressionParameter
 	virtual bool SetParameterValue(FName InParameterName, FLinearColor InValue);
 
 	virtual void PostEditChangeProperty(FPropertyChangedEvent& PropertyChangedEvent) override;	
+
+	void ApplyChannelNames();
+	
+
 	virtual void ValidateParameterName(const bool bAllowDuplicateName) override;
 	virtual bool HasClassAndNameCollision(UMaterialExpression* OtherExpression) const override;
 	virtual void SetValueToMatchingExpression(UMaterialExpression* OtherExpression) override;
@@ -44,6 +53,14 @@ class UMaterialExpressionVectorParameter : public UMaterialExpressionParameter
 
 	virtual bool IsUsedAsChannelMask() const {return false;}
 
+#if WITH_EDITOR
+	FParameterChannelNames GetVectorChannelNames() const
+	{
+		return ChannelNames;
+	}
+#endif
+
 	virtual void GetAllParameterInfo(TArray<FMaterialParameterInfo> &OutParameterInfo, TArray<FGuid> &OutParameterIds, const FMaterialParameterInfo& InBaseParameterInfo) const override;
 };
 
+
diff --git a/Engine/Source/Runtime/Engine/Classes/Materials/MaterialInstance.h b/Engine/Source/Runtime/Engine/Classes/Materials/MaterialInstance.h
--- a/Engine/Source/Runtime/Engine/Classes/Materials/MaterialInstance.h
+++ b/Engine/Source/Runtime/Engine/Classes/Materials/MaterialInstance.h
@@ -421,7 +421,13 @@ public:
 	virtual ENGINE_API bool IsScalarParameterUsedAsAtlasPosition(const FMaterialParameterInfo& ParameterInfo, bool& OutValue, TSoftObjectPtr<class UCurveLinearColor>& Curve, TSoftObjectPtr<class UCurveLinearColorAtlas>& Atlas) const override;
 	virtual ENGINE_API bool GetVectorParameterValue(const FMaterialParameterInfo& ParameterInfo, FLinearColor& OutValue, bool bOveriddenOnly = false) const override;
 	virtual ENGINE_API bool IsVectorParameterUsedAsChannelMask(const FMaterialParameterInfo& ParameterInfo, bool& OutValue) const override;
+#if WITH_EDITOR
+	virtual ENGINE_API bool GetVectorParameterChannelNames(const FMaterialParameterInfo& ParameterInfo, FParameterChannelNames& OutValue) const override;
+#endif
 	virtual ENGINE_API bool GetTextureParameterValue(const FMaterialParameterInfo& ParameterInfo, class UTexture*& OutValue, bool bOveriddenOnly = false) const override;
+#if WITH_EDITOR
+	virtual ENGINE_API bool GetTextureParameterChannelNames(const FMaterialParameterInfo& ParameterInfo, FParameterChannelNames& OutValue) const override;
+#endif
 	virtual ENGINE_API bool GetFontParameterValue(const FMaterialParameterInfo& ParameterInfo, class UFont*& OutFontValue, int32& OutFontPage, bool bOveriddenOnly = false) const override;
 	virtual ENGINE_API void GetUsedTextures(TArray<UTexture*>& OutTextures, EMaterialQualityLevel::Type QualityLevel, bool bAllQualityLevels, ERHIFeatureLevel::Type FeatureLevel, bool bAllFeatureLevels) const override;
 	virtual ENGINE_API void GetUsedTexturesAndIndices(TArray<UTexture*>& OutTextures, TArray< TArray<int32> >& OutIndices, EMaterialQualityLevel::Type QualityLevel, ERHIFeatureLevel::Type FeatureLevel) const;
diff --git a/Engine/Source/Runtime/Engine/Classes/Materials/MaterialInterface.h b/Engine/Source/Runtime/Engine/Classes/Materials/MaterialInterface.h
--- a/Engine/Source/Runtime/Engine/Classes/Materials/MaterialInterface.h
+++ b/Engine/Source/Runtime/Engine/Classes/Materials/MaterialInterface.h
@@ -32,6 +32,9 @@ class UTexture;
 struct FPrimitiveViewRelevance;
 struct FMaterialParameterInfo;
 struct FMaterialResourceLocOnDisk;
+#if WITH_EDITORONLY_DATA
+struct FParameterChannelNames;
+#endif
 
 typedef TArray<FMaterialResource*> FMaterialResourceDeferredDeletionArray;
 
@@ -695,10 +698,16 @@ public:
 	ENGINE_API virtual bool GetScalarCurveParameterValue(const FMaterialParameterInfo& ParameterInfo, FInterpCurveFloat& OutValue) const;
 	ENGINE_API virtual bool GetVectorParameterValue(const FMaterialParameterInfo& ParameterInfo, FLinearColor& OutValue, bool bOveriddenOnly = false) const;
 	ENGINE_API virtual bool IsVectorParameterUsedAsChannelMask(const FMaterialParameterInfo& ParameterInfo, bool& OutValue) const;
+#if WITH_EDITOR
+	ENGINE_API virtual bool GetVectorParameterChannelNames(const FMaterialParameterInfo& ParameterInfo, FParameterChannelNames& OutValue) const;
+#endif
 	ENGINE_API virtual bool GetVectorCurveParameterValue(const FMaterialParameterInfo& ParameterInfo, FInterpCurveVector& OutValue) const;
 	ENGINE_API virtual bool GetLinearColorParameterValue(const FMaterialParameterInfo& ParameterInfo, FLinearColor& OutValue) const;
 	ENGINE_API virtual bool GetLinearColorCurveParameterValue(const FMaterialParameterInfo& ParameterInfo, FInterpCurveLinearColor& OutValue) const;
 	ENGINE_API virtual bool GetTextureParameterValue(const FMaterialParameterInfo& ParameterInfo, class UTexture*& OutValue, bool bOveriddenOnly = false) const;
+#if WITH_EDITOR
+	ENGINE_API virtual bool GetTextureParameterChannelNames(const FMaterialParameterInfo& ParameterInfo, FParameterChannelNames& OutValue) const;
+#endif
 	ENGINE_API virtual bool GetFontParameterValue(const FMaterialParameterInfo& ParameterInfo,class UFont*& OutFontValue, int32& OutFontPage, bool bOveriddenOnly = false) const;
 	ENGINE_API virtual bool GetRefractionSettings(float& OutBiasValue) const;
 
diff --git a/Engine/Source/Runtime/Engine/Private/Materials/Material.cpp b/Engine/Source/Runtime/Engine/Private/Materials/Material.cpp
--- a/Engine/Source/Runtime/Engine/Private/Materials/Material.cpp
+++ b/Engine/Source/Runtime/Engine/Private/Materials/Material.cpp
@@ -2340,6 +2340,61 @@ bool UMaterial::IsVectorParameterUsedAsChannelMask(const FMaterialParameterInfo&
 	return false;
 }
 
+#if WITH_EDITOR
+bool UMaterial::GetVectorParameterChannelNames(const FMaterialParameterInfo& ParameterInfo, FParameterChannelNames& OutValue) const
+{
+	// In the case of duplicate parameters with different values, this will return the
+	// first matching expression found, not necessarily the one that's used for rendering
+	UMaterialExpressionVectorParameter* Parameter = nullptr;
+	FLinearColor TempColor;
+
+	for (UMaterialExpression* Expression : Expressions)
+	{
+		// Only need to check parameters that match in associated scope
+		if (ParameterInfo.Association == EMaterialParameterAssociation::GlobalParameter)
+		{
+			if (UMaterialExpressionVectorParameter* ExpressionParameter = Cast<UMaterialExpressionVectorParameter>(Expression))
+			{
+				if (ExpressionParameter->IsNamedParameter(ParameterInfo, TempColor))
+				{
+					OutValue = ExpressionParameter->GetVectorChannelNames();
+					return true;
+				}
+			}
+			else if (UMaterialExpressionMaterialFunctionCall* FunctionCall = Cast<UMaterialExpressionMaterialFunctionCall>(Expression))
+			{
+				UMaterialFunctionInterface* Function = FunctionCall->MaterialFunction;
+				UMaterialFunctionInterface* ParameterOwner = nullptr;
+
+				if (Function && Function->GetNamedParameterOfType(ParameterInfo, Parameter, &ParameterOwner))
+				{
+					OutValue = Parameter->GetVectorChannelNames();
+					return true;
+				}
+			}
+		}
+		else
+		{
+			if (UMaterialExpressionMaterialAttributeLayers* LayersExpression = Cast<UMaterialExpressionMaterialAttributeLayers>(Expression))
+			{
+				UMaterialFunctionInterface* Function = LayersExpression->GetParameterAssociatedFunction(ParameterInfo);
+				UMaterialFunctionInterface* ParameterOwner = nullptr;
+
+				if (Function && Function->GetNamedParameterOfType(ParameterInfo, Parameter, &ParameterOwner))
+				{
+					OutValue = Parameter->GetVectorChannelNames();
+					return true;
+				}
+			}
+		}
+	}
+
+	return false;
+}
+
+#endif
+
+
 bool UMaterial::GetTextureParameterValue(const FMaterialParameterInfo& ParameterInfo, UTexture*& OutValue, bool bOveriddenOnly) const
 {
 	if (bOveriddenOnly && !AreExperimentalMaterialLayersEnabled())
@@ -2410,6 +2465,61 @@ bool UMaterial::GetTextureParameterValue(const FMaterialParameterInfo& Parameter
 	return false;
 }
 
+
+#if WITH_EDITOR
+bool UMaterial::GetTextureParameterChannelNames(const FMaterialParameterInfo& ParameterInfo, FParameterChannelNames& OutValue) const
+{
+	// In the case of duplicate parameters with different values, this will return the
+	// first matching expression found, not necessarily the one that's used for rendering
+	UMaterialExpressionTextureSampleParameter* Parameter = nullptr;
+	UTexture* TempTexture;
+
+	for (UMaterialExpression* Expression : Expressions)
+	{
+		// Only need to check parameters that match in associated scope
+		if (ParameterInfo.Association == EMaterialParameterAssociation::GlobalParameter)
+		{
+			if (UMaterialExpressionTextureSampleParameter* ExpressionParameter = Cast<UMaterialExpressionTextureSampleParameter>(Expression))
+			{
+				if (ExpressionParameter->IsNamedParameter(ParameterInfo, TempTexture))
+				{
+					OutValue = ExpressionParameter->GetTextureChannelNames();
+					return true;
+				}
+			}
+			else if (UMaterialExpressionMaterialFunctionCall* FunctionCall = Cast<UMaterialExpressionMaterialFunctionCall>(Expression))
+			{
+				UMaterialFunctionInterface* Function = FunctionCall->MaterialFunction;
+				UMaterialFunctionInterface* ParameterOwner = nullptr;
+
+				if (Function && Function->GetNamedParameterOfType(ParameterInfo, Parameter, &ParameterOwner))
+				{
+					OutValue = Parameter->GetTextureChannelNames();
+					return true;
+				}
+			}
+		}
+		else
+		{
+			if (UMaterialExpressionMaterialAttributeLayers* LayersExpression = Cast<UMaterialExpressionMaterialAttributeLayers>(Expression))
+			{
+				UMaterialFunctionInterface* Function = LayersExpression->GetParameterAssociatedFunction(ParameterInfo);
+				UMaterialFunctionInterface* ParameterOwner = nullptr;
+
+				if (Function && Function->GetNamedParameterOfType(ParameterInfo, Parameter, &ParameterOwner))
+				{
+					OutValue = Parameter->GetTextureChannelNames();
+					return true;
+				}
+			}
+		}
+	}
+
+	return false;
+}
+
+#endif
+
 bool UMaterial::GetFontParameterValue(const FMaterialParameterInfo& ParameterInfo, UFont*& OutFontValue, int32& OutFontPage, bool bOveriddenOnly) const
 {
 	// In the case of duplicate parameters with different values, this will return the
diff --git a/Engine/Source/Runtime/Engine/Private/Materials/MaterialExpressions.cpp b/Engine/Source/Runtime/Engine/Private/Materials/MaterialExpressions.cpp
--- a/Engine/Source/Runtime/Engine/Private/Materials/MaterialExpressions.cpp
+++ b/Engine/Source/Runtime/Engine/Private/Materials/MaterialExpressions.cpp
@@ -1954,6 +1954,7 @@ UMaterialExpressionTextureSampleParameter::UMaterialExpressionTextureSampleParam
 	MenuCategories.Empty();
 	MenuCategories.Add( ConstructorStatics.NAME_Obsolete);
 	SortPriority = 0;
+	ApplyChannelNames();
 #endif
 }
 
@@ -2036,6 +2037,36 @@ bool UMaterialExpressionTextureSampleParameter::SetParameterValue(FName InParame
 	}
 	return false;
 }
+
+void UMaterialExpressionTextureSampleParameter::PostEditChangeProperty(FPropertyChangedEvent& PropertyChangedEvent)
+{
+	UProperty* PropertyThatChanged = PropertyChangedEvent.MemberProperty;
+	const FString PropertyName = PropertyThatChanged ? PropertyThatChanged->GetName() : TEXT("");
+
+	if (PropertyName == GET_MEMBER_NAME_STRING_CHECKED(UMaterialExpressionTextureSampleParameter, ChannelNames))
+	{
+		ApplyChannelNames();
+
+		if (GraphNode)
+		{
+			GraphNode->ReconstructNode();
+		}
+	}
+
+	Super::PostEditChangeProperty(PropertyChangedEvent);
+}
+
+void UMaterialExpressionTextureSampleParameter::ApplyChannelNames()
+{
+	static const FName Red("R");
+	static const FName Green("G");
+	static const FName Blue("B");
+	static const FName Alpha("A");
+	Outputs[1].OutputName = !ChannelNames.R.IsEmpty() ? FName(*ChannelNames.R.ToString()) : Red;
+	Outputs[2].OutputName = !ChannelNames.G.IsEmpty() ? FName(*ChannelNames.G.ToString()) : Green;
+	Outputs[3].OutputName = !ChannelNames.B.IsEmpty() ? FName(*ChannelNames.B.ToString()) : Blue;
+	Outputs[4].OutputName = !ChannelNames.A.IsEmpty() ? FName(*ChannelNames.A.ToString()) : Alpha;
+}
 #endif
 
 bool UMaterialExpressionTextureSampleParameter::TextureIsValid( UTexture* /*InTexture*/ )
@@ -6398,6 +6429,7 @@ UMaterialExpressionVectorParameter::UMaterialExpressionVectorParameter(const FOb
 	Outputs.Add(FExpressionOutput(TEXT(""), 1, 0, 1, 0, 0));
 	Outputs.Add(FExpressionOutput(TEXT(""), 1, 0, 0, 1, 0));
 	Outputs.Add(FExpressionOutput(TEXT(""), 1, 0, 0, 0, 1));
+	ApplyChannelNames();
 #endif // WITH_EDITORONLY_DATA
 }
 
@@ -6497,10 +6529,29 @@ void UMaterialExpressionVectorParameter::PostEditChangeProperty(FPropertyChanged
 		const int32 PrimDataIndex = PrimitiveDataIndex;
 		PrimitiveDataIndex = (uint8)FMath::Clamp(PrimDataIndex, 0, FCustomPrimitiveData::NumCustomPrimitiveDataFloats-1);
 	}
+	else if (PropertyName == GET_MEMBER_NAME_STRING_CHECKED(UMaterialExpressionVectorParameter, ChannelNames))
+	{
+		ApplyChannelNames();
+
+		if (GraphNode)
+		{
+			GraphNode->ReconstructNode();
+		}
+	}
 
 	Super::PostEditChangeProperty(PropertyChangedEvent);
 }
 
+void UMaterialExpressionVectorParameter::ApplyChannelNames()
+{
+	Outputs[1].OutputName = FName(*ChannelNames.R.ToString());
+	Outputs[2].OutputName = FName(*ChannelNames.G.ToString());
+	Outputs[3].OutputName = FName(*ChannelNames.B.ToString());
+	Outputs[4].OutputName = FName(*ChannelNames.A.ToString());
+	bShowOutputNameOnPin = !ChannelNames.R.IsEmpty() || !ChannelNames.G.IsEmpty() || !ChannelNames.B.IsEmpty() || !ChannelNames.A.IsEmpty();
+}
+
+
 void UMaterialExpressionVectorParameter::ValidateParameterName(const bool bAllowDuplicateName)
 {
 	bool bOverrideDuplicateBehavior = false;
diff --git a/Engine/Source/Runtime/Engine/Private/Materials/MaterialInstance.cpp b/Engine/Source/Runtime/Engine/Private/Materials/MaterialInstance.cpp
--- a/Engine/Source/Runtime/Engine/Private/Materials/MaterialInstance.cpp
+++ b/Engine/Source/Runtime/Engine/Private/Materials/MaterialInstance.cpp
@@ -795,6 +795,47 @@ bool UMaterialInstance::IsVectorParameterUsedAsChannelMask(const FMaterialParame
 	return false;
 }
 
+#if WITH_EDITOR
+bool  UMaterialInstance::GetVectorParameterChannelNames(const FMaterialParameterInfo& ParameterInfo, FParameterChannelNames& OutValue) const
+{
+	bool bFoundAValue = false;
+
+	if (GetReentrantFlag())
+	{
+		return false;
+	}
+
+	// Instance-included default
+	if (ParameterInfo.Association != EMaterialParameterAssociation::GlobalParameter)
+	{
+		UMaterialExpressionVectorParameter* Parameter = nullptr;
+		for (const FStaticMaterialLayersParameter& LayersParam : StaticParameters.MaterialLayersParameters)
+		{
+			if (LayersParam.bOverride)
+			{
+				UMaterialFunctionInterface* Function = LayersParam.GetParameterAssociatedFunction(ParameterInfo);
+				UMaterialFunctionInterface* ParameterOwner = nullptr;
+
+				if (Function && Function->GetNamedParameterOfType(ParameterInfo, Parameter, &ParameterOwner))
+				{
+					OutValue = Parameter->GetVectorChannelNames();
+					return true;
+				}
+			}
+		}
+	}
+
+	// Next material in hierarchy
+	if (Parent)
+	{
+		FMICReentranceGuard	Guard(this);
+		return Parent->GetVectorParameterChannelNames(ParameterInfo, OutValue);
+	}
+
+	return false;
+};
+#endif
+
 bool UMaterialInstance::GetTextureParameterValue(const FMaterialParameterInfo& ParameterInfo, UTexture*& OutValue, bool bOveriddenOnly) const
 {
 	bool bFoundAValue = false;
@@ -850,6 +891,48 @@ bool UMaterialInstance::GetTextureParameterValue(const FMaterialParameterInfo& P
 	return false;
 }
 
+
+#if WITH_EDITOR
+bool  UMaterialInstance::GetTextureParameterChannelNames(const FMaterialParameterInfo& ParameterInfo, FParameterChannelNames& OutValue) const
+{
+	bool bFoundAValue = false;
+
+	if (GetReentrantFlag())
+	{
+		return false;
+	}
+
+	// Instance-included default
+	if (ParameterInfo.Association != EMaterialParameterAssociation::GlobalParameter)
+	{
+		UMaterialExpressionTextureSampleParameter* Parameter = nullptr;
+		for (const FStaticMaterialLayersParameter& LayersParam : StaticParameters.MaterialLayersParameters)
+		{
+			if (LayersParam.bOverride)
+			{
+				UMaterialFunctionInterface* Function = LayersParam.GetParameterAssociatedFunction(ParameterInfo);
+				UMaterialFunctionInterface* ParameterOwner = nullptr;
+
+				if (Function && Function->GetNamedParameterOfType(ParameterInfo, Parameter, &ParameterOwner))
+				{
+					OutValue = Parameter->GetTextureChannelNames();
+					return true;
+				}
+			}
+		}
+	}
+
+	// Next material in hierarchy
+	if (Parent)
+	{
+		FMICReentranceGuard	Guard(this);
+		return Parent->GetTextureParameterChannelNames(ParameterInfo, OutValue);
+	}
+
+	return false;
+};
+#endif
+
 bool UMaterialInstance::GetFontParameterValue(const FMaterialParameterInfo& ParameterInfo,class UFont*& OutFontValue, int32& OutFontPage, bool bOveriddenOnly) const
 {
 	bool bFoundAValue = false;
diff --git a/Engine/Source/Runtime/Engine/Private/Materials/MaterialInterface.cpp b/Engine/Source/Runtime/Engine/Private/Materials/MaterialInterface.cpp
--- a/Engine/Source/Runtime/Engine/Private/Materials/MaterialInterface.cpp
+++ b/Engine/Source/Runtime/Engine/Private/Materials/MaterialInterface.cpp
@@ -294,6 +294,11 @@ bool UMaterialInterface::IsVectorParameterUsedAsChannelMask(const FMaterialParam
 }
 
 #if WITH_EDITOR
+bool UMaterialInterface::GetVectorParameterChannelNames(const FMaterialParameterInfo& ParameterInfo, FParameterChannelNames& OutValue) const
+{
+	return false;
+}
+
 bool UMaterialInterface::GetScalarParameterSliderMinMax(const FMaterialParameterInfo& ParameterInfo, float& OutSliderMin, float& OutSliderMax) const
 {
 	return false;
@@ -336,6 +341,13 @@ bool UMaterialInterface::GetTextureParameterValue(const FMaterialParameterInfo&
 	return false;
 }
 
+#if WITH_EDITOR
+bool UMaterialInterface::GetTextureParameterChannelNames(const FMaterialParameterInfo& ParameterInfo, FParameterChannelNames& OutValue) const
+{
+	return false;
+}
+#endif
+
 bool UMaterialInterface::GetFontParameterValue(const FMaterialParameterInfo& ParameterInfo, class UFont*& OutFontValue, int32& OutFontPage, bool bOveriddenOnly) const
 {
 	return false;
