Fixing missing include for non-unity/linux compile

diff --git a/Engine/Source/Editor/MaterialEditor/Private/MaterialEditor.cpp b/Engine/Source/Editor/MaterialEditor/Private/MaterialEditor.cpp
--- a/Engine/Source/Editor/MaterialEditor/Private/MaterialEditor.cpp
+++ b/Engine/Source/Editor/MaterialEditor/Private/MaterialEditor.cpp
@@ -127,6 +127,7 @@
 
 #include "SMaterialParametersOverviewWidget.h"
 #include "IPropertyRowGenerator.h"
+#include "TextProperty.h"
 
 #define LOCTEXT_NAMESPACE "MaterialEditor"
 
