Fix broken include

diff --git a/Engine/Source/Editor/MaterialEditor/Private/MaterialEditor.cpp b/Engine/Source/Editor/MaterialEditor/Private/MaterialEditor.cpp
--- a/Engine/Source/Editor/MaterialEditor/Private/MaterialEditor.cpp
+++ b/Engine/Source/Editor/MaterialEditor/Private/MaterialEditor.cpp
@@ -127,7 +127,7 @@
 
 #include "SMaterialParametersOverviewWidget.h"
 #include "IPropertyRowGenerator.h"
-#include "TextProperty.h"
+#include "UObject/TextProperty.h"
 
 #define LOCTEXT_NAMESPACE "MaterialEditor"
 
