Fix assert from missing EndSection() call in material editor toolbar

diff --git a/Engine/Source/Editor/MaterialEditor/Private/MaterialEditor.cpp b/Engine/Source/Editor/MaterialEditor/Private/MaterialEditor.cpp
--- a/Engine/Source/Editor/MaterialEditor/Private/MaterialEditor.cpp
+++ b/Engine/Source/Editor/MaterialEditor/Private/MaterialEditor.cpp
@@ -1183,6 +1183,7 @@ void FMaterialEditor::FillToolbar(FToolBarBuilder& ToolbarBuilder)
 				false
 			);
 		}
+		ToolbarBuilder.EndSection();
 	}
 };
 
