Material Instance Editor Hierarchy tool bar button refactored to avoid needing submenu for each item

diff --git a/Engine/Source/Editor/MaterialEditor/MaterialEditor.Build.cs b/Engine/Source/Editor/MaterialEditor/MaterialEditor.Build.cs
--- a/Engine/Source/Editor/MaterialEditor/MaterialEditor.Build.cs
+++ b/Engine/Source/Editor/MaterialEditor/MaterialEditor.Build.cs
@@ -42,7 +42,8 @@ public class MaterialEditor : ModuleRules
                 "AdvancedPreviewScene",
                 "Projects",
                 "AssetRegistry",
-            }
+				"EditorMenus",
+			}
 		);
 
 		DynamicallyLoadedModuleNames.AddRange(
diff --git a/Engine/Source/Editor/MaterialEditor/Private/MaterialInstanceEditor.cpp b/Engine/Source/Editor/MaterialEditor/Private/MaterialInstanceEditor.cpp
--- a/Engine/Source/Editor/MaterialEditor/Private/MaterialInstanceEditor.cpp
+++ b/Engine/Source/Editor/MaterialEditor/Private/MaterialInstanceEditor.cpp
@@ -19,6 +19,8 @@
 #include "Materials/MaterialFunctionInstance.h"
 #include "Materials/MaterialExpressionFunctionOutput.h"
 #include "MaterialEditorModule.h"
+#include "EditorMenuSubsystem.h"
+#include "Toolkits/AssetEditorToolkitMenuContext.h"
 
 #include "Materials/MaterialExpressionTextureBase.h"
 #include "Materials/MaterialExpressionTextureSampleParameter.h"
@@ -771,47 +773,42 @@ void FMaterialInstanceEditor::UpdatePreviewViewportsVisibility()
 	}
 }
 
-void FMaterialInstanceEditor::FillToolbar(FToolBarBuilder& ToolbarBuilder)
+void FMaterialInstanceEditor::RegisterToolBar()
 {
-	ToolbarBuilder.BeginSection("Apply");
+	UEditorMenuSubsystem* EditorMenus = UEditorMenuSubsystem::Get();
+	UEditorMenu* ToolBar = EditorMenus->ExtendMenu(GetEditorMenuToolbarName());
+
+	FEditorMenuInsert InsertAfterAssetSection("Asset", EEditorMenuInsertType::After);
 	{
-		ToolbarBuilder.AddToolBarButton(FMaterialEditorCommands::Get().Apply);
+		FEditorMenuSection& Section = ToolBar->AddSection("Apply", TAttribute<FText>(), InsertAfterAssetSection);
+		Section.AddEntry(FEditorMenuEntry::InitToolBarButton(FMaterialEditorCommands::Get().Apply));
 	}
-	ToolbarBuilder.EndSection();
 
-	ToolbarBuilder.BeginSection("Command");
 	{
-		ToolbarBuilder.AddToolBarButton(FMaterialEditorCommands::Get().ShowAllMaterialParameters);
+		FEditorMenuSection& Section = ToolBar->AddSection("Command", TAttribute<FText>(), InsertAfterAssetSection);
+		Section.AddEntry(FEditorMenuEntry::InitToolBarButton(FMaterialEditorCommands::Get().ShowAllMaterialParameters));
 		// TODO: support in material instance editor.
-		ToolbarBuilder.AddToolBarButton(FMaterialEditorCommands::Get().TogglePlatformStats);
+		Section.AddEntry(FEditorMenuEntry::InitToolBarButton(FMaterialEditorCommands::Get().TogglePlatformStats));
 	}
-	ToolbarBuilder.EndSection();
-	ToolbarBuilder.BeginSection("Parent");
+	
 	{
-		ToolbarBuilder.AddComboButton(
-			FUIAction(),
-			FOnGetContent::CreateSP(this, &FMaterialInstanceEditor::GenerateInheritanceMenu),
+		FEditorMenuSection& Section = ToolBar->AddSection("Parent", TAttribute<FText>(), InsertAfterAssetSection);
+		Section.AddEntry(FEditorMenuEntry::InitComboButton(
+			"Hierarchy",
+			FEditorUIActionChoice(),
+			FNewEditorMenuDelegate::CreateSP(this, &FMaterialInstanceEditor::GenerateInheritanceMenu),
 			LOCTEXT("Hierarchy", "Hierarchy"),
 			FText::GetEmpty(),
-			FSlateIcon(FSlateIcon(FEditorStyle::GetStyleSetName(), "BTEditor.SwitchToBehaviorTreeMode")),
+			FSlateIcon(FEditorStyle::GetStyleSetName(), "BTEditor.SwitchToBehaviorTreeMode"),
 			false
-		);
+		));
 	}
-	ToolbarBuilder.EndSection();
+
 }
 
 void FMaterialInstanceEditor::ExtendToolbar()
 {
-
-	TSharedPtr<FExtender> ToolbarExtender = MakeShareable(new FExtender);
-	ToolbarExtender->AddToolBarExtension(
-		"Asset",
-		EExtensionHook::After,
-		GetToolkitCommands(),
-		FToolBarExtensionDelegate::CreateSP(this, &FMaterialInstanceEditor::FillToolbar)
-		);
-
-	AddToolbarExtender(ToolbarExtender);
+	RegisterToolBar();
 
 	AddToolbarExtender(GetToolBarExtensibilityManager()->GetAllExtenders(GetToolkitCommands(), GetEditingObjects()));
 
@@ -819,102 +816,129 @@ void FMaterialInstanceEditor::ExtendToolbar()
 	AddToolbarExtender(MaterialEditorModule->GetToolBarExtensibilityManager()->GetAllExtenders(GetToolkitCommands(), GetEditingObjects()));
 }
 
-
-TSharedRef< SWidget > FMaterialInstanceEditor::GenerateInheritanceMenu()
+void FMaterialInstanceEditor::GenerateInheritanceMenu(UEditorMenu* Menu)
 {
 	struct Local
 	{
-		static void MakeFunctionSubmenu(FMenuBuilder& MenuBuilder, FAssetData InFunction)
+		static void AddMenuEntry(FEditorMenuSection& Section, FAssetData AssetData, bool bIsFunctionPreviewMaterial)
 		{
-			MenuBuilder.AddMenuEntry(LOCTEXT("OpenInEditor", "Open In Editor"),
-				FText::GetEmpty(),
-				FSlateIcon(FEditorStyle::GetStyleSetName(), "ContentBrowser.AssetActions.OpenInExternalEditor"),
-				FUIAction(FExecuteAction::CreateStatic(&FMaterialEditorUtilities::OnOpenFunction, InFunction)));
-			
-			MenuBuilder.AddMenuEntry(LOCTEXT("FindInContentBrowser", "Find In Content Browser"),
-				FText::GetEmpty(),
-				FSlateIcon(FEditorStyle::GetStyleSetName(), "SystemWideCommands.FindInContentBrowser"),
-				FUIAction(FExecuteAction::CreateStatic(&FMaterialEditorUtilities::OnShowFunctionInContentBrowser, InFunction)));
-		}
+			FExecuteAction OpenAction;
+			FExecuteAction FindInContentBrowserAction;
+			if (bIsFunctionPreviewMaterial)
+			{
+				OpenAction.BindStatic(&FMaterialEditorUtilities::OnOpenFunction, AssetData);
+				FindInContentBrowserAction.BindStatic(&FMaterialEditorUtilities::OnShowFunctionInContentBrowser, AssetData);
+			}
+			else
+			{
+				OpenAction.BindStatic(&FMaterialEditorUtilities::OnOpenMaterial, AssetData);
+				FindInContentBrowserAction.BindStatic(&FMaterialEditorUtilities::OnShowMaterialInContentBrowser, AssetData);
+			}
 
-		static void MakeMaterialSubmenu(FMenuBuilder& MenuBuilder, FAssetData InMaterial)
-		{
-			MenuBuilder.AddMenuEntry(LOCTEXT("OpenInEditor", "Open In Editor"),
-				FText::GetEmpty(),
-				FSlateIcon(FEditorStyle::GetStyleSetName(), "ContentBrowser.AssetActions.OpenInExternalEditor"),
-				FUIAction(FExecuteAction::CreateStatic(&FMaterialEditorUtilities::OnOpenMaterial, InMaterial)));
-			
-			MenuBuilder.AddMenuEntry(LOCTEXT("FindInContentBrowser", "Find In Content Browser"),
-				FText::GetEmpty(),
-				FSlateIcon(FEditorStyle::GetStyleSetName(), "SystemWideCommands.FindInContentBrowser"),
-				FUIAction(FExecuteAction::CreateStatic(&FMaterialEditorUtilities::OnShowMaterialInContentBrowser, InMaterial)));
+			FFormatNamedArguments Args;
+			Args.Add(TEXT("ParentName"), FText::FromName(AssetData.AssetName));
+			FText Label = FText::Format(LOCTEXT("InstanceParentName", "{ParentName}"), Args);
+
+			FSlateIcon OpenIcon(FEditorStyle::GetStyleSetName(), "ContentBrowser.AssetActions.OpenInExternalEditor");
+			FSlateIcon FindInContentBrowserIcon(FEditorStyle::GetStyleSetName(), "SystemWideCommands.FindInContentBrowser");
+
+			TSharedRef<SWidget> EntryWidget =
+				SNew(SHorizontalBox)
+				.ToolTipText(LOCTEXT("OpenInEditor", "Open In Editor"))
+
+				+ SHorizontalBox::Slot()
+				.AutoWidth()
+				.Padding(FMargin(2, 0, 2, 0))
+				[
+					SNew( SBox )
+					.WidthOverride( MultiBoxConstants::MenuIconSize + 2 )
+					.HeightOverride( MultiBoxConstants::MenuIconSize )
+					.HAlign(HAlign_Center)
+					.VAlign(VAlign_Center)
+					[
+						SNew( SBox )
+						.WidthOverride( MultiBoxConstants::MenuIconSize )
+						.HeightOverride( MultiBoxConstants::MenuIconSize )
+						[
+							SNew(SImage)
+							.Image(OpenIcon.GetIcon())
+						]
+					]
+				]
+
+				+ SHorizontalBox::Slot()
+				.FillWidth( 1.0f )
+				.Padding(FMargin(2, 0, 6, 0))
+				.VAlign( VAlign_Center )
+				[
+					SNew(STextBlock)
+					.Text(Label)
+				]
+
+				+ SHorizontalBox::Slot()
+				.AutoWidth()
+				.VAlign( VAlign_Center )
+				.HAlign( HAlign_Right )
+				[
+					SNew(SButton)
+					.ButtonStyle(FEditorStyle::Get(), "HoverHintOnly")
+					.ToolTipText(LOCTEXT("FindInContentBrowser", "Find In Content Browser"))
+					.OnClicked_Lambda([FindInContentBrowserAction]() { FindInContentBrowserAction.ExecuteIfBound(); return FReply::Handled(); })
+					[
+						SNew( SBox )
+						.WidthOverride( MultiBoxConstants::MenuIconSize + 2 )
+						.HeightOverride( MultiBoxConstants::MenuIconSize )
+						.HAlign(HAlign_Center)
+						.VAlign(VAlign_Center)
+						[
+							SNew( SBox )
+							.WidthOverride( MultiBoxConstants::MenuIconSize )
+							.HeightOverride( MultiBoxConstants::MenuIconSize )
+							[
+								SNew(SImage)
+								.Image(FindInContentBrowserIcon.GetIcon())
+							]
+						]
+					]
+				];
+
+			Section.AddEntry(FEditorMenuEntry::InitMenuEntry(
+				NAME_None,
+				FUIAction(OpenAction),
+				EntryWidget
+			));
 		}
 	};
 
 	RebuildInheritanceList();
-	TSharedPtr<FExtender> ToolbarExtender = MakeShareable(new FExtender);
-	const bool bShouldCloseWindowAfterMenuSelection = true;
-	TSharedPtr<FUICommandList> InCommandList = MakeShareable(new FUICommandList);
-	FMenuBuilder MenuBuilder(bShouldCloseWindowAfterMenuSelection, InCommandList, ToolbarExtender);
+	Menu->bShouldCloseWindowAfterMenuSelection = true;
+	Menu->bSearchable = true;
 	const FName ParentName = TEXT("ParentChain");
-	MenuBuilder.BeginSection(ParentName, LOCTEXT("ParentChain", "Parent Chain"));
+	FEditorMenuSection& Section = Menu->AddSection(ParentName, LOCTEXT("ParentChain", "Parent Chain"));
 	if (bIsFunctionPreviewMaterial)
 	{
 		for (FAssetData FunctionParent : FunctionParentList)
 		{
-			FFormatNamedArguments Args;
-			Args.Add(TEXT("ParentName"), FText::FromName(FunctionParent.AssetName));
-			MenuBuilder.AddSubMenu(FText::Format(LOCTEXT("InstanceParentName", "{ParentName}"), Args),
-				FText::GetEmpty(),
-				FNewMenuDelegate::CreateStatic(&Local::MakeFunctionSubmenu, FunctionParent),
-				false,
-				FSlateIcon());
+			Local::AddMenuEntry(Section, FunctionParent, bIsFunctionPreviewMaterial);
 		}
 	}
 	else
 	{
 		for (FAssetData MaterialParent : MaterialParentList)
 		{
-			FFormatNamedArguments Args;
-			Args.Add(TEXT("ParentName"), FText::FromName(MaterialParent.AssetName));
-			MenuBuilder.AddSubMenu(FText::Format(LOCTEXT("InstanceParentName", "{ParentName}"), Args),
-				FText::GetEmpty(),
-				FNewMenuDelegate::CreateStatic(&Local::MakeMaterialSubmenu, MaterialParent),
-				false,
-				FSlateIcon());
+			Local::AddMenuEntry(Section, MaterialParent, bIsFunctionPreviewMaterial);
 		}
 	}
-	MenuBuilder.EndSection();
 
 	if (!bIsFunctionPreviewMaterial)
 	{
 		const FName MaterialInstances = TEXT("MaterialInstances");
-		MenuBuilder.BeginSection(MaterialInstances, LOCTEXT("MaterialInstances", "Material Instances"));
+		FEditorMenuSection& MaterialInstancesSection = Menu->AddSection(MaterialInstances, LOCTEXT("MaterialInstances", "Material Instances"));
 		for (FAssetData MaterialChild : MaterialChildList)
 		{
-			FFormatNamedArguments Args;
-			Args.Add(TEXT("ChildName"), FText::FromName(MaterialChild.AssetName));
-			MenuBuilder.AddSubMenu(FText::Format(LOCTEXT("MaterialInstanceName", "{ChildName}"), Args),
-				FText::GetEmpty(),
-				FNewMenuDelegate::CreateStatic(&Local::MakeMaterialSubmenu, MaterialChild),
-				false,
-				FSlateIcon());
+			Local::AddMenuEntry(MaterialInstancesSection, MaterialChild, bIsFunctionPreviewMaterial);
 		}
-		MenuBuilder.EndSection();
 	}
-
-	TSharedRef<SWidget> ConstrainedMenu = SNew(SVerticalBox)
-		+SVerticalBox::Slot()
-		.MaxHeight(500.0f)
-		[		
-			SNew(SScrollBox)
-			+ SScrollBox::Slot()
-			[
-				MenuBuilder.MakeWidget()
-			]
-		];
-
-	return ConstrainedMenu;
 }
 
 
diff --git a/Engine/Source/Editor/MaterialEditor/Private/MaterialInstanceEditor.h b/Engine/Source/Editor/MaterialEditor/Private/MaterialInstanceEditor.h
--- a/Engine/Source/Editor/MaterialEditor/Private/MaterialInstanceEditor.h
+++ b/Engine/Source/Editor/MaterialEditor/Private/MaterialInstanceEditor.h
@@ -15,6 +15,7 @@
 #include "TickableEditorObject.h"
 
 class FCanvas;
+class UEditorMenu;
 class UMaterialEditorInstanceConstant;
 class UMaterialInterface;
 class UMaterialInstanceConstant;
@@ -153,11 +154,11 @@ private:
 	/** Updates the 3D and UI preview viewport visibility based on material domain */
 	void UpdatePreviewViewportsVisibility();
 
-	void FillToolbar(FToolBarBuilder& ToolbarBuilder);
+	void RegisterToolBar();
 	/** Builds the toolbar widget for the material editor */
 	void ExtendToolbar();
 
-	TSharedRef<SWidget> GenerateInheritanceMenu();
+	void GenerateInheritanceMenu(UEditorMenu* Menu);
 
 	/** If re-initializing for a material function instance re-generate the proxy materials */
 	void ReInitMaterialFunctionProxies();
