Add include in MaterialInstanceEditor.cpp
#rnx

diff --git a/Engine/Source/Editor/MaterialEditor/Private/MaterialInstanceEditor.cpp b/Engine/Source/Editor/MaterialEditor/Private/MaterialInstanceEditor.cpp
--- a/Engine/Source/Editor/MaterialEditor/Private/MaterialInstanceEditor.cpp
+++ b/Engine/Source/Editor/MaterialEditor/Private/MaterialInstanceEditor.cpp
@@ -45,6 +45,7 @@
 #include "MaterialStats.h"
 #include "MaterialEditingLibrary.h"
 #include "Widgets/Layout/SScrollBox.h"
+#include "Widgets/Input/SButton.h"
 
 #define LOCTEXT_NAMESPACE "MaterialInstanceEditor"
 
