PR #5862: Reword 'override original material' dialog to be clearer (Contributed by [REDACTED_PERSON])

diff --git a/Engine/Source/Editor/MaterialEditor/Private/MaterialEditor.cpp b/Engine/Source/Editor/MaterialEditor/Private/MaterialEditor.cpp
--- a/Engine/Source/Editor/MaterialEditor/Private/MaterialEditor.cpp
+++ b/Engine/Source/Editor/MaterialEditor/Private/MaterialEditor.cpp
@@ -699,7 +699,7 @@ void FMaterialEditor::InitMaterialEditor( const EToolkitMode::Type Mode, const T
 	// Store the name of this material (for the tutorial widget meta)
 	if (OriginalMaterial != nullptr)
 	{
-	Material->MaterialGraph->OriginalMaterialFullName = OriginalMaterial->GetName();
+		Material->MaterialGraph->OriginalMaterialFullName = OriginalMaterial->GetName();
 	}
 	Material->MaterialGraph->RebuildGraph();
 	RecenterEditor();
@@ -1345,7 +1345,7 @@ bool FMaterialEditor::OnRequestClose()
 		// find out the user wants to do with this dirty material
 		EAppReturnType::Type YesNoCancelReply = FMessageDialog::Open(EAppMsgType::YesNoCancel,
 			FText::Format(
-				NSLOCTEXT("UnrealEd", "Prompt_MaterialEditorClose", "Would you like to apply changes to this material to the original material?\n{0}\n(No will lose all changes!)"),
+				NSLOCTEXT("UnrealEd", "Prompt_MaterialEditorClose", "Would you like to apply the changes of the modified material to the original material?\n{0}\n(Selecting 'No' will cause all changes to be lost!)"),
 				FText::FromString(OriginalMaterialObject->GetPathName()) ));
 
 		// act on it
