Disabled Substance tutorial in dev / internal. Debug code was checked in.

diff --git a/Engine/Source/Editor/MaterialEditor/Private/MaterialEditorModule.cpp b/Engine/Source/Editor/MaterialEditor/Private/MaterialEditorModule.cpp
--- a/Engine/Source/Editor/MaterialEditor/Private/MaterialEditorModule.cpp
+++ b/Engine/Source/Editor/MaterialEditor/Private/MaterialEditorModule.cpp
@@ -139,7 +139,7 @@ public:
 				GetMutableDefault<UMaterialEditorSettings>());
 		}
 
-		if(true)//!FEngineBuildSettings::IsInternalBuild() && !FEngineBuildSettings::IsSourceDistribution())
+		if(!FEngineBuildSettings::IsInternalBuild() && !FEngineBuildSettings::IsSourceDistribution())
 		{
 			TSharedPtr<IPlugin> SubstancePlugin = IPluginManager::Get().FindPlugin(TEXT("Substance"));
 			if (!SubstancePlugin.IsValid() )
