Expose GetSelectedNodes in the material editor

diff --git a/Engine/Source/Editor/MaterialEditor/Private/MaterialEditingLibrary.cpp b/Engine/Source/Editor/MaterialEditor/Private/MaterialEditingLibrary.cpp
--- a/Engine/Source/Editor/MaterialEditor/Private/MaterialEditingLibrary.cpp
+++ b/Engine/Source/Editor/MaterialEditor/Private/MaterialEditingLibrary.cpp
@@ -702,6 +702,17 @@ bool UMaterialEditingLibrary::GetMaterialDefaultStaticSwitchParameterValue(UMate
 	return bResult;
 }
 
+TSet<UObject*> UMaterialEditingLibrary::GetMaterialSelectedNodes(UMaterial* Material)
+{
+	auto* MaterialEditor = (IMaterialEditor*)GEditor->GetEditorSubsystem<UAssetEditorSubsystem>()->FindEditorForAsset(Material, false);
+	if (MaterialEditor)
+	{
+		return MaterialEditor->GetSelectedNodes();
+	}
+
+	return TSet<UObject*>();
+}
+
 //////////////////////////////////////////////////////////////////////////
 
 int32 UMaterialEditingLibrary::GetNumMaterialExpressionsInFunction(const UMaterialFunction* MaterialFunction)
diff --git a/Engine/Source/Editor/MaterialEditor/Private/MaterialEditor.cpp b/Engine/Source/Editor/MaterialEditor/Private/MaterialEditor.cpp
--- a/Engine/Source/Editor/MaterialEditor/Private/MaterialEditor.cpp
+++ b/Engine/Source/Editor/MaterialEditor/Private/MaterialEditor.cpp
@@ -2659,7 +2659,7 @@ void FMaterialEditor::HideUnrelatedNodes()
 {
 	TArray<UMaterialGraphNode*> NodesToShow;
 
-	const FGraphPanelSelectionSet SelectedNodes = GraphEditor->GetSelectedNodes();
+	const FGraphPanelSelectionSet SelectedNodes = GetSelectedNodes();
 
 	for (FGraphPanelSelectionSet::TConstIterator NodeIt(SelectedNodes); NodeIt; ++NodeIt)
 	{
@@ -2796,7 +2796,7 @@ void FMaterialEditor::OnUseCurrentTexture()
 	if ( SelectedTexture )
 	{
 		const FScopedTransaction Transaction( NSLOCTEXT("UnrealEd", "UseCurrentTexture", "Use Current Texture") );
-		const FGraphPanelSelectionSet SelectedNodes = GraphEditor->GetSelectedNodes();
+		const FGraphPanelSelectionSet SelectedNodes = GetSelectedNodes();
 
 		for (FGraphPanelSelectionSet::TConstIterator NodeIt(SelectedNodes); NodeIt; ++NodeIt)
 		{
@@ -2821,7 +2821,7 @@ void FMaterialEditor::OnUseCurrentTexture()
 
 void FMaterialEditor::OnConvertObjects()
 {
-	const FGraphPanelSelectionSet SelectedNodes = GraphEditor->GetSelectedNodes();
+	const FGraphPanelSelectionSet SelectedNodes = GetSelectedNodes();
 	if (SelectedNodes.Num() > 0)
 	{
 		const FScopedTransaction Transaction( LOCTEXT("MaterialEditorConvert", "Material Editor: Convert") );
@@ -3000,7 +3000,7 @@ void FMaterialEditor::OnConvertObjects()
 
 void FMaterialEditor::OnConvertTextures()
 {
-	const FGraphPanelSelectionSet SelectedNodes = GraphEditor->GetSelectedNodes();
+	const FGraphPanelSelectionSet SelectedNodes = GetSelectedNodes();
 	if (SelectedNodes.Num() > 0)
 	{
 		const FScopedTransaction Transaction( LOCTEXT("MaterialEditorConvertTexture", "Material Editor: Convert to Texture") );
@@ -3085,7 +3085,7 @@ void FMaterialEditor::OnConvertTextures()
 
 void FMaterialEditor::OnPreviewNode()
 {
-	const FGraphPanelSelectionSet SelectedNodes = GraphEditor->GetSelectedNodes();
+	const FGraphPanelSelectionSet SelectedNodes = GetSelectedNodes();
 	if (SelectedNodes.Num() == 1)
 	{
 		for (FGraphPanelSelectionSet::TConstIterator NodeIt(SelectedNodes); NodeIt; ++NodeIt)
@@ -3102,7 +3102,7 @@ void FMaterialEditor::OnPreviewNode()
 
 void FMaterialEditor::OnToggleRealtimePreview()
 {
-	const FGraphPanelSelectionSet SelectedNodes = GraphEditor->GetSelectedNodes();
+	const FGraphPanelSelectionSet SelectedNodes = GetSelectedNodes();
 	if (SelectedNodes.Num() == 1)
 	{
 		for (FGraphPanelSelectionSet::TConstIterator NodeIt(SelectedNodes); NodeIt; ++NodeIt)
@@ -3131,7 +3131,7 @@ void FMaterialEditor::OnSelectDownstreamNodes()
 	TArray<UMaterialGraphNode*> CheckedNodes;
 	TArray<UMaterialGraphNode*> NodesToSelect;
 
-	const FGraphPanelSelectionSet SelectedNodes = GraphEditor->GetSelectedNodes();
+	const FGraphPanelSelectionSet SelectedNodes = GetSelectedNodes();
 
 	for (FGraphPanelSelectionSet::TConstIterator NodeIt(SelectedNodes); NodeIt; ++NodeIt)
 	{
@@ -3184,7 +3184,7 @@ void FMaterialEditor::OnSelectUpstreamNodes()
 	TArray<UMaterialGraphNode*> CheckedNodes;
 	TArray<UMaterialGraphNode*> NodesToSelect;
 
-	const FGraphPanelSelectionSet SelectedNodes = GraphEditor->GetSelectedNodes();
+	const FGraphPanelSelectionSet SelectedNodes = GetSelectedNodes();
 
 	for (FGraphPanelSelectionSet::TConstIterator NodeIt(SelectedNodes); NodeIt; ++NodeIt)
 	{
@@ -3376,7 +3376,7 @@ FString FMaterialEditor::GetDocLinkForSelectedNode()
 {
 	FString DocumentationLink;
 
-	TArray<UObject*> SelectedNodes = GraphEditor->GetSelectedNodes().Array();
+	TArray<UObject*> SelectedNodes = GetSelectedNodes().Array();
 	if (SelectedNodes.Num() == 1)
 	{
 		UMaterialGraphNode* SelectedGraphNode = Cast<UMaterialGraphNode>(SelectedNodes[0]);
@@ -3929,7 +3929,7 @@ bool FMaterialEditor::CanSelectAllNodes() const
 void FMaterialEditor::DeleteSelectedNodes()
 {
 	TArray<UEdGraphNode*> NodesToDelete;
-	const FGraphPanelSelectionSet SelectedNodes = GraphEditor->GetSelectedNodes();
+	const FGraphPanelSelectionSet SelectedNodes = GetSelectedNodes();
 
 	for (FGraphPanelSelectionSet::TConstIterator NodeIt(SelectedNodes); NodeIt; ++NodeIt)
 	{
@@ -4022,7 +4022,7 @@ void FMaterialEditor::DeleteNodes(const TArray<UEdGraphNode*>& NodesToDelete)
 
 bool FMaterialEditor::CanDeleteNodes() const
 {
-	const FGraphPanelSelectionSet SelectedNodes = GraphEditor->GetSelectedNodes();
+	const FGraphPanelSelectionSet SelectedNodes = GetSelectedNodes();
 	bool bDeletableNodeExists = false;
 
 	for (FGraphPanelSelectionSet::TConstIterator NodeIt(SelectedNodes); NodeIt; ++NodeIt)
@@ -4040,7 +4040,7 @@ bool FMaterialEditor::CanDeleteNodes() const
 void FMaterialEditor::DeleteSelectedDuplicatableNodes()
 {
 	// Cache off the old selection
-	const FGraphPanelSelectionSet OldSelectedNodes = GraphEditor->GetSelectedNodes();
+	const FGraphPanelSelectionSet OldSelectedNodes = GetSelectedNodes();
 
 	// Clear the selection and only select the nodes that can be duplicated
 	FGraphPanelSelectionSet RemainingNodes;
@@ -4077,7 +4077,7 @@ void FMaterialEditor::DeleteSelectedDuplicatableNodes()
 void FMaterialEditor::CopySelectedNodes()
 {
 	// Export the selected nodes and place the text on the clipboard
-	const FGraphPanelSelectionSet SelectedNodes = GraphEditor->GetSelectedNodes();
+	const FGraphPanelSelectionSet SelectedNodes = GetSelectedNodes();
 
 	FString ExportedText;
 
@@ -4109,7 +4109,7 @@ void FMaterialEditor::CopySelectedNodes()
 bool FMaterialEditor::CanCopyNodes() const
 {
 	// If any of the nodes can be duplicated then we should allow copying
-	const FGraphPanelSelectionSet SelectedNodes = GraphEditor->GetSelectedNodes();
+	const FGraphPanelSelectionSet SelectedNodes = GetSelectedNodes();
 	for (FGraphPanelSelectionSet::TConstIterator SelectedIter(SelectedNodes); SelectedIter; ++SelectedIter)
 	{
 		UEdGraphNode* Node = Cast<UEdGraphNode>(*SelectedIter);
@@ -4319,6 +4319,16 @@ int32 FMaterialEditor::GetNumberOfSelectedNodes() const
 	return GraphEditor->GetSelectedNodes().Num();
 }
 
+FGraphPanelSelectionSet FMaterialEditor::GetSelectedNodes() const
+{
+	FGraphPanelSelectionSet CurrentSelection;
+	if (GraphEditor.IsValid())
+	{
+		CurrentSelection = GraphEditor->GetSelectedNodes();
+	}
+	return CurrentSelection;
+}
+
 FMatExpressionPreview* FMaterialEditor::GetExpressionPreview(UMaterialExpression* InExpression)
 {
 	bool bNewlyCreated;
@@ -4500,7 +4510,7 @@ void FMaterialEditor::NotifyPostChange( const FPropertyChangedEvent& PropertyCha
 			UpdatePreviewViewportsVisibility();
 		}
 
-		FGraphPanelSelectionSet SelectedNodes = GraphEditor->GetSelectedNodes();
+		FGraphPanelSelectionSet SelectedNodes = GetSelectedNodes();
 
 		for (FGraphPanelSelectionSet::TConstIterator NodeIt(SelectedNodes); NodeIt; ++NodeIt)
 		{
@@ -5018,7 +5028,7 @@ bool FMaterialEditor::CheckExpressionRemovalWarnings(const TArray<UEdGraphNode*>
 
 void FMaterialEditor::RemoveSelectedExpressionFromFavorites()
 {
-	const FGraphPanelSelectionSet SelectedNodes = GraphEditor->GetSelectedNodes();
+	const FGraphPanelSelectionSet SelectedNodes = GetSelectedNodes();
 
 	if (SelectedNodes.Num() == 1)
 	{
@@ -5036,7 +5046,7 @@ void FMaterialEditor::RemoveSelectedExpressionFromFavorites()
 
 void FMaterialEditor::AddSelectedExpressionToFavorites()
 {
-	const FGraphPanelSelectionSet SelectedNodes = GraphEditor->GetSelectedNodes();
+	const FGraphPanelSelectionSet SelectedNodes = GetSelectedNodes();
 
 	if (SelectedNodes.Num() == 1)
 	{
diff --git a/Engine/Source/Editor/MaterialEditor/Private/MaterialEditor.h b/Engine/Source/Editor/MaterialEditor/Private/MaterialEditor.h
--- a/Engine/Source/Editor/MaterialEditor/Private/MaterialEditor.h
+++ b/Engine/Source/Editor/MaterialEditor/Private/MaterialEditor.h
@@ -378,6 +378,7 @@ public:
 	virtual bool CanPasteNodes() const override;
 	virtual void PasteNodesHere(const FVector2D& Location) override;
 	virtual int32 GetNumberOfSelectedNodes() const override;
+	virtual TSet<UObject*> GetSelectedNodes() const override;
 	virtual FMatExpressionPreview* GetExpressionPreview(UMaterialExpression* InExpression) override;
 	virtual void DeleteNodes(const TArray<class UEdGraphNode*>& NodesToDelete) override;
 
diff --git a/Engine/Source/Editor/MaterialEditor/Public/IMaterialEditor.h b/Engine/Source/Editor/MaterialEditor/Public/IMaterialEditor.h
--- a/Engine/Source/Editor/MaterialEditor/Public/IMaterialEditor.h
+++ b/Engine/Source/Editor/MaterialEditor/Public/IMaterialEditor.h
@@ -80,6 +80,9 @@ public:
 	/** Gets the number of selected nodes */
 	virtual int32 GetNumberOfSelectedNodes() const {return 0;}
 
+	/** Get the currently selected set of nodes */
+	virtual TSet<UObject*> GetSelectedNodes() const {return TSet<UObject*>();}
+
 	/**
 	 * Gets the preview for an expression
 	 *
diff --git a/Engine/Source/Editor/MaterialEditor/Public/MaterialEditingLibrary.h b/Engine/Source/Editor/MaterialEditor/Public/MaterialEditingLibrary.h
--- a/Engine/Source/Editor/MaterialEditor/Public/MaterialEditingLibrary.h
+++ b/Engine/Source/Editor/MaterialEditor/Public/MaterialEditingLibrary.h
@@ -135,6 +135,10 @@ public:
 	UFUNCTION(BlueprintPure, Category = "MaterialEditing")
 	static bool GetMaterialDefaultStaticSwitchParameterValue(UMaterial* Material, FName ParameterName);
 
+	/** Get the set of selected nodes from an active material editor */
+	UFUNCTION(BlueprintPure, Category = "MaterialEditing")
+	static TSet<UObject*> GetMaterialSelectedNodes(UMaterial* Material);
+
 	//////// MATERIAL FUNCTION EDITING
 
 	/** Returns number of material expressions in the supplied material */
