Expose GetBoundsForNode on Material Editor

diff --git a/Engine/Source/Editor/MaterialEditor/Private/MaterialEditor.cpp b/Engine/Source/Editor/MaterialEditor/Private/MaterialEditor.cpp
--- a/Engine/Source/Editor/MaterialEditor/Private/MaterialEditor.cpp
+++ b/Engine/Source/Editor/MaterialEditor/Private/MaterialEditor.cpp
@@ -4329,6 +4329,15 @@ FGraphPanelSelectionSet FMaterialEditor::GetSelectedNodes() const
 	return CurrentSelection;
 }
 
+void FMaterialEditor::GetBoundsForNode(const UEdGraphNode* InNode, class FSlateRect& OutRect, float InPadding) const
+{
+	if (GraphEditor.IsValid())
+	{
+		GraphEditor->GetBoundsForNode(InNode, OutRect, InPadding);
+	}
+}
+
+
 FMatExpressionPreview* FMaterialEditor::GetExpressionPreview(UMaterialExpression* InExpression)
 {
 	bool bNewlyCreated;
diff --git a/Engine/Source/Editor/MaterialEditor/Private/MaterialEditor.h b/Engine/Source/Editor/MaterialEditor/Private/MaterialEditor.h
--- a/Engine/Source/Editor/MaterialEditor/Private/MaterialEditor.h
+++ b/Engine/Source/Editor/MaterialEditor/Private/MaterialEditor.h
@@ -379,6 +379,7 @@ public:
 	virtual void PasteNodesHere(const FVector2D& Location) override;
 	virtual int32 GetNumberOfSelectedNodes() const override;
 	virtual TSet<UObject*> GetSelectedNodes() const override;
+	virtual void GetBoundsForNode(const UEdGraphNode* InNode, class FSlateRect& OutRect, float InPadding) const override;
 	virtual FMatExpressionPreview* GetExpressionPreview(UMaterialExpression* InExpression) override;
 	virtual void DeleteNodes(const TArray<class UEdGraphNode*>& NodesToDelete) override;
 
diff --git a/Engine/Source/Editor/MaterialEditor/Public/IMaterialEditor.h b/Engine/Source/Editor/MaterialEditor/Public/IMaterialEditor.h
--- a/Engine/Source/Editor/MaterialEditor/Public/IMaterialEditor.h
+++ b/Engine/Source/Editor/MaterialEditor/Public/IMaterialEditor.h
@@ -83,6 +83,9 @@ public:
 	/** Get the currently selected set of nodes */
 	virtual TSet<UObject*> GetSelectedNodes() const {return TSet<UObject*>();}
 
+	/** Get the visible bounds of the given graph node */
+	virtual void GetBoundsForNode(const UEdGraphNode* InNode, class FSlateRect& OutRect, float InPadding) const {};
+
 	/**
 	 * Gets the preview for an expression
 	 *
