Fix for crash when renaming a material layer and then expanding the UI
#rnx

diff --git a/Engine/Source/Editor/MaterialEditor/Private/SMaterialLayersFunctionsTree.cpp b/Engine/Source/Editor/MaterialEditor/Private/SMaterialLayersFunctionsTree.cpp
--- a/Engine/Source/Editor/MaterialEditor/Private/SMaterialLayersFunctionsTree.cpp
+++ b/Engine/Source/Editor/MaterialEditor/Private/SMaterialLayersFunctionsTree.cpp
@@ -135,7 +135,7 @@ void SMaterialLayersFunctionsInstanceTreeItem::OnNameChanged(const FText& InText
 	const FScopedTransaction Transaction(LOCTEXT("RenamedSection", "Renamed layer and blend section"));
 	InTree->FunctionInstanceHandle->NotifyPreChange();
 	InTree->FunctionInstance->LayerNames[Counter] = InText;
-	InTree->MaterialEditorInstance->CopyToSourceInstance(false);
+	InTree->MaterialEditorInstance->CopyToSourceInstance(true);
 	InTree->FunctionInstanceHandle->NotifyPostChange();
 }
 
