Make sure we have material editor commands before trying to use them to create a node

diff --git a/Engine/Source/Editor/MaterialEditor/Private/MaterialEditor.cpp b/Engine/Source/Editor/MaterialEditor/Private/MaterialEditor.cpp
--- a/Engine/Source/Editor/MaterialEditor/Private/MaterialEditor.cpp
+++ b/Engine/Source/Editor/MaterialEditor/Private/MaterialEditor.cpp
@@ -5282,16 +5282,17 @@ bool FMaterialEditor::OnVerifyNodeTextCommit(const FText& NewText, UEdGraphNode*
 FReply FMaterialEditor::OnSpawnGraphNodeByShortcut(FInputChord InChord, const FVector2D& InPosition, UEdGraph* InGraph)
 {
 	UEdGraph* Graph = InGraph;
-
-	TSharedPtr< FEdGraphSchemaAction > Action = FMaterialEditorSpawnNodeCommands::Get().GetGraphActionByChord(InChord, InGraph);
-
-	if(Action.IsValid())
+	if (FMaterialEditorSpawnNodeCommands::IsRegistered())
 	{
-		TArray<UEdGraphPin*> DummyPins;
-		Action->PerformAction(Graph, DummyPins, InPosition);
-		return FReply::Handled();
-	}
+		TSharedPtr< FEdGraphSchemaAction > Action = FMaterialEditorSpawnNodeCommands::Get().GetGraphActionByChord(InChord, InGraph);
 
+		if (Action.IsValid())
+		{
+			TArray<UEdGraphPin*> DummyPins;
+			Action->PerformAction(Graph, DummyPins, InPosition);
+			return FReply::Handled();
+		}
+	}
 	return FReply::Unhandled();
 }
 
