The Material Parameters list now updates when parameters are hooked up.

diff --git a/Engine/Source/Editor/MaterialEditor/Private/MaterialEditor.cpp b/Engine/Source/Editor/MaterialEditor/Private/MaterialEditor.cpp
--- a/Engine/Source/Editor/MaterialEditor/Private/MaterialEditor.cpp
+++ b/Engine/Source/Editor/MaterialEditor/Private/MaterialEditor.cpp
@@ -712,6 +712,9 @@ void FMaterialEditor::InitMaterialEditor( const EToolkitMode::Type Mode, const T
 	RegenerateCodeView(true);
 
 	ForceRefreshExpressionPreviews();
+	
+	// Update the parameter list now that the material has been fully initialized
+	MaterialParametersOverviewWidget->UpdateEditorInstance(MaterialEditorInstance);
 
 	if (OriginalMaterial->bUsedAsSpecialEngineMaterial)
 	{
@@ -962,7 +965,7 @@ void FMaterialEditor::CreateInternalWidgets()
 	MaterialParametersOverviewWidget = SNew(SMaterialParametersOverviewPanel)
 		.InMaterialEditorInstance(MaterialEditorInstance);
 	MaterialParametersOverviewWidget->GetGenerator()->OnFinishedChangingProperties().AddSP(this, &FMaterialEditor::OnFinishedChangingParametersFromOverview);
-	
+
 	IMaterialEditorModule* MaterialEditorModule = &FModuleManager::LoadModuleChecked<IMaterialEditorModule>("MaterialEditor");
 	if (MaterialEditorModule->MaterialLayersEnabled())
 	{
@@ -3872,7 +3875,6 @@ UMaterialExpression* FMaterialEditor::CreateNewMaterialExpression(UClass* NewExp
 	RefreshExpressionPreviews();
 	GraphEditor->NotifyGraphChanged();
 	SetMaterialDirty();
-	MaterialParametersOverviewWidget->UpdateEditorInstance(MaterialEditorInstance);
 	return NewExpression;
 }
 
@@ -4331,7 +4333,7 @@ void FMaterialEditor::UpdateMaterialAfterGraphChange()
 	RegenerateCodeView();
 	RefreshExpressionPreviews();
 	SetMaterialDirty();
-	
+	MaterialParametersOverviewWidget->UpdateEditorInstance(MaterialEditorInstance);
 	if (bHideUnrelatedNodes && !bLockNodeFadeState && bSelectRegularNode)
 	{
 		GraphEditor->ResetAllNodesUnrelatedStates();
diff --git a/Engine/Source/Editor/MaterialEditor/Private/MaterialEditorDetailCustomization.cpp b/Engine/Source/Editor/MaterialEditor/Private/MaterialEditorDetailCustomization.cpp
--- a/Engine/Source/Editor/MaterialEditor/Private/MaterialEditorDetailCustomization.cpp
+++ b/Engine/Source/Editor/MaterialEditor/Private/MaterialEditorDetailCustomization.cpp
@@ -40,446 +40,6 @@
 
 #define LOCTEXT_NAMESPACE "MaterialEditor"
 
-TSharedRef<IDetailCustomization> FMaterialEditorParameterDetails::MakeInstance(UMaterialEditorPreviewParameters* Material)
-{
-	return MakeShareable(new FMaterialEditorParameterDetails(Material));
-}
-
-FMaterialEditorParameterDetails::FMaterialEditorParameterDetails(UMaterialEditorPreviewParameters* Material)
-	: MaterialEditorInstance(Material)
-{
-}
-
-
-TOptional<float> FMaterialEditorParameterDetails::OnGetValue(TSharedRef<IPropertyHandle> PropertyHandle)
-{
-	float Value = 0.0f;
-	if (PropertyHandle->GetValue(Value) == FPropertyAccess::Success)
-	{
-		return TOptional<float>(Value);
-	}
-
-	// Value couldn't be accessed. Return an unset value
-	return TOptional<float>();
-}
-
-void FMaterialEditorParameterDetails::OnValueCommitted(float NewValue, ETextCommit::Type CommitType, TSharedRef<IPropertyHandle> PropertyHandle)
-{
-	// Try setting as float, if that fails then set as int
-	ensure(PropertyHandle->SetValue(NewValue) == FPropertyAccess::Success);
-}
-
-void FMaterialEditorParameterDetails::CustomizeDetails(IDetailLayoutBuilder& DetailLayout)
-{
-	PropertyUtilities = DetailLayout.GetPropertyUtilities();
-	// Create a new category for a custom layout for the MIC parameters at the very top
-	FName GroupsCategoryName = TEXT("ParameterGroups");
-	IDetailCategoryBuilder& GroupsCategory = DetailLayout.EditCategory(GroupsCategoryName, LOCTEXT("MICParamGroupsTitle", "Parameter Groups"));
-	TSharedRef<IPropertyHandle> ParameterGroupsProperty = DetailLayout.GetProperty("ParameterGroups");
-
-	CreateGroupsWidget(ParameterGroupsProperty, GroupsCategory);
-
-	DetailLayout.HideProperty("ParameterGroups");
-
-	TArray<UObject*> ExternalObjects;
-	ExternalObjects.Add(MaterialEditorInstance->OriginalMaterial);
-}
-
-
-void FMaterialEditorParameterDetails::CreateGroupsWidget(TSharedRef<IPropertyHandle> ParameterGroupsProperty, IDetailCategoryBuilder& GroupsCategory)
-{
-	check(MaterialEditorInstance);
-	bool bShowSaveButtons = false;
-	for (int32 GroupIdx = 0; GroupIdx < MaterialEditorInstance->ParameterGroups.Num(); ++GroupIdx)
-	{
-		FEditorParameterGroup& ParameterGroup = MaterialEditorInstance->ParameterGroups[GroupIdx];
-		if (ParameterGroup.GroupAssociation == EMaterialParameterAssociation::GlobalParameter
-			&& ParameterGroup.GroupName != FMaterialPropertyHelpers::LayerParamName)
-		{
-			IDetailGroup& DetailGroup = GroupsCategory.AddGroup(ParameterGroup.GroupName, FText::FromName(ParameterGroup.GroupName), false, true);
-			CreateSingleGroupWidget(ParameterGroup, ParameterGroupsProperty->GetChildHandle(GroupIdx), DetailGroup);
-			bShowSaveButtons = true;
-		}
-	}
-	if (bShowSaveButtons)
-	{
-		FOnClicked ChildButtonClicked;
-		UObject* LocalEditorInstance = MaterialEditorInstance;
-		if (MaterialEditorInstance->OriginalMaterial)
-		{
-			ChildButtonClicked = FOnClicked::CreateStatic(FMaterialPropertyHelpers::OnClickedSaveNewMaterialInstance,
-				ImplicitConv<UMaterialInterface*>(MaterialEditorInstance->OriginalMaterial), LocalEditorInstance);
-		}
-		if (MaterialEditorInstance->OriginalFunction)
-		{
-			ChildButtonClicked = FOnClicked::CreateStatic(FMaterialPropertyHelpers::OnClickedSaveNewFunctionInstance,
-				ImplicitConv<UMaterialFunctionInterface*>(MaterialEditorInstance->OriginalFunction), ImplicitConv<UMaterialInterface*>(MaterialEditorInstance->PreviewMaterial), LocalEditorInstance);
-		}
-
-		FDetailWidgetRow& SaveInstanceRow = GroupsCategory.AddCustomRow(LOCTEXT("SaveInstances", "Save Instances"));
-		SaveInstanceRow.ValueContent()
-			.HAlign(HAlign_Fill)
-			[
-				SNew(SHorizontalBox)
-				+ SHorizontalBox::Slot()
-				.FillWidth(1.0f)
-				[
-					SNullWidget::NullWidget
-				]
-				+ SHorizontalBox::Slot()
-				.AutoWidth()
-				.Padding(2.0f)
-				[
-					SNew(SButton)
-					.ButtonStyle(FEditorStyle::Get(), "FlatButton.Dark")
-					.HAlign(HAlign_Center)
-					.OnClicked(ChildButtonClicked)
-					.ToolTipText(LOCTEXT("SaveToChildInstance", "Save To Child Instance"))
-					.Content()
-					[
-						SNew(SHorizontalBox)
-						+ SHorizontalBox::Slot()
-						.AutoWidth()
-						[
-							SNew(STextBlock)
-							.Font(FEditorStyle::Get().GetFontStyle("FontAwesome.10"))
-							.TextStyle(FEditorStyle::Get(), "NormalText.Important")
-							.Text(FText::FromString(FString(TEXT("\xf0c7 \xf149"))) /*fa-filter*/)
-						]
-						+ SHorizontalBox::Slot()
-						.AutoWidth()
-						[
-							SNew(STextBlock)
-							.TextStyle(FEditorStyle::Get(), "NormalText.Important")
-							.Text(FText::FromString(FString(TEXT(" Save Child"))) /*fa-filter*/)
-						]
-					]
-				]
-			];
-	}
-	else
-	{
-		FDetailWidgetRow& NoParameters = GroupsCategory.AddCustomRow(LOCTEXT("NoParameters", "No Parameters"));
-		NoParameters.WholeRowContent()
-			[
-				SNew(STextBlock)
-				.Text(LOCTEXT("AddGlobalParameterPrompt", "Add a global parameter to see it here. \nMaterial Attribute Layers parameters are found in the Layer Parameter Preview panel."))
-				.AutoWrapText(true)
-			];
-	}
-}
-
-
-void FMaterialEditorParameterDetails::CreateSingleGroupWidget(FEditorParameterGroup& ParameterGroup, TSharedPtr<IPropertyHandle> ParameterGroupProperty, IDetailGroup& DetailGroup)
-{
-	TSharedPtr<IPropertyHandle> ParametersArrayProperty = ParameterGroupProperty->GetChildHandle("Parameters");
-
-	// Create a custom widget for each parameter in the group
-	for (int32 ParamIdx = 0; ParamIdx < ParameterGroup.Parameters.Num(); ++ParamIdx)
-	{
-		TSharedPtr<IPropertyHandle> ParameterProperty = ParametersArrayProperty->GetChildHandle(ParamIdx);
-
-		FString ParameterName = ParameterGroup.Parameters[ParamIdx]->ParameterInfo.Name.ToString();
-
-		UDEditorParameterValue* Parameter = ParameterGroup.Parameters[ParamIdx];
-		UDEditorFontParameterValue* FontParam = Cast<UDEditorFontParameterValue>(Parameter);
-		UDEditorScalarParameterValue* ScalarParam = Cast<UDEditorScalarParameterValue>(Parameter);
-		UDEditorStaticComponentMaskParameterValue* CompMaskParam = Cast<UDEditorStaticComponentMaskParameterValue>(Parameter);
-		UDEditorStaticSwitchParameterValue* SwitchParam = Cast<UDEditorStaticSwitchParameterValue>(Parameter);
-		UDEditorTextureParameterValue* TextureParam = Cast<UDEditorTextureParameterValue>(Parameter);
-		UDEditorRuntimeVirtualTextureParameterValue* RuntimeVirtualTextureParam = Cast<UDEditorRuntimeVirtualTextureParameterValue>(Parameter);
-		UDEditorVectorParameterValue* VectorParam = Cast<UDEditorVectorParameterValue>(Parameter);
-		UDEditorMaterialLayersParameterValue* LayersParam = Cast<UDEditorMaterialLayersParameterValue>(Parameter);
-
-		if (VectorParam && VectorParam->bIsUsedAsChannelMask)
-		{
-			CreateVectorChannelMaskParameterValueWidget(Parameter, ParameterProperty, DetailGroup);
-		}
-		if (ScalarParam && ScalarParam->AtlasData.bIsUsedAsAtlasPosition)
-		{
-			CreateScalarAtlasPositionParameterValueWidget(Parameter, ParameterProperty, DetailGroup);
-		}
-		else if (ScalarParam || SwitchParam || TextureParam || RuntimeVirtualTextureParam || VectorParam || FontParam)
-		{
-			if (ScalarParam && ScalarParam->SliderMax > ScalarParam->SliderMin)
-			{
-				TSharedPtr<IPropertyHandle> ParameterValueProperty = ParameterProperty->GetChildHandle("ParameterValue");
-				ParameterValueProperty->SetInstanceMetaData("UIMin", FString::Printf(TEXT("%f"), ScalarParam->SliderMin));
-				ParameterValueProperty->SetInstanceMetaData("UIMax", FString::Printf(TEXT("%f"), ScalarParam->SliderMax));
-			}
-
-			CreateParameterValueWidget(Parameter, ParameterProperty, DetailGroup);
-		}
-		else if (LayersParam)
-		{
-		}
-		else if (CompMaskParam)
-		{
-			CreateMaskParameterValueWidget(Parameter, ParameterProperty, DetailGroup);
-		}
-		else
-		{
-			// Unsupported parameter type
-			check(false);
-		}
-	}
-}
-
-
-void FMaterialEditorParameterDetails::CreateParameterValueWidget(UDEditorParameterValue* Parameter, TSharedPtr<IPropertyHandle> ParameterProperty, IDetailGroup& DetailGroup)
-{
-	TSharedPtr<IPropertyHandle> ParameterValueProperty = ParameterProperty->GetChildHandle("ParameterValue");
-
-	if (ParameterValueProperty->IsValidHandle())
-	{
-
-		IDetailPropertyRow& PropertyRow = DetailGroup.AddPropertyRow(ParameterValueProperty.ToSharedRef());
-		PropertyRow
-			.DisplayName(FText::FromName(Parameter->ParameterInfo.Name))
-			.ToolTip(GetParameterExpressionDescription(Parameter));
-
-		// Textures need a special widget that filters based on VT or not
-		UDEditorTextureParameterValue* TextureParam = Cast<UDEditorTextureParameterValue>(Parameter);
-		if (TextureParam != nullptr)
-		{
-			UMaterial *Material = MaterialEditorInstance->PreviewMaterial->GetMaterial();
-			if (Material != nullptr)
-			{
-				UMaterialExpressionTextureSampleParameter* Expression = Material->FindExpressionByGUID<UMaterialExpressionTextureSampleParameter>(TextureParam->ExpressionId);
-				if (Expression != nullptr)
-				{
-					TWeakObjectPtr<UMaterialExpressionTextureSampleParameter> SamplerExpression = Expression;
-
-					PropertyRow.CustomWidget()
-					.NameContent()
-					[
-						ParameterValueProperty->CreatePropertyNameWidget()
-					]
-					.ValueContent()
-					.MaxDesiredWidth(TOptional<float>())
-					[
-						SNew(SObjectPropertyEntryBox)
-						.PropertyHandle(ParameterValueProperty)
-						.AllowedClass(UTexture::StaticClass())
-						.ThumbnailPool(PropertyUtilities.Pin()->GetThumbnailPool())
-						.OnShouldFilterAsset_Lambda([SamplerExpression](const FAssetData& AssetData)
-						{
-							if (SamplerExpression.Get())
-							{
-								bool VirtualTextured = false;
-								AssetData.GetTagValue<bool>("VirtualTextureStreaming", VirtualTextured);
-
-								bool ExpressionIsVirtualTextured = IsVirtualSamplerType(SamplerExpression->SamplerType);
-
-								return VirtualTextured != ExpressionIsVirtualTextured;
-							}
-							else
-							{
-								return false;
-							}
-						})
-					];
-				}
-			}
-		}
-	}
-}
-
-void FMaterialEditorParameterDetails::CreateMaskParameterValueWidget(UDEditorParameterValue* Parameter, TSharedPtr<IPropertyHandle> ParameterProperty, IDetailGroup& DetailGroup)
-{
-	TSharedPtr<IPropertyHandle> ParameterValueProperty = ParameterProperty->GetChildHandle("ParameterValue");
-	TSharedPtr<IPropertyHandle> RMaskProperty = ParameterValueProperty->GetChildHandle("R");
-	TSharedPtr<IPropertyHandle> GMaskProperty = ParameterValueProperty->GetChildHandle("G");
-	TSharedPtr<IPropertyHandle> BMaskProperty = ParameterValueProperty->GetChildHandle("B");
-	TSharedPtr<IPropertyHandle> AMaskProperty = ParameterValueProperty->GetChildHandle("A");
-
-	if (ParameterValueProperty->IsValidHandle())
-	{
-		IDetailPropertyRow& PropertyRow = DetailGroup.AddPropertyRow(ParameterValueProperty.ToSharedRef());
-		const FText ParameterName = FText::FromName(Parameter->ParameterInfo.Name);
-
-		FDetailWidgetRow& CustomWidget = PropertyRow.CustomWidget();
-		CustomWidget
-			.FilterString(ParameterName)
-			.NameContent()
-			[
-				SNew(STextBlock)
-				.Text(ParameterName)
-				.ToolTipText(GetParameterExpressionDescription(Parameter))
-				.Font(FEditorStyle::GetFontStyle(TEXT("PropertyWindow.NormalFont")))
-			]
-			.ValueContent()
-			.MaxDesiredWidth(200.0f)
-			[
-				SNew(SHorizontalBox)
-				+ SHorizontalBox::Slot()
-				.FillWidth(1.0f)
-				[
-					SNew(SHorizontalBox)
-					+ SHorizontalBox::Slot()
-					.HAlign(HAlign_Left)
-					.AutoWidth()
-					[
-						RMaskProperty->CreatePropertyNameWidget(FText::GetEmpty(), FText::GetEmpty(), false)
-					]
-					+ SHorizontalBox::Slot()
-					.HAlign(HAlign_Left)
-					.AutoWidth()
-					[
-						RMaskProperty->CreatePropertyValueWidget()
-					]
-					+ SHorizontalBox::Slot()
-					.HAlign(HAlign_Left)
-					.Padding(FMargin(10.0f, 0.0f, 0.0f, 0.0f))
-					.AutoWidth()
-					[
-						GMaskProperty->CreatePropertyNameWidget(FText::GetEmpty(), FText::GetEmpty(), false)
-					]
-					+ SHorizontalBox::Slot()
-					.HAlign(HAlign_Left)
-					.AutoWidth()
-					[
-						GMaskProperty->CreatePropertyValueWidget()
-					]
-					+ SHorizontalBox::Slot()
-					.HAlign(HAlign_Left)
-					.Padding(FMargin(10.0f, 0.0f, 0.0f, 0.0f))
-					.AutoWidth()
-					[
-						BMaskProperty->CreatePropertyNameWidget(FText::GetEmpty(), FText::GetEmpty(), false)
-					]
-					+ SHorizontalBox::Slot()
-					.HAlign(HAlign_Left)
-					.AutoWidth()
-					[
-						BMaskProperty->CreatePropertyValueWidget()
-					]
-					+ SHorizontalBox::Slot()
-					.HAlign(HAlign_Left)
-					.Padding(FMargin(10.0f, 0.0f, 0.0f, 0.0f))
-					.AutoWidth()
-					[
-						AMaskProperty->CreatePropertyNameWidget(FText::GetEmpty(), FText::GetEmpty(), false)
-					]
-					+ SHorizontalBox::Slot()
-					.HAlign(HAlign_Left)
-					.AutoWidth()
-					[
-						AMaskProperty->CreatePropertyValueWidget()
-					]
-				]
-			];
-	}
-}
-
-void FMaterialEditorParameterDetails::CreateVectorChannelMaskParameterValueWidget(UDEditorParameterValue* Parameter, TSharedPtr<IPropertyHandle> ParameterProperty, IDetailGroup& DetailGroup)
-{
-	TSharedPtr<IPropertyHandle> ParameterValueProperty = ParameterProperty->GetChildHandle("ParameterValue");
-
-	if (ParameterValueProperty->IsValidHandle())
-	{
-		TAttribute<bool> IsParamEnabled = TAttribute<bool>::Create(TAttribute<bool>::FGetter::CreateStatic(&FMaterialPropertyHelpers::IsOverriddenExpression, Parameter));
-
-		IDetailPropertyRow& PropertyRow = DetailGroup.AddPropertyRow(ParameterValueProperty.ToSharedRef());
-		const FText ParameterName = FText::FromName(Parameter->ParameterInfo.Name);
-
-		FOnGetPropertyComboBoxStrings GetMaskStrings = FOnGetPropertyComboBoxStrings::CreateStatic(&FMaterialPropertyHelpers::GetVectorChannelMaskComboBoxStrings);
-		FOnGetPropertyComboBoxValue GetMaskValue = FOnGetPropertyComboBoxValue::CreateStatic(&FMaterialPropertyHelpers::GetVectorChannelMaskValue, Parameter);
-		FOnPropertyComboBoxValueSelected SetMaskValue = FOnPropertyComboBoxValueSelected::CreateStatic(&FMaterialPropertyHelpers::SetVectorChannelMaskValue, ParameterValueProperty, Parameter, (UObject*)MaterialEditorInstance);
-
-		FDetailWidgetRow& CustomWidget = PropertyRow.CustomWidget();
-		CustomWidget
-		.FilterString(ParameterName)
-		.NameContent()
-		[
-			SNew(STextBlock)
-			.Text(ParameterName)
-			.ToolTipText(FMaterialPropertyHelpers::GetParameterExpressionDescription(Parameter, MaterialEditorInstance))
-			.Font(FEditorStyle::GetFontStyle(TEXT("PropertyWindow.NormalFont")))
-		]
-		.ValueContent()
-		.MaxDesiredWidth(200.0f)
-		[
-			SNew(SHorizontalBox)
-			+ SHorizontalBox::Slot()
-			.FillWidth(1.0f)
-			[
-				SNew(SHorizontalBox)
-				+ SHorizontalBox::Slot()
-				.HAlign(HAlign_Left)
-				.AutoWidth()
-				[
-					PropertyCustomizationHelpers::MakePropertyComboBox(ParameterValueProperty, GetMaskStrings, GetMaskValue, SetMaskValue)
-				]
-			]
-		];
-	}
-}
-
-void FMaterialEditorParameterDetails::CreateScalarAtlasPositionParameterValueWidget(class UDEditorParameterValue* Parameter, TSharedPtr<IPropertyHandle> ParameterProperty, IDetailGroup& DetailGroup)
-{
-	TSharedPtr<IPropertyHandle> ParameterValueProperty = ParameterProperty->GetChildHandle("ParameterValue");
-
-	if (ParameterValueProperty->IsValidHandle())
-	{
-		TAttribute<bool> IsParamEnabled = TAttribute<bool>::Create(TAttribute<bool>::FGetter::CreateStatic(&FMaterialPropertyHelpers::IsOverriddenExpression, Parameter));
-
-		IDetailPropertyRow& PropertyRow = DetailGroup.AddPropertyRow(ParameterValueProperty.ToSharedRef());
-	
-		const FText ParameterName = FText::FromName(Parameter->ParameterInfo.Name);
-		UDEditorScalarParameterValue* AtlasParameter = Cast<UDEditorScalarParameterValue>(Parameter);
-
-		FDetailWidgetRow& CustomWidget = PropertyRow.CustomWidget();
-		CustomWidget
-			.FilterString(ParameterName)
-			.NameContent()
-			[
-				SNew(STextBlock)
-				.Text(ParameterName)
-				.ToolTipText(FMaterialPropertyHelpers::GetParameterExpressionDescription(Parameter, MaterialEditorInstance))
-				.Font(FEditorStyle::GetFontStyle(TEXT("PropertyWindow.NormalFont")))
-			]
-			.ValueContent()
-			.HAlign(HAlign_Fill)
-			.MaxDesiredWidth(400.0f)
-			[
-				SNew(SObjectPropertyEntryBox)
-				.ObjectPath(this, &FMaterialEditorParameterDetails::GetCurvePath, AtlasParameter)
-				.AllowedClass(UCurveLinearColor::StaticClass())
-				.NewAssetFactories(TArray<UFactory*>())
-				.DisplayThumbnail(true)
-				.ThumbnailPool(PropertyUtilities.Pin()->GetThumbnailPool())
-				.OnShouldSetAsset(FOnShouldSetAsset::CreateStatic(&FMaterialPropertyHelpers::OnShouldSetCurveAsset, AtlasParameter->AtlasData.Atlas))
-				.OnObjectChanged(FOnSetObject::CreateStatic(&FMaterialPropertyHelpers::SetPositionFromCurveAsset, AtlasParameter->AtlasData.Atlas, AtlasParameter, ParameterProperty, (UObject*)MaterialEditorInstance))
-				.DisplayCompactSize(true)
-			];
-	}
-}
-
-FString FMaterialEditorParameterDetails::GetCurvePath(UDEditorScalarParameterValue* Parameter) const
-{
-	FString Path = Parameter->AtlasData.Curve->GetPathName();
-	return Path;
-}
-
-FText FMaterialEditorParameterDetails::GetParameterExpressionDescription(UDEditorParameterValue* Parameter) const
-{
-	if (MaterialEditorInstance && MaterialEditorInstance->OriginalMaterial)
-	{
-		UMaterial* BaseMaterial = MaterialEditorInstance->OriginalMaterial->GetMaterial();
-		if (BaseMaterial)
-		{
-			UMaterialExpression* MaterialExpression = BaseMaterial->FindExpressionByGUID<UMaterialExpression>(Parameter->ExpressionId);
-
-			if (MaterialExpression)
-			{
-				return FText::FromString(MaterialExpression->Desc);
-			}
-		}
-	}
-	return FText::GetEmpty();
-}
-
 TSharedRef<IDetailCustomization> FMaterialExpressionParameterDetails::MakeInstance(FOnCollectParameterGroups InCollectGroupsDelegate)
 {
 	return MakeShareable( new FMaterialExpressionParameterDetails(InCollectGroupsDelegate) );
diff --git a/Engine/Source/Editor/MaterialEditor/Private/MaterialEditorDetailCustomization.h b/Engine/Source/Editor/MaterialEditor/Private/MaterialEditorDetailCustomization.h
--- a/Engine/Source/Editor/MaterialEditor/Private/MaterialEditorDetailCustomization.h
+++ b/Engine/Source/Editor/MaterialEditor/Private/MaterialEditorDetailCustomization.h
@@ -14,47 +14,6 @@ class SEditableText;
 
 DECLARE_DELEGATE_OneParam(FOnCollectParameterGroups, TArray<FString>*);
 
-class FMaterialEditorParameterDetails : public IDetailCustomization
-{
-
-public:
-	/** Makes a new instance of this detail layout class for a specific detail view requesting it */
-	static TSharedRef<class IDetailCustomization> MakeInstance(class UMaterialEditorPreviewParameters* Material);
-
-	/** Constructor */
-	FMaterialEditorParameterDetails(class UMaterialEditorPreviewParameters* Material);
-
-	/** IDetailCustomization interface */
-	virtual void CustomizeDetails(IDetailLayoutBuilder& DetailLayout) override;
-
-	static TOptional<float> OnGetValue(TSharedRef<IPropertyHandle> PropertyHandle);
-	static void OnValueCommitted(float NewValue, ETextCommit::Type CommitType, TSharedRef<IPropertyHandle> PropertyHandle);
-
-private:
-	/** Builds the custom parameter groups category */
-	void CreateGroupsWidget(TSharedRef<IPropertyHandle> ParameterGroupsProperty, class IDetailCategoryBuilder& GroupsCategory);
-
-	/** Builds the widget for an individual parameter group */
-	void CreateSingleGroupWidget(struct FEditorParameterGroup& ParameterGroup, TSharedPtr<IPropertyHandle> ParameterGroupProperty, class IDetailGroup& DetailGroup);
-
-	/** These methods generate the custom widgets for the various parameter types */
-	void CreateParameterValueWidget(class UDEditorParameterValue* Parameter, TSharedPtr<IPropertyHandle> ParameterProperty, IDetailGroup& DetailGroup);
-	void CreateMaskParameterValueWidget(class UDEditorParameterValue* Parameter, TSharedPtr<IPropertyHandle> ParameterProperty, IDetailGroup& DetailGroup);
-	void CreateVectorChannelMaskParameterValueWidget(class UDEditorParameterValue* Parameter, TSharedPtr<IPropertyHandle> ParameterProperty, IDetailGroup& DetailGroup);
-	void CreateScalarAtlasPositionParameterValueWidget(class UDEditorParameterValue* Parameter, TSharedPtr<IPropertyHandle> ParameterProperty, IDetailGroup& DetailGroup);
-	FString GetCurvePath(class UDEditorScalarParameterValue* Parameter) const;
-	/** Gets the expression description of this parameter from the the base material */
-	FText GetParameterExpressionDescription(class UDEditorParameterValue* Parameter) const;
-
-private:
-	/** Object that stores all of the possible parameters we can edit */
-	UMaterialEditorPreviewParameters* MaterialEditorInstance;
-
-	/** Associated UMaterialEditorPreviewParameters utilities */
-	TWeakPtr<class IPropertyUtilities> PropertyUtilities;
-
-};
-
 class FMaterialExpressionParameterDetails : public IDetailCustomization
 {
 public:
diff --git a/Engine/Source/Editor/MaterialEditor/Private/MaterialEditorInstanceDetailCustomization.cpp b/Engine/Source/Editor/MaterialEditor/Private/MaterialEditorInstanceDetailCustomization.cpp
--- a/Engine/Source/Editor/MaterialEditor/Private/MaterialEditorInstanceDetailCustomization.cpp
+++ b/Engine/Source/Editor/MaterialEditor/Private/MaterialEditorInstanceDetailCustomization.cpp
@@ -763,146 +763,174 @@ void FMaterialInstanceParameterDetails::CreateLabeledTextureParameterValueWidget
 		UDEditorTextureParameterValue* TextureParam = Cast<UDEditorTextureParameterValue>(Parameter);
 		if (TextureParam)
 		{
-			TAttribute<bool> IsParamEnabled = TAttribute<bool>::Create(TAttribute<bool>::FGetter::CreateStatic(&FMaterialPropertyHelpers::IsOverriddenExpression, Parameter));
+			UMaterial *Material = MaterialEditorInstance->SourceInstance->GetMaterial();
+			if (Material != nullptr)
+			{
+				UMaterialExpressionTextureSampleParameter* Expression = Material->FindExpressionByGUID<UMaterialExpressionTextureSampleParameter>(TextureParam->ExpressionId);
+				if (Expression != nullptr)
+				{
+					TWeakObjectPtr<UMaterialExpressionTextureSampleParameter> SamplerExpression = Expression;
+					TAttribute<bool> IsParamEnabled = TAttribute<bool>::Create(TAttribute<bool>::FGetter::CreateStatic(&FMaterialPropertyHelpers::IsOverriddenExpression, Parameter));
 
-			IDetailPropertyRow& PropertyRow = DetailGroup.AddPropertyRow(ParameterValueProperty.ToSharedRef());
+					IDetailPropertyRow& PropertyRow = DetailGroup.AddPropertyRow(ParameterValueProperty.ToSharedRef());
 
-			FIsResetToDefaultVisible IsResetVisible = FIsResetToDefaultVisible::CreateStatic(&FMaterialPropertyHelpers::ShouldShowResetToDefault, Parameter, MaterialEditorInstance);
-			FResetToDefaultHandler ResetHandler = FResetToDefaultHandler::CreateStatic(&FMaterialPropertyHelpers::ResetToDefault, Parameter, MaterialEditorInstance);
-			FResetToDefaultOverride ResetOverride = FResetToDefaultOverride::Create(IsResetVisible, ResetHandler);
+					FIsResetToDefaultVisible IsResetVisible = FIsResetToDefaultVisible::CreateStatic(&FMaterialPropertyHelpers::ShouldShowResetToDefault, Parameter, MaterialEditorInstance);
+					FResetToDefaultHandler ResetHandler = FResetToDefaultHandler::CreateStatic(&FMaterialPropertyHelpers::ResetToDefault, Parameter, MaterialEditorInstance);
+					FResetToDefaultOverride ResetOverride = FResetToDefaultOverride::Create(IsResetVisible, ResetHandler);
 
-			PropertyRow
-				.DisplayName(FText::FromName(Parameter->ParameterInfo.Name))
-				.EditCondition(IsParamEnabled, FOnBooleanValueChanged::CreateStatic(&FMaterialPropertyHelpers::OnOverrideParameter, Parameter, MaterialEditorInstance))
-				.ToolTip(FMaterialPropertyHelpers::GetParameterTooltip(Parameter, MaterialEditorInstance))
-				.Visibility(TAttribute<EVisibility>::Create(TAttribute<EVisibility>::FGetter::CreateStatic(&FMaterialPropertyHelpers::ShouldShowExpression, Parameter, MaterialEditorInstance, ShowHiddenDelegate)))
-				.OverrideResetToDefault(ResetOverride);
+					PropertyRow
+						.DisplayName(FText::FromName(Parameter->ParameterInfo.Name))
+						.EditCondition(IsParamEnabled, FOnBooleanValueChanged::CreateStatic(&FMaterialPropertyHelpers::OnOverrideParameter, Parameter, MaterialEditorInstance))
+						.ToolTip(FMaterialPropertyHelpers::GetParameterTooltip(Parameter, MaterialEditorInstance))
+						.Visibility(TAttribute<EVisibility>::Create(TAttribute<EVisibility>::FGetter::CreateStatic(&FMaterialPropertyHelpers::ShouldShowExpression, Parameter, MaterialEditorInstance, ShowHiddenDelegate)))
+						.OverrideResetToDefault(ResetOverride);
 
-			
-			TSharedPtr<SWidget> NameWidget;
-			TSharedPtr<SWidget> ValueWidget;
-			FDetailWidgetRow Row;
-			PropertyRow.GetDefaultWidgets(NameWidget, ValueWidget, Row);
-
-			FDetailWidgetRow &DetailWidgetRow = PropertyRow.CustomWidget();
-			TSharedPtr<SVerticalBox> NameVerticalBox;
-			DetailWidgetRow.NameContent()
-			[
-				SAssignNew(NameVerticalBox, SVerticalBox)
-				+ SVerticalBox::Slot()
-				.AutoHeight()
-				[
-					SNew(STextBlock)
-					.Text(FText::FromName(Parameter->ParameterInfo.Name))
-					.ToolTipText(FMaterialPropertyHelpers::GetParameterTooltip(Parameter, MaterialEditorInstance))
-					.Font(FEditorStyle::GetFontStyle(TEXT("PropertyWindow.NormalFont")))
-				]
-			];
-			DetailWidgetRow.ValueContent()
-			.MinDesiredWidth(Row.ValueWidget.MinWidth)
-			.MaxDesiredWidth(Row.ValueWidget.MaxWidth)
-			[
-				ValueWidget.ToSharedRef()
-			];
 
-			static const FName Red("R");
-			static const FName Green("G");
-			static const FName Blue("B");
-			static const FName Alpha("A");
+					TSharedPtr<SWidget> NameWidget;
+					TSharedPtr<SWidget> ValueWidget;
+					FDetailWidgetRow Row;
+					PropertyRow.GetDefaultWidgets(NameWidget, ValueWidget, Row);
 
-			if(!TextureParam->ChannelNames.R.IsEmpty())
-			{
-				NameVerticalBox->AddSlot()
-				[
-					SNew(SHorizontalBox)
-					+SHorizontalBox::Slot()
-					.AutoWidth()
-					.Padding(20.0, 2.0, 4.0, 2.0)
-					[
-						SNew(STextBlock)
-						.Text(FText::FromName(Red))
-						.Font(FEditorStyle::GetFontStyle(TEXT("PropertyWindow.BoldFont")))
-					]
-					+SHorizontalBox::Slot()
-					.HAlign(HAlign_Left)
-					.Padding(4.0, 2.0)
-					[
-						SNew(STextBlock)
-						.Text(TextureParam->ChannelNames.R)
-						.Font(FEditorStyle::GetFontStyle(TEXT("PropertyWindow.NormalFont")))
-					]
-				];
-			}
-			if (!TextureParam->ChannelNames.G.IsEmpty())
-			{
-				NameVerticalBox->AddSlot()
-					[
-						SNew(SHorizontalBox)
-						+ SHorizontalBox::Slot()
-						.Padding(20.0, 2.0, 4.0, 2.0)
-						.AutoWidth()
+					FDetailWidgetRow &DetailWidgetRow = PropertyRow.CustomWidget();
+					TSharedPtr<SVerticalBox> NameVerticalBox;
+					DetailWidgetRow.NameContent()
+						[
+							SAssignNew(NameVerticalBox, SVerticalBox)
+							+ SVerticalBox::Slot()
+						.AutoHeight()
 						[
 							SNew(STextBlock)
-							.Text(FText::FromName(Green))
-							.Font(FEditorStyle::GetFontStyle(TEXT("PropertyWindow.BoldFont")))
+							.Text(FText::FromName(Parameter->ParameterInfo.Name))
+						.ToolTipText(FMaterialPropertyHelpers::GetParameterTooltip(Parameter, MaterialEditorInstance))
+						.Font(FEditorStyle::GetFontStyle(TEXT("PropertyWindow.NormalFont")))
 						]
-						+ SHorizontalBox::Slot()
-						.HAlign(HAlign_Left)
-						.Padding(4.0, 2.0)
+						];
+					DetailWidgetRow.ValueContent()
+						.MinDesiredWidth(Row.ValueWidget.MinWidth)
+						.MaxDesiredWidth(Row.ValueWidget.MaxWidth)
 						[
-							SNew(STextBlock)
-							.Text(TextureParam->ChannelNames.G)
+							SNew(SObjectPropertyEntryBox)
+							.PropertyHandle(ParameterValueProperty)
+							.AllowedClass(UTexture::StaticClass())
+							.CustomResetToDefault(ResetOverride)
+							.ThumbnailPool(PropertyUtilities.Pin()->GetThumbnailPool())
+							.OnShouldFilterAsset_Lambda([SamplerExpression](const FAssetData& AssetData)
+							{
+								if (SamplerExpression.Get())
+								{
+									bool VirtualTextured = false;
+									AssetData.GetTagValue<bool>("VirtualTextureStreaming", VirtualTextured);
+
+									bool ExpressionIsVirtualTextured = IsVirtualSamplerType(SamplerExpression->SamplerType);
+
+									return VirtualTextured != ExpressionIsVirtualTextured;
+								}
+								else
+								{
+									return false;
+								}
+							})
+						];
+
+					static const FName Red("R");
+					static const FName Green("G");
+					static const FName Blue("B");
+					static const FName Alpha("A");
+
+					if (!TextureParam->ChannelNames.R.IsEmpty())
+					{
+						NameVerticalBox->AddSlot()
+							[
+								SNew(SHorizontalBox)
+								+ SHorizontalBox::Slot()
+								.AutoWidth()
+								.Padding(20.0, 2.0, 4.0, 2.0)
+								[
+									SNew(STextBlock)
+									.Text(FText::FromName(Red))
+									.Font(FEditorStyle::GetFontStyle(TEXT("PropertyWindow.BoldFont")))
+								]
+								+ SHorizontalBox::Slot()
+								.HAlign(HAlign_Left)
+								.Padding(4.0, 2.0)
+								[
+									SNew(STextBlock)
+									.Text(TextureParam->ChannelNames.R)
+									.Font(FEditorStyle::GetFontStyle(TEXT("PropertyWindow.NormalFont")))
+								]
+							];
+					}
+					if (!TextureParam->ChannelNames.G.IsEmpty())
+					{
+						NameVerticalBox->AddSlot()
+							[
+								SNew(SHorizontalBox)
+								+ SHorizontalBox::Slot()
+							.Padding(20.0, 2.0, 4.0, 2.0)
+							.AutoWidth()
+							[
+								SNew(STextBlock)
+								.Text(FText::FromName(Green))
+							.Font(FEditorStyle::GetFontStyle(TEXT("PropertyWindow.BoldFont")))
+							]
+						+ SHorizontalBox::Slot()
+							.HAlign(HAlign_Left)
+							.Padding(4.0, 2.0)
+							[
+								SNew(STextBlock)
+								.Text(TextureParam->ChannelNames.G)
 							.Font(FEditorStyle::GetFontStyle(TEXT("PropertyWindow.NormalFont")))
-						]
-					];
-			}
-			if (!TextureParam->ChannelNames.B.IsEmpty())
-			{
-				NameVerticalBox->AddSlot()
-					[
-						SNew(SHorizontalBox)
+							]
+							];
+					}
+					if (!TextureParam->ChannelNames.B.IsEmpty())
+					{
+						NameVerticalBox->AddSlot()
+							[
+								SNew(SHorizontalBox)
+								+ SHorizontalBox::Slot()
+							.Padding(20.0, 2.0, 4.0, 2.0)
+							.AutoWidth()
+							[
+								SNew(STextBlock)
+								.Text(FText::FromName(Blue))
+							.Font(FEditorStyle::GetFontStyle(TEXT("PropertyWindow.BoldFont")))
+							]
 						+ SHorizontalBox::Slot()
-						.Padding(20.0, 2.0, 4.0, 2.0)
-						.AutoWidth()
-						[
-							SNew(STextBlock)
-							.Text(FText::FromName(Blue))
+							.HAlign(HAlign_Left)
+							.Padding(4.0, 2.0)
+							[
+								SNew(STextBlock)
+								.Text(TextureParam->ChannelNames.B)
+							.Font(FEditorStyle::GetFontStyle(TEXT("PropertyWindow.NormalFont")))
+							]
+							];
+					}
+					if (!TextureParam->ChannelNames.A.IsEmpty())
+					{
+						NameVerticalBox->AddSlot()
+							[
+								SNew(SHorizontalBox)
+								+ SHorizontalBox::Slot()
+							.Padding(20.0, 2.0, 4.0, 2.0)
+							.AutoWidth()
+							[
+								SNew(STextBlock)
+								.Text(FText::FromName(Alpha))
 							.Font(FEditorStyle::GetFontStyle(TEXT("PropertyWindow.BoldFont")))
-						]
+							]
 						+ SHorizontalBox::Slot()
-						.HAlign(HAlign_Left)
-						.Padding(4.0, 2.0)
-						[
-							SNew(STextBlock)
-							.Text(TextureParam->ChannelNames.B)
+							.HAlign(HAlign_Left)
+							.Padding(4.0, 2.0)
+							[
+								SNew(STextBlock)
+								.Text(TextureParam->ChannelNames.A)
 							.Font(FEditorStyle::GetFontStyle(TEXT("PropertyWindow.NormalFont")))
-						]
-					];
-			}
-			if (!TextureParam->ChannelNames.A.IsEmpty())
-			{
-				NameVerticalBox->AddSlot()
-				[
-					SNew(SHorizontalBox)
-					+ SHorizontalBox::Slot()
-					.Padding(20.0, 2.0, 4.0, 2.0)
-					.AutoWidth()
-					[
-						SNew(STextBlock)
-						.Text(FText::FromName(Alpha))
-						.Font(FEditorStyle::GetFontStyle(TEXT("PropertyWindow.BoldFont")))
-					]
-					+ SHorizontalBox::Slot()
-					.HAlign(HAlign_Left)
-					.Padding(4.0, 2.0)
-					[
-						SNew(STextBlock)
-						.Text(TextureParam->ChannelNames.A)
-						.Font(FEditorStyle::GetFontStyle(TEXT("PropertyWindow.NormalFont")))
-					]
-				];
+							]
+							];
+					}
+				}
 			}
-			
 		}
 	}
 }
diff --git a/Engine/Source/Editor/MaterialEditor/Private/SMaterialLayersFunctionsTree.cpp b/Engine/Source/Editor/MaterialEditor/Private/SMaterialLayersFunctionsTree.cpp
--- a/Engine/Source/Editor/MaterialEditor/Private/SMaterialLayersFunctionsTree.cpp
+++ b/Engine/Source/Editor/MaterialEditor/Private/SMaterialLayersFunctionsTree.cpp
@@ -622,6 +622,7 @@ void SMaterialLayersFunctionsInstanceTreeItem::Construct(const FArguments& InArg
 		UDEditorStaticComponentMaskParameterValue* CompMaskParam = Cast<UDEditorStaticComponentMaskParameterValue>(StackParameterData->Parameter);
 		UDEditorVectorParameterValue* VectorParam = Cast<UDEditorVectorParameterValue>(StackParameterData->Parameter);
 		UDEditorScalarParameterValue* ScalarParam = Cast<UDEditorScalarParameterValue>(StackParameterData->Parameter);
+		UDEditorTextureParameterValue* TextureParam = Cast<UDEditorTextureParameterValue>(StackParameterData->Parameter);
 
 		TAttribute<bool> IsParamEnabled = TAttribute<bool>::Create(TAttribute<bool>::FGetter::CreateStatic(&FMaterialPropertyHelpers::IsOverriddenExpression, StackParameterData->Parameter));
 		NameOverride = FText::FromName(StackParameterData->Parameter->ParameterInfo.Name);
@@ -707,6 +708,159 @@ void SMaterialLayersFunctionsInstanceTreeItem::Construct(const FArguments& InArg
 				];
 			
 		}
+		else if (TextureParam)
+		{
+			UMaterial *Material = MaterialEditorInstance->SourceInstance->GetMaterial();
+			if (Material != nullptr)
+			{
+				UMaterialExpressionTextureSampleParameter* Expression = Material->FindExpressionByGUID<UMaterialExpressionTextureSampleParameter>(TextureParam->ExpressionId);
+				if (Expression != nullptr)
+				{
+					TWeakObjectPtr<UMaterialExpressionTextureSampleParameter> SamplerExpression = Expression;
+					TSharedPtr<SWidget> NameWidget;
+					TSharedPtr<SWidget> ValueWidget;
+					TSharedPtr<SVerticalBox> NameVerticalBox;
+					const FText ParameterName = FText::FromName(StackParameterData->Parameter->ParameterInfo.Name);
+					FDetailWidgetRow& CustomWidget = Row.CustomWidget();
+					CustomWidget
+						.FilterString(ParameterName)
+						.NameContent()
+						[
+							SAssignNew(NameVerticalBox, SVerticalBox)
+							+ SVerticalBox::Slot()
+							.AutoHeight()
+							[
+							SNew(STextBlock)
+							.Text(ParameterName)
+							.ToolTipText(FMaterialPropertyHelpers::GetParameterExpressionDescription(StackParameterData->Parameter, MaterialEditorInstance))
+							.Font(FEditorStyle::GetFontStyle(TEXT("PropertyWindow.NormalFont")))
+							]
+						];
+					CustomWidget.ValueContent()
+						[
+							SNew(SObjectPropertyEntryBox)
+							.PropertyHandle(StackParameterData->ParameterNode->CreatePropertyHandle())
+							.AllowedClass(UTexture::StaticClass())
+							.CustomResetToDefault(ResetOverride)
+							.ThumbnailPool(Tree->GetTreeThumbnailPool())
+							.OnShouldFilterAsset_Lambda([SamplerExpression](const FAssetData& AssetData)
+							{
+								if (SamplerExpression.Get())
+								{
+									bool VirtualTextured = false;
+									AssetData.GetTagValue<bool>("VirtualTextureStreaming", VirtualTextured);
+
+									bool ExpressionIsVirtualTextured = IsVirtualSamplerType(SamplerExpression->SamplerType);
+
+									return VirtualTextured != ExpressionIsVirtualTextured;
+								}
+								else
+								{
+									return false;
+								}
+							})
+						];
+
+					static const FName Red("R");
+					static const FName Green("G");
+					static const FName Blue("B");
+					static const FName Alpha("A");
+
+					if (!TextureParam->ChannelNames.R.IsEmpty())
+					{
+						NameVerticalBox->AddSlot()
+							[
+								SNew(SHorizontalBox)
+								+ SHorizontalBox::Slot()
+								.AutoWidth()
+								.Padding(20.0, 2.0, 4.0, 2.0)
+								[
+									SNew(STextBlock)
+									.Text(FText::FromName(Red))
+									.Font(FEditorStyle::GetFontStyle(TEXT("PropertyWindow.BoldFont")))
+								]
+								+ SHorizontalBox::Slot()
+								.HAlign(HAlign_Left)
+								.Padding(4.0, 2.0)
+								[
+									SNew(STextBlock)
+									.Text(TextureParam->ChannelNames.R)
+									.Font(FEditorStyle::GetFontStyle(TEXT("PropertyWindow.NormalFont")))
+								]
+							];
+					}
+					if (!TextureParam->ChannelNames.G.IsEmpty())
+					{
+						NameVerticalBox->AddSlot()
+							[
+								SNew(SHorizontalBox)
+								+ SHorizontalBox::Slot()
+								.Padding(20.0, 2.0, 4.0, 2.0)
+								.AutoWidth()
+								[
+									SNew(STextBlock)
+									.Text(FText::FromName(Green))
+									.Font(FEditorStyle::GetFontStyle(TEXT("PropertyWindow.BoldFont")))
+								]
+								+ SHorizontalBox::Slot()
+								.HAlign(HAlign_Left)
+								.Padding(4.0, 2.0)
+								[
+									SNew(STextBlock)
+									.Text(TextureParam->ChannelNames.G)
+									.Font(FEditorStyle::GetFontStyle(TEXT("PropertyWindow.NormalFont")))
+								]
+							];
+					}
+					if (!TextureParam->ChannelNames.B.IsEmpty())
+					{
+						NameVerticalBox->AddSlot()
+							[
+								SNew(SHorizontalBox)
+								+ SHorizontalBox::Slot()
+								.Padding(20.0, 2.0, 4.0, 2.0)
+								.AutoWidth()
+								[
+									SNew(STextBlock)
+									.Text(FText::FromName(Blue))
+									.Font(FEditorStyle::GetFontStyle(TEXT("PropertyWindow.BoldFont")))
+								]
+								+ SHorizontalBox::Slot()
+								.HAlign(HAlign_Left)
+								.Padding(4.0, 2.0)
+								[
+									SNew(STextBlock)
+									.Text(TextureParam->ChannelNames.B)
+									.Font(FEditorStyle::GetFontStyle(TEXT("PropertyWindow.NormalFont")))
+								]
+							];
+					}
+					if (!TextureParam->ChannelNames.A.IsEmpty())
+					{
+						NameVerticalBox->AddSlot()
+							[
+								SNew(SHorizontalBox)
+								+ SHorizontalBox::Slot()
+								.Padding(20.0, 2.0, 4.0, 2.0)
+								.AutoWidth()
+								[
+									SNew(STextBlock)
+									.Text(FText::FromName(Alpha))
+									.Font(FEditorStyle::GetFontStyle(TEXT("PropertyWindow.BoldFont")))
+								]
+								+ SHorizontalBox::Slot()
+								.HAlign(HAlign_Left)
+								.Padding(4.0, 2.0)
+								[
+									SNew(STextBlock)
+									.Text(TextureParam->ChannelNames.A)
+									.Font(FEditorStyle::GetFontStyle(TEXT("PropertyWindow.NormalFont")))
+								]
+							];
+						}
+					}
+			}
+		}
 		else if (!CompMaskParam)
 		{
 			FNodeWidgets StoredNodeWidgets = Node.CreateNodeWidgets();
diff --git a/Engine/Source/Editor/MaterialEditor/Private/SMaterialParametersOverviewWidget.cpp b/Engine/Source/Editor/MaterialEditor/Private/SMaterialParametersOverviewWidget.cpp
--- a/Engine/Source/Editor/MaterialEditor/Private/SMaterialParametersOverviewWidget.cpp
+++ b/Engine/Source/Editor/MaterialEditor/Private/SMaterialParametersOverviewWidget.cpp
@@ -39,6 +39,8 @@
 #include "Curves/CurveLinearColorAtlas.h"
 #include "Widgets/Views/SExpanderArrow.h"
 #include "Widgets/Layout/SWidgetSwitcher.h"
+#include "MaterialEditor/DEditorRuntimeVirtualTextureParameterValue.h"
+
 
 
 #define LOCTEXT_NAMESPACE "MaterialLayerCustomization"
@@ -174,129 +176,157 @@ void SMaterialParametersOverviewTreeItem::Construct(const FArguments& InArgs, co
 		}
 		else if (TextureParam)
 		{
-
-			TSharedPtr<SWidget> NameWidget;
-			TSharedPtr<SWidget> ValueWidget;
-			FDetailWidgetRow DefaultRow;
-			Row.GetDefaultWidgets(NameWidget, ValueWidget, DefaultRow);
-
-			FDetailWidgetRow &DetailWidgetRow = Row.CustomWidget();
-			TSharedPtr<SVerticalBox> NameVerticalBox;
-			DetailWidgetRow.NameContent()
-				[
-					SAssignNew(NameVerticalBox, SVerticalBox)
-					+ SVerticalBox::Slot()
-					.AutoHeight()
-					[
-						SNew(STextBlock)
-						.Text(FText::FromName(StackParameterData->Parameter->ParameterInfo.Name))
-						.Font(FEditorStyle::GetFontStyle(TEXT("PropertyWindow.NormalFont")))
-					]
-				];
-
-			DetailWidgetRow.ValueContent()
-				.MinDesiredWidth(DefaultRow.ValueWidget.MinWidth)
-				.MaxDesiredWidth(DefaultRow.ValueWidget.MaxWidth)
-				[
-					ValueWidget.ToSharedRef()
-				];
-
-			static const FName Red("R");
-			static const FName Green("G");
-			static const FName Blue("B");
-			static const FName Alpha("A");
-
-			if (!TextureParam->ChannelNames.R.IsEmpty())
-			{
-				NameVerticalBox->AddSlot()
-					[
-						SNew(SHorizontalBox)
-						+ SHorizontalBox::Slot()
-						.AutoWidth()
-						.Padding(20.0, 2.0, 4.0, 2.0)
-						[
-							SNew(STextBlock)
-							.Text(FText::FromName(Red))
-							.Font(FEditorStyle::GetFontStyle(TEXT("PropertyWindow.BoldFont")))
-						]
-						+ SHorizontalBox::Slot()
-						.HAlign(HAlign_Left)
-						.Padding(4.0, 2.0)
-						[
-							SNew(STextBlock)
-							.Text(TextureParam->ChannelNames.R)
-							.Font(FEditorStyle::GetFontStyle(TEXT("PropertyWindow.NormalFont")))
-						]
-					];
-			}
-			if (!TextureParam->ChannelNames.G.IsEmpty())
-			{
-				NameVerticalBox->AddSlot()
-					[
-						SNew(SHorizontalBox)
-						+ SHorizontalBox::Slot()
-						.Padding(20.0, 2.0, 4.0, 2.0)
-						.AutoWidth()
-						[
-							SNew(STextBlock)
-							.Text(FText::FromName(Green))
-							.Font(FEditorStyle::GetFontStyle(TEXT("PropertyWindow.BoldFont")))
-						]
-						+ SHorizontalBox::Slot()
-						.HAlign(HAlign_Left)
-						.Padding(4.0, 2.0)
-						[
-							SNew(STextBlock)
-							.Text(TextureParam->ChannelNames.G)
-							.Font(FEditorStyle::GetFontStyle(TEXT("PropertyWindow.NormalFont")))
-						]
-					];
-			}
-			if (!TextureParam->ChannelNames.B.IsEmpty())
+			UMaterial *Material = MaterialEditorInstance->PreviewMaterial;
+			if (Material != nullptr)
 			{
-				NameVerticalBox->AddSlot()
-					[
-						SNew(SHorizontalBox)
-						+ SHorizontalBox::Slot()
-						.Padding(20.0, 2.0, 4.0, 2.0)
-						.AutoWidth()
-						[
-							SNew(STextBlock)
-							.Text(FText::FromName(Blue))
-							.Font(FEditorStyle::GetFontStyle(TEXT("PropertyWindow.BoldFont")))
-						]
-						+ SHorizontalBox::Slot()
-						.HAlign(HAlign_Left)
-						.Padding(4.0, 2.0)
+				UMaterialExpressionTextureSampleParameter* Expression = Material->FindExpressionByGUID<UMaterialExpressionTextureSampleParameter>(TextureParam->ExpressionId);
+				if (Expression != nullptr)
+				{
+					TWeakObjectPtr<UMaterialExpressionTextureSampleParameter> SamplerExpression = Expression;
+					TSharedPtr<SWidget> NameWidget;
+					TSharedPtr<SWidget> ValueWidget;
+					FDetailWidgetRow DefaultRow;
+					Row.GetDefaultWidgets(NameWidget, ValueWidget, DefaultRow);
+
+					FDetailWidgetRow &DetailWidgetRow = Row.CustomWidget();
+					TSharedPtr<SVerticalBox> NameVerticalBox;
+					DetailWidgetRow.NameContent()
 						[
-							SNew(STextBlock)
-							.Text(TextureParam->ChannelNames.B)
-						.Font(FEditorStyle::GetFontStyle(TEXT("PropertyWindow.NormalFont")))
-						]
-					];
-			}
-			if (!TextureParam->ChannelNames.A.IsEmpty())
-			{
-				NameVerticalBox->AddSlot()
-					[
-						SNew(SHorizontalBox)
-						+ SHorizontalBox::Slot()
-						.Padding(20.0, 2.0, 4.0, 2.0)
-						.AutoWidth()
+							SAssignNew(NameVerticalBox, SVerticalBox)
+							+ SVerticalBox::Slot()
+							.AutoHeight()
+							[
+								SNew(STextBlock)
+								.Text(FText::FromName(StackParameterData->Parameter->ParameterInfo.Name))
+								.Font(FEditorStyle::GetFontStyle(TEXT("PropertyWindow.NormalFont")))
+							]
+						];
+
+					DetailWidgetRow.ValueContent()
+						.MinDesiredWidth(DefaultRow.ValueWidget.MinWidth)
+						.MaxDesiredWidth(DefaultRow.ValueWidget.MaxWidth)
 						[
-							SNew(STextBlock)
-							.Text(FText::FromName(Alpha))
+							SNew(SObjectPropertyEntryBox)
+							.PropertyHandle(StackParameterData->ParameterNode->CreatePropertyHandle())
+							.AllowedClass(UTexture::StaticClass())
+							.ThumbnailPool(Tree.Pin()->GetTreeThumbnailPool())
+							.OnShouldFilterAsset_Lambda([SamplerExpression](const FAssetData& AssetData)
+							{
+								if (SamplerExpression.Get())
+								{
+									bool VirtualTextured = false;
+									AssetData.GetTagValue<bool>("VirtualTextureStreaming", VirtualTextured);
+
+									bool ExpressionIsVirtualTextured = IsVirtualSamplerType(SamplerExpression->SamplerType);
+
+									return VirtualTextured != ExpressionIsVirtualTextured;
+								}
+								else
+								{
+									return false;
+								}
+							})
+						];
+
+					static const FName Red("R");
+					static const FName Green("G");
+					static const FName Blue("B");
+					static const FName Alpha("A");
+
+					if (!TextureParam->ChannelNames.R.IsEmpty())
+					{
+						NameVerticalBox->AddSlot()
+							[
+								SNew(SHorizontalBox)
+								+ SHorizontalBox::Slot()
+							.AutoWidth()
+							.Padding(20.0, 2.0, 4.0, 2.0)
+							[
+								SNew(STextBlock)
+								.Text(FText::FromName(Red))
 							.Font(FEditorStyle::GetFontStyle(TEXT("PropertyWindow.BoldFont")))
-						]
+							]
 						+ SHorizontalBox::Slot()
-						.HAlign(HAlign_Left)
-						.Padding(4.0, 2.0)
-						[
-							SNew(STextBlock)
-							.Text(TextureParam->ChannelNames.A)
+							.HAlign(HAlign_Left)
+							.Padding(4.0, 2.0)
+							[
+								SNew(STextBlock)
+								.Text(TextureParam->ChannelNames.R)
 							.Font(FEditorStyle::GetFontStyle(TEXT("PropertyWindow.NormalFont")))
-						]
-					];
+							]
+							];
+					}
+					if (!TextureParam->ChannelNames.G.IsEmpty())
+					{
+						NameVerticalBox->AddSlot()
+							[
+								SNew(SHorizontalBox)
+								+ SHorizontalBox::Slot()
+								.Padding(20.0, 2.0, 4.0, 2.0)
+								.AutoWidth()
+								[
+									SNew(STextBlock)
+									.Text(FText::FromName(Green))
+									.Font(FEditorStyle::GetFontStyle(TEXT("PropertyWindow.BoldFont")))
+								]
+								+ SHorizontalBox::Slot()
+								.HAlign(HAlign_Left)
+								.Padding(4.0, 2.0)
+								[
+									SNew(STextBlock)
+									.Text(TextureParam->ChannelNames.G)
+									.Font(FEditorStyle::GetFontStyle(TEXT("PropertyWindow.NormalFont")))
+								]
+							];
+					}
+					if (!TextureParam->ChannelNames.B.IsEmpty())
+					{
+						NameVerticalBox->AddSlot()
+							[
+								SNew(SHorizontalBox)
+								+ SHorizontalBox::Slot()
+								.Padding(20.0, 2.0, 4.0, 2.0)
+								.AutoWidth()
+								[
+									SNew(STextBlock)
+									.Text(FText::FromName(Blue))
+									.Font(FEditorStyle::GetFontStyle(TEXT("PropertyWindow.BoldFont")))
+								]
+								+ SHorizontalBox::Slot()
+								.HAlign(HAlign_Left)
+								.Padding(4.0, 2.0)
+								[
+									SNew(STextBlock)
+									.Text(TextureParam->ChannelNames.B)
+									.Font(FEditorStyle::GetFontStyle(TEXT("PropertyWindow.NormalFont")))
+								]
+							];
+					}
+					if (!TextureParam->ChannelNames.A.IsEmpty())
+					{
+						NameVerticalBox->AddSlot()
+							[
+								SNew(SHorizontalBox)
+								+ SHorizontalBox::Slot()
+								.Padding(20.0, 2.0, 4.0, 2.0)
+								.AutoWidth()
+								[
+									SNew(STextBlock)
+									.Text(FText::FromName(Alpha))
+									.Font(FEditorStyle::GetFontStyle(TEXT("PropertyWindow.BoldFont")))
+								]
+								+ SHorizontalBox::Slot()
+								.HAlign(HAlign_Left)
+								.Padding(4.0, 2.0)
+								[
+									SNew(STextBlock)
+									.Text(TextureParam->ChannelNames.A)
+									.Font(FEditorStyle::GetFontStyle(TEXT("PropertyWindow.NormalFont")))
+								]
+							];
+					}
+
+				}
 			}
 		}
 		else if (!CompMaskParam)
@@ -715,6 +745,7 @@ void SMaterialParametersOverviewTree::ShowSubParameters()
 
 
 			UDEditorStaticComponentMaskParameterValue* CompMaskParam = Cast<UDEditorStaticComponentMaskParameterValue>(Parameter);
+			// No children for masks
 			if (!CompMaskParam)
 			{
 				TArray<TSharedRef<IDetailTreeNode>> ParamChildren;
@@ -731,16 +762,21 @@ void SMaterialParametersOverviewTree::ShowSubParameters()
 					ChildProperty->Children.Add(ParamChildProperty);
 				}
 			}
-			for (TSharedPtr<struct FSortedParamData> GroupChild : SortedParameters)
+
+			UDEditorRuntimeVirtualTextureParameterValue* VTParameter = Cast<UDEditorRuntimeVirtualTextureParameterValue>(Parameter);
+			//Don't show VT samples here
+			if (!VTParameter)
 			{
-				if (GroupChild->Group.GroupName == Property.ParameterGroup.GroupName
-					&& GroupChild->ParameterInfo.Association == ChildProperty->ParameterInfo.Association
-					&&  GroupChild->ParameterInfo.Index == ChildProperty->ParameterInfo.Index)
+				for (TSharedPtr<struct FSortedParamData> GroupChild : SortedParameters)
 				{
-					GroupChild->Children.Add(ChildProperty);
+					if (GroupChild->Group.GroupName == Property.ParameterGroup.GroupName
+						&& GroupChild->ParameterInfo.Association == ChildProperty->ParameterInfo.Association
+						&&  GroupChild->ParameterInfo.Index == ChildProperty->ParameterInfo.Index)
+					{
+						GroupChild->Children.Add(ChildProperty);
+					}
 				}
 			}
-
 		}
 	}
 }
