Added new Material Editor tab window which lists all the used Custom Primitive Data parameter reads from material expression nodes in the material. They will show up red if there is a potential conflict in reading a specific slot.

diff --git a/Engine/Source/Editor/MaterialEditor/Private/MaterialEditor.cpp b/Engine/Source/Editor/MaterialEditor/Private/MaterialEditor.cpp
--- a/Engine/Source/Editor/MaterialEditor/Private/MaterialEditor.cpp
+++ b/Engine/Source/Editor/MaterialEditor/Private/MaterialEditor.cpp
@@ -130,6 +130,7 @@
 #include "Materials/MaterialExpression.h"
 
 #include "SMaterialParametersOverviewWidget.h"
+#include "SMaterialEditorCustomPrimitiveDataWidget.h"
 #include "IPropertyRowGenerator.h"
 #include "Widgets/Layout/SScrollBox.h"
 #include "UObject/TextProperty.h"
@@ -151,7 +152,8 @@ const FName FMaterialEditor::PropertiesTabId( TEXT( "MaterialEditor_MaterialProp
 const FName FMaterialEditor::PaletteTabId( TEXT( "MaterialEditor_Palette" ) );
 const FName FMaterialEditor::FindTabId( TEXT( "MaterialEditor_Find" ) );
 const FName FMaterialEditor::PreviewSettingsTabId( TEXT ("MaterialEditor_PreviewSettings" ) );
-const FName FMaterialEditor::ParameterDefaultsTabId( TEXT ("MaterialEditor_ParameterDefaults" ) );
+const FName FMaterialEditor::ParameterDefaultsTabId(TEXT("MaterialEditor_ParameterDefaults"));
+const FName FMaterialEditor::CustomPrimitiveTabId(TEXT("MaterialEditor_CustomPrimitiveData"));
 const FName FMaterialEditor::LayerPropertiesTabId(TEXT("MaterialInstanceEditor_MaterialLayerProperties"));
 ///////////////////////////
 // FMatExpressionPreview //
@@ -287,6 +289,11 @@ void FMaterialEditor::RegisterTabSpawners(const TSharedRef<class FTabManager>& I
 		.SetGroup(WorkspaceMenuCategoryRef)
 		.SetIcon(FSlateIcon(FEditorStyle::GetStyleSetName(), "LevelEditor.Tabs.Details"));
 
+	InTabManager->RegisterTabSpawner(CustomPrimitiveTabId, FOnSpawnTab::CreateSP(this, &FMaterialEditor::SpawnTab_CustomPrimitiveData))
+		.SetDisplayName(LOCTEXT("CustomPrimitiveTab", "Custom Primitive Data"))
+		.SetGroup(WorkspaceMenuCategoryRef)
+		.SetIcon(FSlateIcon(FEditorStyle::GetStyleSetName(), "LevelEditor.Tabs.Details"));
+
 	IMaterialEditorModule* MaterialEditorModule = &FModuleManager::LoadModuleChecked<IMaterialEditorModule>("MaterialEditor");
 	if (MaterialEditorModule->MaterialLayersEnabled())
 	{
@@ -312,7 +319,8 @@ void FMaterialEditor::UnregisterTabSpawners(const TSharedRef<class FTabManager>&
 	InTabManager->UnregisterTabSpawner( PaletteTabId );
 	InTabManager->UnregisterTabSpawner( FindTabId );
 	InTabManager->UnregisterTabSpawner( PreviewSettingsTabId );
-	InTabManager->UnregisterTabSpawner( ParameterDefaultsTabId );
+	InTabManager->UnregisterTabSpawner(ParameterDefaultsTabId);
+	InTabManager->UnregisterTabSpawner(CustomPrimitiveTabId);
 	InTabManager->UnregisterTabSpawner( LayerPropertiesTabId );
 
 	MaterialStatsManager->UnregisterTabs();
@@ -478,7 +486,8 @@ void FMaterialEditor::InitMaterialEditor( const EToolkitMode::Type Mode, const T
 					FTabManager::NewStack()
 					->AddTab( PropertiesTabId, ETabState::OpenedTab )
 					->AddTab( PreviewSettingsTabId, ETabState::ClosedTab )
-					->AddTab( ParameterDefaultsTabId, ETabState::OpenedTab )
+					->AddTab(ParameterDefaultsTabId, ETabState::OpenedTab)
+					->AddTab(CustomPrimitiveTabId, ETabState::ClosedTab)
 					->AddTab( LayerPropertiesTabId, ETabState::ClosedTab )
 					->SetForegroundTab( PropertiesTabId )
 				)
@@ -967,6 +976,7 @@ void FMaterialEditor::CreateInternalWidgets()
 	MaterialParametersOverviewWidget = SNew(SMaterialParametersOverviewPanel)
 		.InMaterialEditorInstance(MaterialEditorInstance);
 	MaterialParametersOverviewWidget->GetGenerator()->OnFinishedChangingProperties().AddSP(this, &FMaterialEditor::OnFinishedChangingParametersFromOverview);
+	MaterialCustomPrimitiveDataWidget = SNew(SMaterialCustomPrimitiveDataPanel, MaterialEditorInstance);
 
 	IMaterialEditorModule* MaterialEditorModule = &FModuleManager::LoadModuleChecked<IMaterialEditorModule>("MaterialEditor");
 	if (MaterialEditorModule->MaterialLayersEnabled())
@@ -1020,6 +1030,7 @@ void FMaterialEditor::OnFinishedChangingProperties(const FPropertyChangedEvent&
 		RefreshPreviewViewport();
 		UpdatePreviewMaterial();
 		MaterialParametersOverviewWidget->UpdateEditorInstance(MaterialEditorInstance);
+		MaterialCustomPrimitiveDataWidget->UpdateEditorInstance(MaterialEditorInstance);
 	}
 }
 
@@ -3379,6 +3390,7 @@ void FMaterialEditor::OnPromoteToParameter()
 	if (MaterialEditorInstance != nullptr)
 	{
 		MaterialParametersOverviewWidget->UpdateEditorInstance(MaterialEditorInstance);
+		MaterialCustomPrimitiveDataWidget->UpdateEditorInstance(MaterialEditorInstance);
 	}
 }
 
@@ -3731,6 +3743,21 @@ TSharedRef<SDockTab> FMaterialEditor::SpawnTab_ParameterDefaults(const FSpawnTab
 	return SpawnedTab;
 }
 
+TSharedRef<SDockTab> FMaterialEditor::SpawnTab_CustomPrimitiveData(const FSpawnTabArgs& Args)
+{
+	TSharedRef<SDockTab> SpawnedTab = SNew(SDockTab)
+		.Icon(FEditorStyle::GetBrush("LevelEditor.Tabs.Details"))
+		.Label(LOCTEXT("CustomPrimitiveData", "Custom Primitive Data"))
+		[
+			SNew(SBox)
+			[
+				MaterialCustomPrimitiveDataWidget.ToSharedRef()
+			]
+		];
+
+	return SpawnedTab;
+}
+
 TSharedRef<SDockTab> FMaterialEditor::SpawnTab_LayerProperties(const FSpawnTabArgs& Args)
 {
 	TSharedRef<SDockTab> SpawnedTab = SNew(SDockTab)
@@ -4505,6 +4532,8 @@ void FMaterialEditor::NotifyPostChange( const FPropertyChangedEvent& PropertyCha
 			MaterialLayersFunctionsInstance->SetEditorInstance(MaterialEditorInstance);
 		}
 		MaterialParametersOverviewWidget->UpdateEditorInstance(MaterialEditorInstance);
+		MaterialCustomPrimitiveDataWidget->UpdateEditorInstance(MaterialEditorInstance);
+
 		const FName NameOfPropertyThatChanged( *PropertyThatChanged->GetName() );
 		if ((NameOfPropertyThatChanged == GET_MEMBER_NAME_CHECKED(UMaterialInterface, PreviewMesh)) ||
 			(NameOfPropertyThatChanged == GET_MEMBER_NAME_CHECKED(UMaterial, bUsedWithSkeletalMesh)))
@@ -4763,6 +4792,7 @@ void FMaterialEditor::OnColorPickerCommitted(FLinearColor LinearColor)
 
 	RefreshExpressionPreviews();
 	MaterialParametersOverviewWidget->UpdateEditorInstance(MaterialEditorInstance);
+	MaterialCustomPrimitiveDataWidget->UpdateEditorInstance(MaterialEditorInstance);
 }
 
 TSharedRef<SGraphEditor> FMaterialEditor::CreateGraphEditorWidget()
@@ -5255,6 +5285,7 @@ void FMaterialEditor::OnNodeTitleCommitted(const FText& NewText, ETextCommit::Ty
 		NodeBeingChanged->Modify();
 		NodeBeingChanged->OnRenameNode(NewText.ToString());
 		MaterialParametersOverviewWidget->UpdateEditorInstance(MaterialEditorInstance);
+		MaterialCustomPrimitiveDataWidget->UpdateEditorInstance(MaterialEditorInstance);
 	}
 }
 
diff --git a/Engine/Source/Editor/MaterialEditor/Private/MaterialEditor.h b/Engine/Source/Editor/MaterialEditor/Private/MaterialEditor.h
--- a/Engine/Source/Editor/MaterialEditor/Private/MaterialEditor.h
+++ b/Engine/Source/Editor/MaterialEditor/Private/MaterialEditor.h
@@ -767,6 +767,7 @@ private:
 	TSharedRef<SDockTab> SpawnTab_Find(const FSpawnTabArgs& Args);
 	TSharedRef<SDockTab> SpawnTab_PreviewSettings(const FSpawnTabArgs& Args);
 	TSharedRef<SDockTab> SpawnTab_ParameterDefaults(const FSpawnTabArgs& Args);
+	TSharedRef<SDockTab> SpawnTab_CustomPrimitiveData(const FSpawnTabArgs& Args);
 	TSharedRef<SDockTab> SpawnTab_LayerProperties(const FSpawnTabArgs& Args);
 
 	void OnFinishedChangingProperties(const FPropertyChangedEvent& PropertyChangedEvent);
@@ -803,6 +804,8 @@ private:
 	/** Parameter overview list View */
 	TSharedPtr<class SMaterialParametersOverviewPanel> MaterialParametersOverviewWidget;
 
+	TSharedPtr<class SMaterialCustomPrimitiveDataPanel> MaterialCustomPrimitiveDataWidget;
+
 	/** Layer Properties View */
 	TSharedPtr<class SMaterialLayersFunctionsMaterialWrapper> MaterialLayersFunctionsInstance;
 
@@ -859,6 +862,7 @@ private:
 	static const FName FindTabId;
 	static const FName PreviewSettingsTabId;
 	static const FName ParameterDefaultsTabId;
+	static const FName CustomPrimitiveTabId;
 	static const FName LayerPropertiesTabId;
 
 	/** Object that stores all of the possible parameters we can edit. */
diff --git a/Engine/Source/Editor/MaterialEditor/Private/SMaterialEditorCustomPrimitiveDataWidget.cpp b/Engine/Source/Editor/MaterialEditor/Private/SMaterialEditorCustomPrimitiveDataWidget.cpp
new file mode 100644
--- /dev/null
+++ b/Engine/Source/Editor/MaterialEditor/Private/SMaterialEditorCustomPrimitiveDataWidget.cpp
@@ -0,0 +1,201 @@
+// Copyright 1998-2019 Epic Games, Inc. All Rights Reserved.
+
+#include "SMaterialEditorCustomPrimitiveDataWidget.h"
+#include "Widgets/Layout/SScrollBox.h"
+#include "MaterialEditor/MaterialEditorPreviewParameters.h"
+#include "Materials/MaterialExpressionScalarParameter.h"
+#include "Materials/MaterialExpressionVectorParameter.h"
+#include "EdGraph/EdGraphNode.h"
+#include "EdGraph/EdGraph.h"
+#include "Widgets/SToolTip.h"
+
+class SCustomPrimitiveDataRow : public SMultiColumnTableRow<TSharedPtr<FCustomPrimitiveDataRowData>>
+{
+public:
+	SLATE_BEGIN_ARGS(SCustomPrimitiveDataRow) {}
+		SLATE_ARGUMENT(TSharedPtr<FCustomPrimitiveDataRowData>, Entry)
+	SLATE_END_ARGS()
+
+	void Construct(const FArguments& Args, const TSharedRef<STableViewBase>& OwnerTableView)
+	{
+		RowData = Args._Entry;
+
+		FSuperRowType::Construct(
+			FSuperRowType::FArguments()
+			.Padding(1.0f),
+			OwnerTableView
+		);
+	}
+
+	/** Overridden from SMultiColumnTableRow.  Generates a widget for this column of the list view. */
+	virtual TSharedRef<SWidget> GenerateWidgetForColumn(const FName& ColumnName) override
+	{
+		TSharedRef<STextBlock> TextBlock = SNew(STextBlock);
+
+		if (RowData->bIsDuplicate)
+		{
+			TextBlock->SetToolTip(SNew(SToolTip).Text(FText::FromString("This slot is potentially incorrectly overlapping")).BorderImage(FCoreStyle::Get().GetBrush("ToolTip.BrightBackground")));
+			TextBlock->SetColorAndOpacity(FSlateColor(FLinearColor::Red));
+		}
+
+		if (ColumnName == "Slot")
+		{
+			TextBlock->SetText(FText::FromString(FString::FormatAsNumber(RowData->Slot)));
+			return TextBlock;
+		}
+		else if (ColumnName == "Name")
+		{
+			TextBlock->SetText(FText::FromString(RowData->Name));
+			return TextBlock;
+		}
+
+		return SNullWidget::NullWidget;
+	}
+
+private:
+	TSharedPtr<FCustomPrimitiveDataRowData> RowData;
+};
+
+void SMaterialCustomPrimitiveDataPanel::Refresh()
+{
+	Items.Empty();
+
+	if (MaterialEditorInstance && MaterialEditorInstance->PreviewMaterial)
+	{
+		TArray<UMaterialExpressionScalarParameter*> Scalars;
+		MaterialEditorInstance->PreviewMaterial->GetAllExpressionsInMaterialAndFunctionsOfType<UMaterialExpressionScalarParameter>(Scalars);
+
+		for (const UMaterialExpressionScalarParameter* Expr : Scalars)
+		{
+			if (Expr->bUseCustomPrimitiveData)
+			{
+				FString FunctionName;
+				if (Expr->GraphNode)
+				{
+					Expr->GraphNode->GetGraph()->GetName(FunctionName);
+				}
+
+				Items.Add(MakeShareable(new FCustomPrimitiveDataRowData((int32)Expr->PrimitiveDataIndex, Expr->GetParameterName().ToString(), FunctionName)));
+			}
+		}
+
+		TArray<UMaterialExpressionVectorParameter*> Vectors;
+		MaterialEditorInstance->PreviewMaterial->GetAllExpressionsInMaterialAndFunctionsOfType<UMaterialExpressionVectorParameter>(Vectors);
+
+		for (const UMaterialExpressionVectorParameter* Expr : Vectors)
+		{
+			if (Expr->bUseCustomPrimitiveData)
+			{
+				Items.Add(MakeShareable(new FCustomPrimitiveDataRowData((int32)Expr->PrimitiveDataIndex + 0, Expr->GetParameterName().ToString() + ".r")));
+				Items.Add(MakeShareable(new FCustomPrimitiveDataRowData((int32)Expr->PrimitiveDataIndex + 1, Expr->GetParameterName().ToString() + ".g")));
+				Items.Add(MakeShareable(new FCustomPrimitiveDataRowData((int32)Expr->PrimitiveDataIndex + 2, Expr->GetParameterName().ToString() + ".b")));
+				Items.Add(MakeShareable(new FCustomPrimitiveDataRowData((int32)Expr->PrimitiveDataIndex + 3, Expr->GetParameterName().ToString() + ".a")));
+			}
+		}
+
+		// Sort the items before construction the list
+		Items.StableSort([](const TSharedPtr<FCustomPrimitiveDataRowData>& A, const TSharedPtr<FCustomPrimitiveDataRowData>& B) {
+			int32 SlotA = A->Slot;
+			int32 SlotB = B->Slot;
+			if (SlotA == SlotB)
+			{
+				return A->Name < B->Name;
+			}
+			return SlotA < SlotB;
+		});
+
+		TSet<int32> DuplicateIndices;
+
+		// Once sorted, find suplicate slots
+		for (int32 i = 1; i < Items.Num(); i++)
+		{
+			TSharedPtr<FCustomPrimitiveDataRowData>& PrevItem = Items[i - 1];
+			TSharedPtr<FCustomPrimitiveDataRowData>& CurrItem = Items[i];
+			const int32 PrevSlot = PrevItem->Slot;
+			const int32 CurrSlot = CurrItem->Slot;
+			const FString& PrevName = PrevItem->Name;
+			const FString& CurrName = CurrItem->Name;
+
+			// Duplicate found, store the slot for later processing
+			if ((CurrSlot == PrevSlot) && (CurrName != PrevName))
+			{
+				DuplicateIndices.Add(CurrItem->Slot);
+			}
+		}
+
+		// Go through and mark slots as duplicate
+		for (int32 i = 0; i < Items.Num(); i++)
+		{
+			if (DuplicateIndices.Find(Items[i]->Slot))
+			{
+				Items[i]->bIsDuplicate = true;
+			}
+		}
+	}
+
+	ListViewWidget->RequestListRefresh();
+}
+
+void SMaterialCustomPrimitiveDataPanel::Construct(const FArguments& InArgs, UMaterialEditorPreviewParameters* InMaterialEditorInstance)
+{	
+	this->ChildSlot
+		[
+			SNew(SVerticalBox)
+			+ SVerticalBox::Slot()
+			[
+				SNew(SBorder)
+				.BorderImage(FEditorStyle::GetBrush("DetailsView.CategoryTop_Hovered"))
+				.Padding(FMargin(4.0f))
+				[
+					SNew(SVerticalBox)
+					+ SVerticalBox::Slot()
+					.AutoHeight()
+					[
+						SNew(SHorizontalBox)
+						+ SHorizontalBox::Slot()
+						.Padding(FMargin(3.0f, 4.0f))
+						.HAlign(HAlign_Left)
+						.AutoWidth()
+						.VAlign(VAlign_Center)
+						[
+							SNew(STextBlock)
+							.Text(FText::FromString("Custom Primitive Data Parameters"))
+							.Font(FEditorStyle::GetFontStyle("DetailsView.CategoryFontStyle"))
+							.ShadowOffset(FVector2D(1.0f, 1.0f))
+						]
+					]
+					+ SVerticalBox::Slot()
+					.Padding(FMargin(3.0f, 2.0f, 3.0f, 3.0f))
+					[
+						SNew(SBorder)
+						.BorderImage(FEditorStyle::GetBrush("DetailsView.CategoryTop"))
+						[
+							SNew(SScrollBox)
+							+ SScrollBox::Slot()
+							[
+								SAssignNew(ListViewWidget, SListView<TSharedPtr<FCustomPrimitiveDataRowData>>)
+								.ItemHeight(24)
+								.ListItemsSource(&Items)
+								.OnGenerateRow(this, &SMaterialCustomPrimitiveDataPanel::OnGenerateRowForList)
+								.SelectionMode(ESelectionMode::None)
+								.HeaderRow
+								(
+									SNew(SHeaderRow)
+									+ SHeaderRow::Column("Slot").DefaultLabel(FText::FromString("Slot")).ManualWidth(48.0f)
+									+ SHeaderRow::Column("Name").DefaultLabel(FText::FromString("Name"))
+								)
+							]
+						]
+					]
+				]
+			]
+		];
+
+	MaterialEditorInstance = InMaterialEditorInstance;
+	Refresh();
+}
+
+TSharedRef<ITableRow> SMaterialCustomPrimitiveDataPanel::OnGenerateRowForList(TSharedPtr<FCustomPrimitiveDataRowData> Item, const TSharedRef<STableViewBase>& OwnerTable)
+{
+	return SNew(SCustomPrimitiveDataRow, OwnerTable).Entry(Item);
+}
diff --git a/Engine/Source/Editor/MaterialEditor/Private/SMaterialEditorCustomPrimitiveDataWidget.h b/Engine/Source/Editor/MaterialEditor/Private/SMaterialEditorCustomPrimitiveDataWidget.h
new file mode 100644
--- /dev/null
+++ b/Engine/Source/Editor/MaterialEditor/Private/SMaterialEditorCustomPrimitiveDataWidget.h
@@ -0,0 +1,54 @@
+// Copyright 1998-2019 Epic Games, Inc. All Rights Reserved.
+
+#pragma once
+
+#include "CoreMinimal.h"
+#include "UObject/ObjectMacros.h"
+#include "UObject/Object.h"
+#include "Engine/EngineTypes.h"
+#include "Widgets/DeclarativeSyntaxSupport.h"
+#include "Widgets/SCompoundWidget.h"
+#include "Widgets/Views/SListView.h"
+#include "Materials/Material.h"
+#include "IDetailTreeNode.h"
+#include "IDetailPropertyRow.h"
+#include "MaterialPropertyHelpers.h"
+
+class UMaterialEditorPreviewParameters;
+
+/** Data for a row in the list of custom primitive data entries*/
+struct FCustomPrimitiveDataRowData
+{
+	FCustomPrimitiveDataRowData(int32 InSlot, FString InName, FString InUsage = "") : Slot(InSlot), Name(InName) {}
+
+	int32 Slot = 0;
+	FString Name = "";
+	bool bIsDuplicate = false;
+};
+
+class SMaterialCustomPrimitiveDataPanel : public SCompoundWidget
+{
+public:
+	SLATE_BEGIN_ARGS(SMaterialCustomPrimitiveDataPanel) {}
+	SLATE_END_ARGS()
+
+	void Construct(const FArguments& InArgs, UMaterialEditorPreviewParameters* InMaterialEditorInstance);
+
+	void UpdateEditorInstance(UMaterialEditorPreviewParameters* InMaterialEditorInstance) { MaterialEditorInstance = InMaterialEditorInstance; Refresh(); }
+
+private:
+
+	void Refresh();
+
+	/** Adds a new textbox with the string to the list */
+	TSharedRef<ITableRow> OnGenerateRowForList(TSharedPtr<FCustomPrimitiveDataRowData> Item, const TSharedRef<STableViewBase>& OwnerTable);
+
+	/** The list of strings */
+	TArray<TSharedPtr<FCustomPrimitiveDataRowData>> Items;
+
+	/** The actual UI list */
+	TSharedPtr<SListView<TSharedPtr<FCustomPrimitiveDataRowData>>> ListViewWidget;
+
+	/** The set of material parameters this is associated with */
+	UMaterialEditorPreviewParameters* MaterialEditorInstance;
+};
\ No newline at end of file
