Context menus now close when an undo or redo action is performed in graph editors.

diff --git a/Engine/Source/Editor/AudioEditor/Private/SoundClassEditor.cpp b/Engine/Source/Editor/AudioEditor/Private/SoundClassEditor.cpp
--- a/Engine/Source/Editor/AudioEditor/Private/SoundClassEditor.cpp
+++ b/Engine/Source/Editor/AudioEditor/Private/SoundClassEditor.cpp
@@ -465,6 +465,8 @@ void FSoundClassEditor::PostUndo(bool bSuccess)
 {
 	GraphEditor->ClearSelectionSet();
 	GraphEditor->NotifyGraphChanged();
+
+	FSlateApplication::Get().DismissAllMenus();
 }
 
 #undef LOCTEXT_NAMESPACE
diff --git a/Engine/Source/Editor/AudioEditor/Private/SoundCueEditor.cpp b/Engine/Source/Editor/AudioEditor/Private/SoundCueEditor.cpp
--- a/Engine/Source/Editor/AudioEditor/Private/SoundCueEditor.cpp
+++ b/Engine/Source/Editor/AudioEditor/Private/SoundCueEditor.cpp
@@ -256,6 +256,7 @@ void FSoundCueEditor::PostUndo(bool bSuccess)
 	{
 		SoundCueGraphEditor->ClearSelectionSet();
 		SoundCueGraphEditor->NotifyGraphChanged();
+		FSlateApplication::Get().DismissAllMenus();
 	}
 
 }
diff --git a/Engine/Source/Editor/MaterialEditor/Private/MaterialEditor.cpp b/Engine/Source/Editor/MaterialEditor/Private/MaterialEditor.cpp
--- a/Engine/Source/Editor/MaterialEditor/Private/MaterialEditor.cpp
+++ b/Engine/Source/Editor/MaterialEditor/Private/MaterialEditor.cpp
@@ -4430,9 +4430,6 @@ void FMaterialEditor::UndoGraphAction()
 void FMaterialEditor::RedoGraphAction()
 {
 	FlushRenderingCommands();
-	
-	// Clear selection, to avoid holding refs to nodes that go away
-	GraphEditor->ClearSelectionSet();
 
 	int32 NumExpressions = Material->Expressions.Num();
 	GEditor->RedoTransaction();
@@ -4533,6 +4530,8 @@ void FMaterialEditor::PostUndo(bool bSuccess)
 		RefreshExpressionPreviews();
 		GraphEditor->NotifyGraphChanged();
 		SetMaterialDirty();
+
+		FSlateApplication::Get().DismissAllMenus();
 	}
 }
 
