Fixed "User Interface" material domain for UI shaders; the default UI material no longer uses an instanced vertex shader; also added check for null pointer.

#rnx

diff --git a/Engine/Source/Editor/MaterialEditor/Private/MaterialStatsCommon.cpp b/Engine/Source/Editor/MaterialEditor/Private/MaterialStatsCommon.cpp
--- a/Engine/Source/Editor/MaterialEditor/Private/MaterialStatsCommon.cpp
+++ b/Engine/Source/Editor/MaterialEditor/Private/MaterialStatsCommon.cpp
@@ -342,9 +342,9 @@ void FMaterialStatsUtils::GetRepresentativeShaderTypesAndDescriptions(TMap<FName
 
 	if (TargetMaterial->IsUIMaterial())
 	{
-		static FName TSlateMaterialShaderPSDefaultfalseName = TEXT("TSlateMaterialShaderPSDefaultfalse");
+		static FName TSlateMaterialShaderPSDefaultName = TEXT("TSlateMaterialShaderPSDefault");
 		ShaderTypeNamesAndDescriptions.FindOrAdd(FLocalVertexFactoryName)
-			.Add(FRepresentativeShaderInfo(ERepresentativeShader::UIDefaultFragmentShader, TSlateMaterialShaderPSDefaultfalseName, TEXT("Default UI Pixel Shader")));
+			.Add(FRepresentativeShaderInfo(ERepresentativeShader::UIDefaultFragmentShader, TSlateMaterialShaderPSDefaultName, TEXT("Default UI Pixel Shader")));
 
 		static FName TSlateMaterialShaderVSfalseName = TEXT("TSlateMaterialShaderVSfalse");
 		ShaderTypeNamesAndDescriptions.FindOrAdd(FLocalVertexFactoryName)
@@ -600,6 +600,7 @@ void FMaterialStatsUtils::GetRepresentativeInstructionCounts(TArray<FShaderInstr
 					const FRepresentativeShaderInfo& ShaderInfo = DescriptionArray[i];
 
 					FShaderType* ShaderType = FindShaderTypeByName(ShaderInfo.ShaderName);
+					check(ShaderType);
 					const int32 NumInstructions = MaterialShaderMap->GetMaxNumInstructionsForShader(ShaderType);
 
 					FShaderInstructionsInfo Info;
