Fix duplicate items on Material Instance Toolbar

diff --git a/Engine/Source/Editor/MaterialEditor/Private/MaterialInstanceEditor.cpp b/Engine/Source/Editor/MaterialEditor/Private/MaterialInstanceEditor.cpp
--- a/Engine/Source/Editor/MaterialEditor/Private/MaterialInstanceEditor.cpp
+++ b/Engine/Source/Editor/MaterialEditor/Private/MaterialInstanceEditor.cpp
@@ -22,6 +22,7 @@
 #include "MaterialEditorModule.h"
 #include "ToolMenus.h"
 #include "Toolkits/AssetEditorToolkitMenuContext.h"
+#include "MaterialEditorContext.h"
 
 #include "Materials/MaterialExpressionTextureBase.h"
 #include "Materials/MaterialExpressionTextureSampleParameter.h"
@@ -385,6 +386,7 @@ void FMaterialInstanceEditor::InitMaterialInstanceEditor( const EToolkitMode::Ty
 	CreateInternalWidgets();
 
 	BindCommands();
+	RegisterToolBar();
 
 	UpdatePreviewViewportsVisibility();
 	IMaterialEditorModule* MaterialEditorModule = &FModuleManager::LoadModuleChecked<IMaterialEditorModule>("MaterialEditor");
@@ -797,41 +799,65 @@ void FMaterialInstanceEditor::UpdatePreviewViewportsVisibility()
 
 void FMaterialInstanceEditor::RegisterToolBar()
 {
-	UToolMenus* ToolMenus = UToolMenus::Get();
-	UToolMenu* ToolBar = ToolMenus->ExtendMenu(GetToolMenuToolbarName());
-
-	FToolMenuInsert InsertAfterAssetSection("Asset", EToolMenuInsertType::After);
+	const FName MenuName = GetToolMenuToolbarName();
+	if (!UToolMenus::Get()->IsMenuRegistered(MenuName))
 	{
-		FToolMenuSection& Section = ToolBar->AddSection("Apply", TAttribute<FText>(), InsertAfterAssetSection);
-		Section.AddEntry(FToolMenuEntry::InitToolBarButton(FMaterialEditorCommands::Get().Apply));
-	}
+		UToolMenu* ToolBar = UToolMenus::Get()->RegisterMenu(MenuName, "AssetEditor.DefaultToolBar", EMultiBoxType::ToolBar);
 
-	{
-		FToolMenuSection& Section = ToolBar->AddSection("Command", TAttribute<FText>(), InsertAfterAssetSection);
-		Section.AddEntry(FToolMenuEntry::InitToolBarButton(FMaterialEditorCommands::Get().ShowAllMaterialParameters));
-		// TODO: support in material instance editor.
-		Section.AddEntry(FToolMenuEntry::InitToolBarButton(FMaterialEditorCommands::Get().TogglePlatformStats));
-	}
-	
-	{
-		FToolMenuSection& Section = ToolBar->AddSection("Parent", TAttribute<FText>(), InsertAfterAssetSection);
-		Section.AddEntry(FToolMenuEntry::InitComboButton(
-			"Hierarchy",
-			FToolUIActionChoice(),
-			FNewToolMenuDelegate::CreateSP(this, &FMaterialInstanceEditor::GenerateInheritanceMenu),
-			LOCTEXT("Hierarchy", "Hierarchy"),
-			FText::GetEmpty(),
-			FSlateIcon(FEditorStyle::GetStyleSetName(), "BTEditor.SwitchToBehaviorTreeMode"),
-			false
-		));
+		FToolMenuInsert InsertAfterAssetSection("Asset", EToolMenuInsertType::After);
+		{
+			FToolMenuSection& Section = ToolBar->AddSection("Apply", TAttribute<FText>(), InsertAfterAssetSection);
+			Section.AddEntry(FToolMenuEntry::InitToolBarButton(FMaterialEditorCommands::Get().Apply));
+		}
+
+		{
+			FToolMenuSection& Section = ToolBar->AddSection("Command", TAttribute<FText>(), InsertAfterAssetSection);
+			Section.AddEntry(FToolMenuEntry::InitToolBarButton(FMaterialEditorCommands::Get().ShowAllMaterialParameters));
+			// TODO: support in material instance editor.
+			Section.AddEntry(FToolMenuEntry::InitToolBarButton(FMaterialEditorCommands::Get().TogglePlatformStats));
+		}
+
+		{
+			FToolMenuSection& Section = ToolBar->AddSection("Parent", TAttribute<FText>(), InsertAfterAssetSection);
+
+			struct Local
+			{
+				static void GenerateInheritanceMenu(UToolMenu* InMenu)
+				{
+					UMaterialEditorMenuContext* Context = InMenu->FindContext<UMaterialEditorMenuContext>();
+					if (Context && Context->MaterialEditor.IsValid())
+					{
+						TSharedPtr<IMaterialEditor> Pinned = Context->MaterialEditor.Pin();
+						TSharedPtr<FMaterialInstanceEditor> Recast = StaticCastSharedPtr<FMaterialInstanceEditor>(Pinned);
+						Recast->GenerateInheritanceMenu(InMenu);
+					}
+				}
+			};
+
+			Section.AddEntry(FToolMenuEntry::InitComboButton(
+				"Hierarchy",
+				FToolUIActionChoice(),
+				FNewToolMenuDelegate::CreateStatic(&Local::GenerateInheritanceMenu),
+				LOCTEXT("Hierarchy", "Hierarchy"),
+				FText::GetEmpty(),
+				FSlateIcon(FEditorStyle::GetStyleSetName(), "BTEditor.SwitchToBehaviorTreeMode"),
+				false
+			));
+		}
 	}
+}
+
+void FMaterialInstanceEditor::InitToolMenuContext(FToolMenuContext& MenuContext)
+{
+	FAssetEditorToolkit::InitToolMenuContext(MenuContext);
 
+	UMaterialEditorMenuContext* Context = NewObject<UMaterialEditorMenuContext>();
+	Context->MaterialEditor = SharedThis(this);
+	MenuContext.AddObject(Context);
 }
 
 void FMaterialInstanceEditor::ExtendToolbar()
 {
-	RegisterToolBar();
-
 	AddToolbarExtender(GetToolBarExtensibilityManager()->GetAllExtenders(GetToolkitCommands(), GetEditingObjects()));
 
 	IMaterialEditorModule* MaterialEditorModule = &FModuleManager::LoadModuleChecked<IMaterialEditorModule>( "MaterialEditor" );
diff --git a/Engine/Source/Editor/MaterialEditor/Private/MaterialInstanceEditor.h b/Engine/Source/Editor/MaterialEditor/Private/MaterialInstanceEditor.h
--- a/Engine/Source/Editor/MaterialEditor/Private/MaterialInstanceEditor.h
+++ b/Engine/Source/Editor/MaterialEditor/Private/MaterialInstanceEditor.h
@@ -59,6 +59,7 @@ public:
 	virtual FName GetToolkitFName() const override;
 	virtual FText GetBaseToolkitName() const override;
 	virtual FString GetWorldCentricTabPrefix() const override;
+	virtual void InitToolMenuContext(struct FToolMenuContext& MenuContext) override;
 
 	/** @return Returns the color and opacity to use for the color that appears behind the tab text for this toolkit's tab in world-centric mode. */
 	virtual FLinearColor GetWorldCentricTabColorScale() const override;
diff --git a/Engine/Source/Editor/MaterialEditor/Public/MaterialEditorContext.h b/Engine/Source/Editor/MaterialEditor/Public/MaterialEditorContext.h
new file mode 100644
--- /dev/null
+++ b/Engine/Source/Editor/MaterialEditor/Public/MaterialEditorContext.h
@@ -0,0 +1,19 @@
+// Copyright Epic Games, Inc. All Rights Reserved.
+
+#pragma once
+
+#include "CoreMinimal.h"
+
+#include "MaterialEditorContext.generated.h"
+
+class IMaterialEditor;
+
+UCLASS()
+class MATERIALEDITOR_API UMaterialEditorMenuContext : public UObject
+{
+	GENERATED_BODY()
+
+public:
+
+	TWeakPtr<IMaterialEditor> MaterialEditor;
+};
