Material Editor: Toolbar converted into tool menus. Hierarchy menu converted to use same custom widget as material instance editor.

diff --git a/Engine/Source/Editor/MaterialEditor/Private/MaterialEditor.cpp b/Engine/Source/Editor/MaterialEditor/Private/MaterialEditor.cpp
--- a/Engine/Source/Editor/MaterialEditor/Private/MaterialEditor.cpp
+++ b/Engine/Source/Editor/MaterialEditor/Private/MaterialEditor.cpp
@@ -135,6 +135,8 @@
 #include "Widgets/Layout/SScrollBox.h"
 #include "UObject/TextProperty.h"
 #include "Subsystems/AssetEditorSubsystem.h"
+#include "ToolMenus.h"
+#include "MaterialEditorContext.h"
 #include "UObject/MetaData.h"
 
 #define LOCTEXT_NAMESPACE "MaterialEditor"
@@ -458,6 +460,7 @@ void FMaterialEditor::InitMaterialEditor( const EToolkitMode::Type Mode, const T
 	UpdatePreviewViewportsVisibility();
 
 	BindCommands();
+	RegisterToolBar();
 
 	const TSharedRef<FTabManager::FLayout> StandaloneDefaultLayout = FTabManager::NewLayout("Standalone_MaterialEditor_Layout_v10")
 	->AddArea
@@ -1118,142 +1121,221 @@ void FMaterialEditor::UpdateThumbnailInfoPreviewMesh(UMaterialInterface* MatInte
 
 void FMaterialEditor::ExtendToolbar()
 {
-	TSharedPtr<FExtender> ToolbarExtender = MakeShareable(new FExtender);
-
-	ToolbarExtender->AddToolBarExtension(
-		"Asset",
-		EExtensionHook::After,
-		GetToolkitCommands(),
-		FToolBarExtensionDelegate::CreateSP( this, &FMaterialEditor::FillToolbar )
-		);
-	
-	AddToolbarExtender(ToolbarExtender);
-
 	AddToolbarExtender(GetToolBarExtensibilityManager()->GetAllExtenders(GetToolkitCommands(), GetEditingObjects()));
 
 	IMaterialEditorModule* MaterialEditorModule = &FModuleManager::LoadModuleChecked<IMaterialEditorModule>( "MaterialEditor" );
 	AddToolbarExtender(MaterialEditorModule->GetToolBarExtensibilityManager()->GetAllExtenders(GetToolkitCommands(), GetEditingObjects()));
 }
 
-void FMaterialEditor::FillToolbar(FToolBarBuilder& ToolbarBuilder)
+void FMaterialEditor::InitToolMenuContext(FToolMenuContext& MenuContext)
 {
-	ToolbarBuilder.BeginSection("Apply");
-	{
-		ToolbarBuilder.AddToolBarButton(FMaterialEditorCommands::Get().Apply);
-	}
-	ToolbarBuilder.EndSection();
+	FAssetEditorToolkit::InitToolMenuContext(MenuContext);
 
-	ToolbarBuilder.BeginSection("Search");
-	{
-		ToolbarBuilder.AddToolBarButton(FMaterialEditorCommands::Get().FindInMaterial);
-	}
-	ToolbarBuilder.EndSection();
+	UMaterialEditorMenuContext* Context = NewObject<UMaterialEditorMenuContext>();
+	Context->MaterialEditor = SharedThis(this);
+	MenuContext.AddObject(Context);
+}
 
-	ToolbarBuilder.BeginSection("Graph");
+void FMaterialEditor::RegisterToolBar()
+{
+	const FName MenuName = GetToolMenuToolbarName();
+	if (!UToolMenus::Get()->IsMenuRegistered(MenuName))
 	{
-		ToolbarBuilder.AddToolBarButton(FMaterialEditorCommands::Get().CameraHome);
-		ToolbarBuilder.AddToolBarButton(FMaterialEditorCommands::Get().CleanUnusedExpressions);
-		ToolbarBuilder.AddToolBarButton(FMaterialEditorCommands::Get().ShowHideConnectors);
-		ToolbarBuilder.AddToolBarButton(FMaterialEditorCommands::Get().ToggleLivePreview);
-		ToolbarBuilder.AddToolBarButton(FMaterialEditorCommands::Get().ToggleRealtimeExpressions);
-		ToolbarBuilder.AddToolBarButton(FMaterialEditorCommands::Get().AlwaysRefreshAllPreviews);
-		ToolbarBuilder.AddToolBarButton(
-			FMaterialEditorCommands::Get().ToggleHideUnrelatedNodes,
-				NAME_None, 
-				TAttribute<FText>(), 
-				TAttribute<FText>(), 
+		UToolMenu* ToolBar = UToolMenus::Get()->RegisterMenu(MenuName, "AssetEditor.DefaultToolBar", EMultiBoxType::ToolBar);
+
+		FToolMenuInsert InsertAfterAssetSection("Asset", EToolMenuInsertType::After);
+
+		{
+			FToolMenuSection& Section = ToolBar->AddSection("Apply", TAttribute<FText>(), InsertAfterAssetSection);
+			Section.AddEntry(FToolMenuEntry::InitToolBarButton(FMaterialEditorCommands::Get().Apply));
+		}
+
+		{
+			FToolMenuSection& Section = ToolBar->AddSection("Search", TAttribute<FText>(), InsertAfterAssetSection);
+			Section.AddEntry(FToolMenuEntry::InitToolBarButton(FMaterialEditorCommands::Get().FindInMaterial));
+		}
+
+		{
+			FToolMenuSection& Section = ToolBar->AddSection("Graph", TAttribute<FText>(), InsertAfterAssetSection);
+			Section.AddEntry(FToolMenuEntry::InitToolBarButton(FMaterialEditorCommands::Get().CameraHome));
+			Section.AddEntry(FToolMenuEntry::InitToolBarButton(FMaterialEditorCommands::Get().CleanUnusedExpressions));
+			Section.AddEntry(FToolMenuEntry::InitToolBarButton(FMaterialEditorCommands::Get().ShowHideConnectors));
+			Section.AddEntry(FToolMenuEntry::InitToolBarButton(FMaterialEditorCommands::Get().ToggleLivePreview));
+			Section.AddEntry(FToolMenuEntry::InitToolBarButton(FMaterialEditorCommands::Get().ToggleRealtimeExpressions));
+			Section.AddEntry(FToolMenuEntry::InitToolBarButton(FMaterialEditorCommands::Get().AlwaysRefreshAllPreviews));
+			Section.AddEntry(FToolMenuEntry::InitToolBarButton(
+				FMaterialEditorCommands::Get().ToggleHideUnrelatedNodes,
+				TAttribute<FText>(),
+				TAttribute<FText>(),
 				FSlateIcon(FEditorStyle::GetStyleSetName(), "GraphEditor.ToggleHideUnrelatedNodes")
-			);
-		ToolbarBuilder.AddComboButton(
-			FUIAction(),
-			FOnGetContent::CreateSP(this, &FMaterialEditor::MakeHideUnrelatedNodesOptionsMenu),
-			LOCTEXT("HideUnrelatedNodesOptions", "Hide Unrelated Nodes Options"),
-			LOCTEXT("HideUnrelatedNodesOptionsMenu", "Hide Unrelated Nodes options menu"),
-			TAttribute<FSlateIcon>(),
-			true
-		); 
-	}
-	ToolbarBuilder.EndSection();
-
-	ToolbarBuilder.BeginSection("Stats");
-	{
-		ToolbarBuilder.AddToolBarButton(FMaterialEditorCommands::Get().ToggleMaterialStats);
-		ToolbarBuilder.AddToolBarButton(FMaterialEditorCommands::Get().TogglePlatformStats);
-		ToolbarBuilder.AddComboButton(
-			FUIAction(),
-			FOnGetContent::CreateSP(this, &FMaterialEditor::GeneratePreviewMenuContent),
-			LOCTEXT("NodePreview_Label", "Preview Nodes"),
-			LOCTEXT("NodePreviewToolTip", "Preview the nodes for a given feature level and/or material quality."),
-			FSlateIcon(FEditorStyle::GetStyleSetName(), "FullBlueprintEditor.SwitchToScriptingMode"),
-			false);
-	}
-	ToolbarBuilder.EndSection();
-	if (!MaterialFunction)
-	{
-		ToolbarBuilder.BeginSection("Hierarchy");
+			));
+			Section.AddEntry(FToolMenuEntry::InitComboButton(
+				"HideUnrelatedNodesOptions",
+				FUIAction(),
+				FOnGetContent::CreateSP(this, &FMaterialEditor::MakeHideUnrelatedNodesOptionsMenu),
+				LOCTEXT("HideUnrelatedNodesOptions", "Hide Unrelated Nodes Options"),
+				LOCTEXT("HideUnrelatedNodesOptionsMenu", "Hide Unrelated Nodes options menu"),
+				TAttribute<FSlateIcon>(),
+				true
+			));
+		}
+
 		{
-			ToolbarBuilder.AddComboButton(
+			FToolMenuSection& Section = ToolBar->AddSection("Stats", TAttribute<FText>(), InsertAfterAssetSection);
+			Section.AddEntry(FToolMenuEntry::InitToolBarButton(FMaterialEditorCommands::Get().ToggleMaterialStats));
+			Section.AddEntry(FToolMenuEntry::InitToolBarButton(FMaterialEditorCommands::Get().TogglePlatformStats));
+			Section.AddEntry(FToolMenuEntry::InitComboButton(
+				"NodePreview",
 				FUIAction(),
-				FOnGetContent::CreateSP(this, &FMaterialEditor::GenerateInheritanceMenu),
-				LOCTEXT("Hierarchy", "Hierarchy"),
-				FText::GetEmpty(),
-				FSlateIcon(FSlateIcon(FEditorStyle::GetStyleSetName(), "BTEditor.SwitchToBehaviorTreeMode")),
+				FOnGetContent::CreateSP(this, &FMaterialEditor::GeneratePreviewMenuContent),
+				LOCTEXT("NodePreview_Label", "Preview Nodes"),
+				LOCTEXT("NodePreviewToolTip", "Preview the nodes for a given feature level and/or material quality."),
+				FSlateIcon(FEditorStyle::GetStyleSetName(), "FullBlueprintEditor.SwitchToScriptingMode"),
 				false
-			);
+			));
 		}
-		ToolbarBuilder.EndSection();
+
+		ToolBar->AddDynamicSection("Hierarchy", FNewToolMenuDelegate::CreateLambda([](UToolMenu* InMenu)
+		{
+			UMaterialEditorMenuContext* Context = InMenu->FindContext<UMaterialEditorMenuContext>();
+			TSharedPtr<FMaterialEditor> MaterialEditor = StaticCastSharedPtr<FMaterialEditor>(Context->MaterialEditor.Pin());
+			if (!MaterialEditor->MaterialFunction)
+			{
+				{
+					FToolMenuSection& Section = InMenu->AddSection("Hierarchy");
+					Section.AddEntry(FToolMenuEntry::InitComboButton(
+						"Hierarchy",
+						FToolUIActionChoice(),
+						FNewToolMenuDelegate::CreateLambda([](UToolMenu* InSubMenu)
+						{
+							UMaterialEditorMenuContext* SubMenuContext = InSubMenu->FindContext<UMaterialEditorMenuContext>();
+							if (SubMenuContext && SubMenuContext->MaterialEditor.IsValid())
+							{
+								SubMenuContext->MaterialEditor.Pin()->GenerateInheritanceMenu(InSubMenu);
+							}
+						}),
+						LOCTEXT("Hierarchy", "Hierarchy"),
+						FText::GetEmpty(),
+						FSlateIcon(FSlateIcon(FEditorStyle::GetStyleSetName(), "BTEditor.SwitchToBehaviorTreeMode")),
+						false
+					));
+				}
+			}
+		})).InsertPosition = InsertAfterAssetSection;
 	}
 };
 
-TSharedRef<SWidget> FMaterialEditor::GenerateInheritanceMenu()
+void FMaterialEditor::AddInheritanceMenuEntry(FToolMenuSection& Section, const FAssetData& AssetData, bool bIsFunctionPreviewMaterial)
 {
-	struct Local
+	FExecuteAction OpenAction;
+	FExecuteAction FindInContentBrowserAction;
+	if (bIsFunctionPreviewMaterial)
 	{
-		static void MakeMaterialSubmenu(FMenuBuilder& MenuBuilder, FAssetData InMaterial)
-		{
-			MenuBuilder.AddMenuEntry(LOCTEXT("OpenInEditor", "Open In Editor"),
-				FText::GetEmpty(),
-				FSlateIcon(FEditorStyle::GetStyleSetName(), "ContentBrowser.AssetActions.OpenInExternalEditor"),
-				FUIAction(FExecuteAction::CreateStatic(&FMaterialEditorUtilities::OnOpenMaterial, InMaterial)));
+		OpenAction.BindStatic(&FMaterialEditorUtilities::OnOpenFunction, AssetData);
+		FindInContentBrowserAction.BindStatic(&FMaterialEditorUtilities::OnShowFunctionInContentBrowser, AssetData);
+	}
+	else
+	{
+		OpenAction.BindStatic(&FMaterialEditorUtilities::OnOpenMaterial, AssetData);
+		FindInContentBrowserAction.BindStatic(&FMaterialEditorUtilities::OnShowMaterialInContentBrowser, AssetData);
+	}
 
-			MenuBuilder.AddMenuEntry(LOCTEXT("FindInContentBrowser", "Find In Content Browser"),
-				FText::GetEmpty(),
-				FSlateIcon(FEditorStyle::GetStyleSetName(), "SystemWideCommands.FindInContentBrowser"),
-				FUIAction(FExecuteAction::CreateStatic(&FMaterialEditorUtilities::OnShowMaterialInContentBrowser, InMaterial)));
-		}
-	};
+	FFormatNamedArguments Args;
+	Args.Add(TEXT("ParentName"), FText::FromName(AssetData.AssetName));
+	FText Label = FText::Format(LOCTEXT("InstanceParentName", "{ParentName}"), Args);
+
+	FSlateIcon OpenIcon(FEditorStyle::GetStyleSetName(), "ContentBrowser.AssetActions.OpenInExternalEditor");
+	FSlateIcon FindInContentBrowserIcon(FEditorStyle::GetStyleSetName(), "SystemWideCommands.FindInContentBrowser");
+
+	TSharedRef<SWidget> EntryWidget =
+		SNew(SHorizontalBox)
+		.ToolTipText(LOCTEXT("OpenInEditor", "Open In Editor"))
+
+		+ SHorizontalBox::Slot()
+		.AutoWidth()
+		.Padding(FMargin(2, 0, 2, 0))
+		[
+			SNew( SBox )
+			.WidthOverride( MultiBoxConstants::MenuIconSize + 2 )
+			.HeightOverride( MultiBoxConstants::MenuIconSize )
+			.HAlign(HAlign_Center)
+			.VAlign(VAlign_Center)
+			[
+				SNew( SBox )
+				.WidthOverride( MultiBoxConstants::MenuIconSize )
+				.HeightOverride( MultiBoxConstants::MenuIconSize )
+				[
+					SNew(SImage)
+					.Image(OpenIcon.GetIcon())
+				]
+			]
+		]
+
+		+ SHorizontalBox::Slot()
+		.FillWidth( 1.0f )
+		.Padding(FMargin(2, 0, 6, 0))
+		.VAlign( VAlign_Center )
+		[
+			SNew(STextBlock)
+			.Text(Label)
+		]
+
+		+ SHorizontalBox::Slot()
+		.AutoWidth()
+		.VAlign( VAlign_Center )
+		.HAlign( HAlign_Right )
+		[
+			SNew(SButton)
+			.ButtonStyle(FEditorStyle::Get(), "HoverHintOnly")
+			.ToolTipText(LOCTEXT("FindInContentBrowser", "Find In Content Browser"))
+			.OnClicked_Lambda([FindInContentBrowserAction]() { FindInContentBrowserAction.ExecuteIfBound(); return FReply::Handled(); })
+			[
+				SNew( SBox )
+				.WidthOverride( MultiBoxConstants::MenuIconSize + 2 )
+				.HeightOverride( MultiBoxConstants::MenuIconSize )
+				.HAlign(HAlign_Center)
+				.VAlign(VAlign_Center)
+				[
+					SNew( SBox )
+					.WidthOverride( MultiBoxConstants::MenuIconSize )
+					.HeightOverride( MultiBoxConstants::MenuIconSize )
+					[
+						SNew(SImage)
+						.Image(FindInContentBrowserIcon.GetIcon())
+					]
+				]
+			]
+		];
+
+	Section.AddEntry(FToolMenuEntry::InitMenuEntry(
+		NAME_None,
+		FUIAction(OpenAction),
+		EntryWidget
+	));
+}
 
+void FMaterialEditor::GenerateInheritanceMenu(UToolMenu* Menu)
+{
 	RebuildInheritanceList();
-	TSharedPtr<FExtender> ToolbarExtender = MakeShareable(new FExtender);
-	const bool bShouldCloseWindowAfterMenuSelection = true;
-	TSharedPtr<FUICommandList> InCommandList = MakeShareable(new FUICommandList);
-	FMenuBuilder MenuBuilder(bShouldCloseWindowAfterMenuSelection, InCommandList, ToolbarExtender);
+	Menu->bShouldCloseWindowAfterMenuSelection = true;
+	Menu->bSearchable = true;
+	Menu->SetMaxHeight(500);
 
 	if (!MaterialFunction)
 	{
-		const FName MaterialInstances = TEXT("MaterialInstances");
-		MenuBuilder.BeginSection(MaterialInstances, LOCTEXT("MaterialInstances", "Material Instances"));
+		FToolMenuSection& Section = Menu->AddSection("MaterialInstances", LOCTEXT("MaterialInstances", "Material Instances"));
 		if (MaterialChildList.Num() == 0)
 		{
 			const FText NoChildText = LOCTEXT("NoInstancesFound", "No Instances Found");
 			TSharedRef<SWidget> NoChildWidget = SNew(STextBlock)
 				.Text(NoChildText);
-			MenuBuilder.AddWidget(NoChildWidget, FText::GetEmpty());
+			Section.AddEntry(FToolMenuEntry::InitWidget("NoInstancesFound", NoChildWidget, FText::GetEmpty()));
 		}
-		for (FAssetData MaterialChild : MaterialChildList)
+		for (const FAssetData& MaterialChild : MaterialChildList)
 		{
-			FFormatNamedArguments Args;
-			Args.Add(TEXT("ChildName"), FText::FromName(MaterialChild.AssetName));
-			MenuBuilder.AddSubMenu(FText::Format(LOCTEXT("MaterialChildName", "{ChildName}"), Args),
-				FText::GetEmpty(),
-				FNewMenuDelegate::CreateStatic(&Local::MakeMaterialSubmenu, MaterialChild),
-				false,
-				FSlateIcon());
+			FMaterialEditor::AddInheritanceMenuEntry(Section, MaterialChild, false);
 		}
-		MenuBuilder.EndSection();
 	}
-
-	return MenuBuilder.MakeWidget(nullptr, 500);
 }
 
 TSharedRef< SWidget > FMaterialEditor::GeneratePreviewMenuContent()
diff --git a/Engine/Source/Editor/MaterialEditor/Private/MaterialEditor.h b/Engine/Source/Editor/MaterialEditor/Private/MaterialEditor.h
--- a/Engine/Source/Editor/MaterialEditor/Private/MaterialEditor.h
+++ b/Engine/Source/Editor/MaterialEditor/Private/MaterialEditor.h
@@ -22,6 +22,7 @@
 #include "UObject/WeakFieldPtr.h"
 
 struct FAssetData;
+struct FToolMenuSection;
 class FCanvas;
 class FMaterialCompiler;
 class FScopedTransaction;
@@ -275,6 +276,8 @@ public:
 	virtual FText GetToolkitName() const override;
 	virtual FText GetToolkitToolTipText() const override;
 	virtual FString GetWorldCentricTabPrefix() const override;
+	virtual void InitToolMenuContext(struct FToolMenuContext& MenuContext) override;
+
 
 	/** @return the documentation location for this editor */
 	virtual FString GetDocumentationLink() const override
@@ -392,7 +395,7 @@ public:
 	virtual void GetBoundsForNode(const UEdGraphNode* InNode, class FSlateRect& OutRect, float InPadding) const override;
 	virtual FMatExpressionPreview* GetExpressionPreview(UMaterialExpression* InExpression) override;
 	virtual void DeleteNodes(const TArray<class UEdGraphNode*>& NodesToDelete) override;
-
+	virtual void GenerateInheritanceMenu(class UToolMenu* Menu) override;
 
 	void UpdateStatsMaterials();
 
@@ -414,6 +417,9 @@ public:
 	/** Rebuilds the inheritance list for this material. */
 	void RebuildInheritanceList();
 
+	/** Add entry to hierarchy menu */
+	static void AddInheritanceMenuEntry(FToolMenuSection& Section, const FAssetData& AssetData, bool bIsFunctionPreviewMaterial);
+
 public:
 	/** Set to true when modifications have been made to the material */
 	bool bMaterialDirty;
@@ -545,6 +551,7 @@ protected:
 private:
 	/** Builds the toolbar widget for the material editor */
 	void ExtendToolbar();
+	void RegisterToolBar();
 
 	/** Creates the toolbar buttons. Bound by ExtendToolbar*/
 	void FillToolbar(FToolBarBuilder& ToolbarBuilder);
diff --git a/Engine/Source/Editor/MaterialEditor/Private/MaterialEditorUtilities.cpp b/Engine/Source/Editor/MaterialEditor/Private/MaterialEditorUtilities.cpp
--- a/Engine/Source/Editor/MaterialEditor/Private/MaterialEditorUtilities.cpp
+++ b/Engine/Source/Editor/MaterialEditor/Private/MaterialEditorUtilities.cpp
@@ -718,26 +718,26 @@ void FMaterialEditorUtilities::BuildTextureStreamingData(UMaterialInterface* Upd
 }
 
 
-void FMaterialEditorUtilities::OnOpenMaterial(FAssetData InMaterial)
+void FMaterialEditorUtilities::OnOpenMaterial(const FAssetData InMaterial)
 {
 	UMaterialInterface* MaterialInterface = Cast<UMaterialInterface>(InMaterial.GetAsset());
 	OpenSelectedParentEditor(MaterialInterface);
 }
 
-void FMaterialEditorUtilities::OnOpenFunction(FAssetData InFunction)
+void FMaterialEditorUtilities::OnOpenFunction(const FAssetData InFunction)
 {
 	UMaterialFunctionInterface* MaterialFunctionInterface = Cast<UMaterialFunctionInterface>(InFunction.GetAsset());
 	OpenSelectedParentEditor(MaterialFunctionInterface);
 }
 
-void FMaterialEditorUtilities::OnShowMaterialInContentBrowser(FAssetData InMaterial)
+void FMaterialEditorUtilities::OnShowMaterialInContentBrowser(const FAssetData InMaterial)
 {
 	TArray<UObject*> SyncedObject;
 	SyncedObject.Add(InMaterial.GetAsset());
 	GEditor->SyncBrowserToObjects(SyncedObject);
 }
 
-void FMaterialEditorUtilities::OnShowFunctionInContentBrowser(FAssetData InFunction)
+void FMaterialEditorUtilities::OnShowFunctionInContentBrowser(const FAssetData InFunction)
 {
 	TArray<UObject*> SyncedObject;
 	SyncedObject.Add(InFunction.GetAsset());
diff --git a/Engine/Source/Editor/MaterialEditor/Private/MaterialInstanceEditor.cpp b/Engine/Source/Editor/MaterialEditor/Private/MaterialInstanceEditor.cpp
--- a/Engine/Source/Editor/MaterialEditor/Private/MaterialInstanceEditor.cpp
+++ b/Engine/Source/Editor/MaterialEditor/Private/MaterialInstanceEditor.cpp
@@ -819,25 +819,17 @@ void FMaterialInstanceEditor::RegisterToolBar()
 
 		{
 			FToolMenuSection& Section = ToolBar->AddSection("Parent", TAttribute<FText>(), InsertAfterAssetSection);
-
-			struct Local
-			{
-				static void GenerateInheritanceMenu(UToolMenu* InMenu)
-				{
-					UMaterialEditorMenuContext* Context = InMenu->FindContext<UMaterialEditorMenuContext>();
-					if (Context && Context->MaterialEditor.IsValid())
-					{
-						TSharedPtr<IMaterialEditor> Pinned = Context->MaterialEditor.Pin();
-						TSharedPtr<FMaterialInstanceEditor> Recast = StaticCastSharedPtr<FMaterialInstanceEditor>(Pinned);
-						Recast->GenerateInheritanceMenu(InMenu);
-					}
-				}
-			};
-
 			Section.AddEntry(FToolMenuEntry::InitComboButton(
 				"Hierarchy",
 				FToolUIActionChoice(),
-				FNewToolMenuDelegate::CreateStatic(&Local::GenerateInheritanceMenu),
+				FNewToolMenuDelegate::CreateLambda([](UToolMenu* InSubMenu)
+				{
+					UMaterialEditorMenuContext* SubMenuContext = InSubMenu->FindContext<UMaterialEditorMenuContext>();
+					if (SubMenuContext && SubMenuContext->MaterialEditor.IsValid())
+					{
+						SubMenuContext->MaterialEditor.Pin()->GenerateInheritanceMenu(InSubMenu);
+					}
+				}),
 				LOCTEXT("Hierarchy", "Hierarchy"),
 				FText::GetEmpty(),
 				FSlateIcon(FEditorStyle::GetStyleSetName(), "BTEditor.SwitchToBehaviorTreeMode"),
@@ -866,140 +858,49 @@ void FMaterialInstanceEditor::ExtendToolbar()
 
 void FMaterialInstanceEditor::GenerateInheritanceMenu(UToolMenu* Menu)
 {
-	struct Local
-	{
-		static void AddMenuEntry(FToolMenuSection& Section, FAssetData AssetData, bool bIsFunctionPreviewMaterial)
-		{
-			FExecuteAction OpenAction;
-			FExecuteAction FindInContentBrowserAction;
-			if (bIsFunctionPreviewMaterial)
-			{
-				OpenAction.BindStatic(&FMaterialEditorUtilities::OnOpenFunction, AssetData);
-				FindInContentBrowserAction.BindStatic(&FMaterialEditorUtilities::OnShowFunctionInContentBrowser, AssetData);
-			}
-			else
-			{
-				OpenAction.BindStatic(&FMaterialEditorUtilities::OnOpenMaterial, AssetData);
-				FindInContentBrowserAction.BindStatic(&FMaterialEditorUtilities::OnShowMaterialInContentBrowser, AssetData);
-			}
-
-			FFormatNamedArguments Args;
-			Args.Add(TEXT("ParentName"), FText::FromName(AssetData.AssetName));
-			FText Label = FText::Format(LOCTEXT("InstanceParentName", "{ParentName}"), Args);
-
-			FSlateIcon OpenIcon(FEditorStyle::GetStyleSetName(), "ContentBrowser.AssetActions.OpenInExternalEditor");
-			FSlateIcon FindInContentBrowserIcon(FEditorStyle::GetStyleSetName(), "SystemWideCommands.FindInContentBrowser");
-
-			TSharedRef<SWidget> EntryWidget =
-				SNew(SHorizontalBox)
-				.ToolTipText(LOCTEXT("OpenInEditor", "Open In Editor"))
-
-				+ SHorizontalBox::Slot()
-				.AutoWidth()
-				.Padding(FMargin(2, 0, 2, 0))
-				[
-					SNew( SBox )
-					.WidthOverride( MultiBoxConstants::MenuIconSize + 2 )
-					.HeightOverride( MultiBoxConstants::MenuIconSize )
-					.HAlign(HAlign_Center)
-					.VAlign(VAlign_Center)
-					[
-						SNew( SBox )
-						.WidthOverride( MultiBoxConstants::MenuIconSize )
-						.HeightOverride( MultiBoxConstants::MenuIconSize )
-						[
-							SNew(SImage)
-							.Image(OpenIcon.GetIcon())
-						]
-					]
-				]
-
-				+ SHorizontalBox::Slot()
-				.FillWidth( 1.0f )
-				.Padding(FMargin(2, 0, 6, 0))
-				.VAlign( VAlign_Center )
-				[
-					SNew(STextBlock)
-					.Text(Label)
-				]
-
-				+ SHorizontalBox::Slot()
-				.AutoWidth()
-				.VAlign( VAlign_Center )
-				.HAlign( HAlign_Right )
-				[
-					SNew(SButton)
-					.ButtonStyle(FEditorStyle::Get(), "HoverHintOnly")
-					.ToolTipText(LOCTEXT("FindInContentBrowser", "Find In Content Browser"))
-					.OnClicked_Lambda([FindInContentBrowserAction]() { FindInContentBrowserAction.ExecuteIfBound(); return FReply::Handled(); })
-					[
-						SNew( SBox )
-						.WidthOverride( MultiBoxConstants::MenuIconSize + 2 )
-						.HeightOverride( MultiBoxConstants::MenuIconSize )
-						.HAlign(HAlign_Center)
-						.VAlign(VAlign_Center)
-						[
-							SNew( SBox )
-							.WidthOverride( MultiBoxConstants::MenuIconSize )
-							.HeightOverride( MultiBoxConstants::MenuIconSize )
-							[
-								SNew(SImage)
-								.Image(FindInContentBrowserIcon.GetIcon())
-							]
-						]
-					]
-				];
-
-			Section.AddEntry(FToolMenuEntry::InitMenuEntry(
-				NAME_None,
-				FUIAction(OpenAction),
-				EntryWidget
-			));
-		}
-	};
-
 	RebuildInheritanceList();
 	Menu->bShouldCloseWindowAfterMenuSelection = true;
 	Menu->bSearchable = true;
 	Menu->SetMaxHeight(500);
-	const FName ParentName = TEXT("ParentChain");
-	FToolMenuSection& Section = Menu->AddSection(ParentName, LOCTEXT("ParentChain", "Parent Chain"));
-	if (bIsFunctionPreviewMaterial)
-	{
-		if (FunctionParentList.Num() == 0)
-		{
-			const FText NoParentText = LOCTEXT("NoParentFound", "No Parent Found");
-			TSharedRef<SWidget> NoParentWidget = SNew(STextBlock)
-				.Text(NoParentText);
-			Section.AddEntry(FToolMenuEntry::InitWidget("NoParentEntry", NoParentWidget, FText::GetEmpty()));
-		}
-		for (FAssetData FunctionParent : FunctionParentList)
-		{
-			Local::AddMenuEntry(Section, FunctionParent, bIsFunctionPreviewMaterial);
-		}
-	}
-	else
+
 	{
-		if (MaterialParentList.Num() == 0)
+		FToolMenuSection& Section = Menu->AddSection("ParentChain", LOCTEXT("ParentChain", "Parent Chain"));
+		if (bIsFunctionPreviewMaterial)
 		{
-			const FText NoParentText = LOCTEXT("NoParentFound", "No Parent Found");
-			TSharedRef<SWidget> NoParentWidget = SNew(STextBlock)
-				.Text(NoParentText);
-			Section.AddEntry(FToolMenuEntry::InitWidget("NoParentEntry", NoParentWidget, FText::GetEmpty()));
+			if (FunctionParentList.Num() == 0)
+			{
+				const FText NoParentText = LOCTEXT("NoParentFound", "No Parent Found");
+				TSharedRef<SWidget> NoParentWidget = SNew(STextBlock)
+					.Text(NoParentText);
+				Section.AddEntry(FToolMenuEntry::InitWidget("NoParentEntry", NoParentWidget, FText::GetEmpty()));
+			}
+			for (FAssetData FunctionParent : FunctionParentList)
+			{
+				FMaterialEditor::AddInheritanceMenuEntry(Section, FunctionParent, bIsFunctionPreviewMaterial);
+			}
 		}
-		for (FAssetData MaterialParent : MaterialParentList)
+		else
 		{
-			Local::AddMenuEntry(Section, MaterialParent, bIsFunctionPreviewMaterial);
+			if (MaterialParentList.Num() == 0)
+			{
+				const FText NoParentText = LOCTEXT("NoParentFound", "No Parent Found");
+				TSharedRef<SWidget> NoParentWidget = SNew(STextBlock)
+					.Text(NoParentText);
+				Section.AddEntry(FToolMenuEntry::InitWidget("NoParentEntry", NoParentWidget, FText::GetEmpty()));
+			}
+			for (FAssetData MaterialParent : MaterialParentList)
+			{
+				FMaterialEditor::AddInheritanceMenuEntry(Section, MaterialParent, bIsFunctionPreviewMaterial);
+			}
 		}
 	}
 
 	if (!bIsFunctionPreviewMaterial)
 	{
-		const FName MaterialInstances = TEXT("MaterialInstances");
-		FToolMenuSection& MaterialInstancesSection = Menu->AddSection(MaterialInstances, LOCTEXT("MaterialInstances", "Material Instances"));
-		for (FAssetData MaterialChild : MaterialChildList)
+		FToolMenuSection& Section = Menu->AddSection("MaterialInstances", LOCTEXT("MaterialInstances", "Material Instances"));
+		for (const FAssetData& MaterialChild : MaterialChildList)
 		{
-			Local::AddMenuEntry(MaterialInstancesSection, MaterialChild, bIsFunctionPreviewMaterial);
+			FMaterialEditor::AddInheritanceMenuEntry(Section, MaterialChild, bIsFunctionPreviewMaterial);
 		}
 	}
 }
diff --git a/Engine/Source/Editor/MaterialEditor/Private/MaterialInstanceEditor.h b/Engine/Source/Editor/MaterialEditor/Private/MaterialInstanceEditor.h
--- a/Engine/Source/Editor/MaterialEditor/Private/MaterialInstanceEditor.h
+++ b/Engine/Source/Editor/MaterialEditor/Private/MaterialInstanceEditor.h
@@ -111,6 +111,9 @@ public:
 	/** call this to notify the editor that the edited material changed from outside */
 	virtual void NotifyExternalMaterialChange() override;
 
+	// IMaterial Editor Interface
+	virtual void GenerateInheritanceMenu(class UToolMenu* Menu) override;
+
 protected:
 	//~ FAssetEditorToolkit interface
 	virtual void SaveAsset_Execute() override;
@@ -159,8 +162,6 @@ private:
 	/** Builds the toolbar widget for the material editor */
 	void ExtendToolbar();
 
-	void GenerateInheritanceMenu(UToolMenu* Menu);
-
 	/** If re-initializing for a material function instance re-generate the proxy materials */
 	void ReInitMaterialFunctionProxies();
 
diff --git a/Engine/Source/Editor/MaterialEditor/Public/IMaterialEditor.h b/Engine/Source/Editor/MaterialEditor/Public/IMaterialEditor.h
--- a/Engine/Source/Editor/MaterialEditor/Public/IMaterialEditor.h
+++ b/Engine/Source/Editor/MaterialEditor/Public/IMaterialEditor.h
@@ -126,6 +126,9 @@ public:
 	/** call this to notify the editor that the edited material changed from outside */
 	virtual void NotifyExternalMaterialChange() {}
 
+	/** Fills inheritance drop down menu */
+	virtual void GenerateInheritanceMenu(class UToolMenu* Menu) {};
+
 private:
 	FMaterialEditorClosedEvent MaterialEditorClosedEvent;
 	FRegisterTabSpawnersEvent RegisterTabSpawnersEvent;
diff --git a/Engine/Source/Editor/MaterialEditor/Public/MaterialEditorUtilities.h b/Engine/Source/Editor/MaterialEditor/Public/MaterialEditorUtilities.h
--- a/Engine/Source/Editor/MaterialEditor/Public/MaterialEditorUtilities.h
+++ b/Engine/Source/Editor/MaterialEditor/Public/MaterialEditorUtilities.h
@@ -196,10 +196,10 @@ public:
 
 
 	/** Commands for the Parents menu */
-	static void OnOpenMaterial(FAssetData InMaterial);
-	static void OnOpenFunction(FAssetData InFunction);
-	static void OnShowMaterialInContentBrowser(FAssetData InMaterial);
-	static void OnShowFunctionInContentBrowser(FAssetData InFunction);
+	static void OnOpenMaterial(const FAssetData InMaterial);
+	static void OnOpenFunction(const FAssetData InFunction);
+	static void OnShowMaterialInContentBrowser(const FAssetData InMaterial);
+	static void OnShowFunctionInContentBrowser(const FAssetData InFunction);
 
 private:
 
