Getting material layer parameter visibility to be impacted by static switches

diff --git a/Engine/Source/Editor/MaterialEditor/Private/MaterialInstanceEditor.cpp b/Engine/Source/Editor/MaterialEditor/Private/MaterialInstanceEditor.cpp
--- a/Engine/Source/Editor/MaterialEditor/Private/MaterialInstanceEditor.cpp
+++ b/Engine/Source/Editor/MaterialEditor/Private/MaterialInstanceEditor.cpp
@@ -772,7 +772,8 @@ void FMaterialInstanceEditor::CreateInternalWidgets()
 	if (MaterialEditorModule->MaterialLayersEnabled() && !bIsFunctionPreviewMaterial)
 	{
 		MaterialLayersFunctionsInstance = SNew(SMaterialLayersFunctionsInstanceWrapper)
-			.InMaterialEditorInstance(MaterialEditorInstance);
+			.InMaterialEditorInstance(MaterialEditorInstance)
+			.InShowHiddenDelegate(FGetShowHiddenParameters::CreateSP(this, &FMaterialInstanceEditor::GetShowHiddenParameters));
 	}
 }
 
diff --git a/Engine/Source/Editor/MaterialEditor/Private/SMaterialLayersFunctionsTree.cpp b/Engine/Source/Editor/MaterialEditor/Private/SMaterialLayersFunctionsTree.cpp
--- a/Engine/Source/Editor/MaterialEditor/Private/SMaterialLayersFunctionsTree.cpp
+++ b/Engine/Source/Editor/MaterialEditor/Private/SMaterialLayersFunctionsTree.cpp
@@ -209,6 +209,11 @@ bool SMaterialLayersFunctionsInstanceTree::IsOverriddenExpression(class UDEditor
 	return FMaterialPropertyHelpers::IsOverriddenExpression(Parameter) && FunctionInstance->LayerStates[InIndex];
 }
 
+FGetShowHiddenParameters SMaterialLayersFunctionsInstanceTree::GetShowHiddenDelegate() const
+{
+	return ShowHiddenDelegate;
+}
+
 void  SMaterialLayersFunctionsInstanceTreeItem::OnOverrideParameter(bool NewValue, class UDEditorParameterValue* Parameter)
 {
 	FMaterialPropertyHelpers::OnOverrideParameter(NewValue, Parameter, MaterialEditorInstance);
@@ -560,6 +565,8 @@ void SMaterialLayersFunctionsInstanceTreeItem::Construct(const FArguments& InArg
 			.OverrideResetToDefault(ResetOverride)
 			.EditCondition(IsParamEnabled, FOnBooleanValueChanged::CreateSP(this, &SMaterialLayersFunctionsInstanceTreeItem::OnOverrideParameter, StackParameterData->Parameter));
 
+		WrapperWidget->SetVisibility(TAttribute<EVisibility>::Create(TAttribute<EVisibility>::FGetter::CreateStatic(&FMaterialPropertyHelpers::ShouldShowExpression, StackParameterData->Parameter, MaterialEditorInstance, Tree->GetShowHiddenDelegate())));
+
 		if (VectorParam && VectorParam->bIsUsedAsChannelMask)
 		{
 			FOnGetPropertyComboBoxStrings GetMaskStrings = FOnGetPropertyComboBoxStrings::CreateStatic(&FMaterialPropertyHelpers::GetVectorChannelMaskComboBoxStrings);
@@ -1055,6 +1062,7 @@ void SMaterialLayersFunctionsInstanceTree::Construct(const FArguments& InArgs)
 	ColumnWidth = 0.5f;
 	MaterialEditorInstance = InArgs._InMaterialEditorInstance;
 	Wrapper = InArgs._InWrapper;
+	ShowHiddenDelegate = InArgs._InShowHiddenDelegate;
 	CreateGroupsWidget();
 
 #ifdef WITH_EDITOR
@@ -1732,7 +1740,8 @@ void SMaterialLayersFunctionsInstanceWrapper::Construct(const FArguments& InArgs
 {
 	NestedTree = SNew(SMaterialLayersFunctionsInstanceTree)
 		.InMaterialEditorInstance(InArgs._InMaterialEditorInstance)
-		.InWrapper(this);
+		.InWrapper(this)
+		.InShowHiddenDelegate(InArgs._InShowHiddenDelegate);
 
 	LayerParameter = NestedTree->FunctionParameter;
 
@@ -2351,7 +2360,8 @@ void SMaterialLayersFunctionsMaterialWrapper::Refresh()
 void SMaterialLayersFunctionsMaterialWrapper::Construct(const FArguments& InArgs)
 {
 	NestedTree = SNew(SMaterialLayersFunctionsMaterialTree)
-		.InMaterialEditorInstance(InArgs._InMaterialEditorInstance);
+		.InMaterialEditorInstance(InArgs._InMaterialEditorInstance)
+		.InShowHiddenDelegate(InArgs._InShowHiddenDelegate);
 
 	LayerParameter = NestedTree->FunctionParameter;
 	MaterialEditorInstance = InArgs._InMaterialEditorInstance;
@@ -2371,7 +2381,7 @@ void SMaterialLayersFunctionsMaterialTree::Construct(const FArguments& InArgs)
 {
 	ColumnWidth = 0.5f;
 	MaterialEditorInstance = InArgs._InMaterialEditorInstance;
-
+	ShowHiddenDelegate = InArgs._InShowHiddenDelegate;
 	CreateGroupsWidget();
 
 	STreeView<TSharedPtr<FSortedParamData>>::Construct(
diff --git a/Engine/Source/Editor/MaterialEditor/Private/SMaterialLayersFunctionsTree.h b/Engine/Source/Editor/MaterialEditor/Private/SMaterialLayersFunctionsTree.h
--- a/Engine/Source/Editor/MaterialEditor/Private/SMaterialLayersFunctionsTree.h
+++ b/Engine/Source/Editor/MaterialEditor/Private/SMaterialLayersFunctionsTree.h
@@ -95,10 +95,12 @@ class SMaterialLayersFunctionsInstanceWrapper : public SCompoundWidget
 {
 public:
 	SLATE_BEGIN_ARGS(SMaterialLayersFunctionsInstanceWrapper)
-		: _InMaterialEditorInstance(nullptr)
+		: _InMaterialEditorInstance(nullptr),
+		_InShowHiddenDelegate()
 	{}
 
 	SLATE_ARGUMENT(UMaterialEditorInstanceConstant*, InMaterialEditorInstance)
+	SLATE_ARGUMENT(FGetShowHiddenParameters, InShowHiddenDelegate)
 
 	SLATE_END_ARGS()
 	void Refresh();
@@ -122,6 +124,7 @@ public:
 
 	SLATE_ARGUMENT(UMaterialEditorInstanceConstant*, InMaterialEditorInstance)
 	SLATE_ARGUMENT(SMaterialLayersFunctionsInstanceWrapper*, InWrapper)
+	SLATE_ARGUMENT(FGetShowHiddenParameters, InShowHiddenDelegate)
 	SLATE_END_ARGS()
 
 	/** Constructs this widget with InArgs */
@@ -162,6 +165,8 @@ public:
 	void UpdateThumbnailMaterial(TEnumAsByte<EMaterialParameterAssociation> InAssociation, int32 InIndex, bool bAlterBlendIndex = false);
 	FReply OnThumbnailDoubleClick(const FGeometry& Geometry, const FPointerEvent& MouseEvent, EMaterialParameterAssociation InAssociation, int32 InIndex);
 	bool IsOverriddenExpression(class UDEditorParameterValue* Parameter, int32 InIndex);
+
+	FGetShowHiddenParameters GetShowHiddenDelegate() const;
 protected:
 
 	void ShowSubParameters(TSharedPtr<FSortedParamData> ParentParameter);
@@ -180,6 +185,9 @@ private:
 
 	bool bLayerIsolated;
 
+	/** Delegate to call to determine if hidden parameters should be shown */
+	FGetShowHiddenParameters ShowHiddenDelegate;
+
 };
 
 class UMaterialEditorPreviewParameters;
@@ -188,10 +196,12 @@ class SMaterialLayersFunctionsMaterialWrapper : public SCompoundWidget
 {
 public:
 	SLATE_BEGIN_ARGS(SMaterialLayersFunctionsMaterialWrapper)
-		: _InMaterialEditorInstance(nullptr)
+		: _InMaterialEditorInstance(nullptr),
+		_InShowHiddenDelegate()
 	{}
 
 	SLATE_ARGUMENT(UMaterialEditorPreviewParameters*, InMaterialEditorInstance)
+	SLATE_ARGUMENT(FGetShowHiddenParameters, InShowHiddenDelegate)
 
 	SLATE_END_ARGS()
 	void Refresh();
@@ -209,10 +219,12 @@ class SMaterialLayersFunctionsMaterialTree : public STreeView<TSharedPtr<FSorted
 	friend class SMaterialLayersFunctionsMaterialTreeItem;
 public:
 	SLATE_BEGIN_ARGS(SMaterialLayersFunctionsMaterialTree)
-		: _InMaterialEditorInstance(nullptr)
+		: _InMaterialEditorInstance(nullptr),
+		_InShowHiddenDelegate()
 	{}
 
 	SLATE_ARGUMENT(UMaterialEditorPreviewParameters*, InMaterialEditorInstance)
+	SLATE_ARGUMENT(FGetShowHiddenParameters, InShowHiddenDelegate)
 
 	SLATE_END_ARGS()
 
@@ -235,6 +247,9 @@ public:
 	/** Object that stores all of the possible parameters we can edit */
 	UMaterialEditorPreviewParameters* MaterialEditorInstance;
 
+	/** Delegate to call to determine if hidden parameters should be shown */
+	FGetShowHiddenParameters ShowHiddenDelegate;
+
 	/** Builds the custom parameter groups category */
 	void CreateGroupsWidget();
 
