When deciding what to recompile after chaning a MF, look at all MIs, instead of stopping after finding the first match

diff --git a/Engine/Source/Editor/MaterialEditor/Private/MaterialEditingLibrary.cpp b/Engine/Source/Editor/MaterialEditor/Private/MaterialEditingLibrary.cpp
--- a/Engine/Source/Editor/MaterialEditor/Private/MaterialEditingLibrary.cpp
+++ b/Engine/Source/Editor/MaterialEditor/Private/MaterialEditingLibrary.cpp
@@ -866,8 +866,6 @@ void UMaterialEditingLibrary::UpdateMaterialFunction(UMaterialFunctionInterface*
 						{
 							MaterialEditor->NotifyExternalMaterialChange();
 						}
-
-						break;
 					}
 				}
 			}
