Don't allow channel names for the channel mask parameter

diff --git a/Engine/Source/Editor/MaterialEditor/Private/MaterialEditorDetailCustomization.cpp b/Engine/Source/Editor/MaterialEditor/Private/MaterialEditorDetailCustomization.cpp
--- a/Engine/Source/Editor/MaterialEditor/Private/MaterialEditorDetailCustomization.cpp
+++ b/Engine/Source/Editor/MaterialEditor/Private/MaterialEditorDetailCustomization.cpp
@@ -215,7 +215,7 @@ void FMaterialExpressionParameterDetails::CustomizeDetails( IDetailLayoutBuilder
 				DetailLayout.HideCategory(MaterialExpressionCategory);
 			}
 
-			if (ValueHandle.IsValid() && ValueHandle->IsValidHandle())
+			if (ValueHandle.IsValid() && ValueHandle->IsValidHandle() && !VectorParameter->IsUsedAsChannelMask())
 			{
 				static const FName Red("R");
 				static const FName Green("G");
@@ -238,6 +238,12 @@ void FMaterialExpressionParameterDetails::CustomizeDetails( IDetailLayoutBuilder
 					ValueHandle->GetChildHandle(Alpha)->SetPropertyDisplayName(VectorParameter->ChannelNames.A);
 				}
 			}
+
+			if (VectorParameter->IsUsedAsChannelMask())
+			{
+				TSharedPtr<IPropertyHandle> ChannelNameHandle = DetailLayout.GetProperty("ChannelNames", UMaterialExpressionVectorParameter::StaticClass());
+				ChannelNameHandle->MarkHiddenByCustomization();
+			}
 		}
 
 		UMaterialExpressionTextureSampleParameter* TextureParameter = Cast<UMaterialExpressionTextureSampleParameter>(Object);
diff --git a/Engine/Source/Runtime/Engine/Private/Materials/MaterialExpressions.cpp b/Engine/Source/Runtime/Engine/Private/Materials/MaterialExpressions.cpp
--- a/Engine/Source/Runtime/Engine/Private/Materials/MaterialExpressions.cpp
+++ b/Engine/Source/Runtime/Engine/Private/Materials/MaterialExpressions.cpp
@@ -7288,7 +7288,8 @@ void UMaterialExpressionVectorParameter::PostEditChangeProperty(FPropertyChanged
 		const int32 PrimDataIndex = PrimitiveDataIndex;
 		PrimitiveDataIndex = (uint8)FMath::Clamp(PrimDataIndex, 0, FCustomPrimitiveData::NumCustomPrimitiveDataFloats-1);
 	}
-	else if (PropertyName == GET_MEMBER_NAME_STRING_CHECKED(UMaterialExpressionVectorParameter, ChannelNames))
+	else if (PropertyName == GET_MEMBER_NAME_STRING_CHECKED(UMaterialExpressionVectorParameter, ChannelNames)
+		&& !IsUsedAsChannelMask())
 	{
 		ApplyChannelNames();
 
