Material layer parameter preview panel will no longer crash in the material editor.

diff --git a/Engine/Source/Editor/MaterialEditor/Private/MaterialEditor.cpp b/Engine/Source/Editor/MaterialEditor/Private/MaterialEditor.cpp
--- a/Engine/Source/Editor/MaterialEditor/Private/MaterialEditor.cpp
+++ b/Engine/Source/Editor/MaterialEditor/Private/MaterialEditor.cpp
@@ -727,9 +727,18 @@ void FMaterialEditor::InitMaterialEditor( const EToolkitMode::Type Mode, const T
 	RegenerateCodeView(true);
 
 	ForceRefreshExpressionPreviews();
-	
+	if (Generator.IsValid())
+	{
+		TArray<UObject*> Objects;
+		Objects.Add(MaterialEditorInstance);
+		Generator->SetObjects(Objects);
+	}
 	// Update the parameter list now that the material has been fully initialized
 	MaterialParametersOverviewWidget->UpdateEditorInstance(MaterialEditorInstance);
+	if (MaterialLayersFunctionsInstance.IsValid())
+	{
+		MaterialLayersFunctionsInstance->SetEditorInstance(MaterialEditorInstance);
+	}
 
 	if (OriginalMaterial->bUsedAsSpecialEngineMaterial)
 	{
@@ -977,16 +986,26 @@ void FMaterialEditor::CreateInternalWidgets()
 	{
 		MaterialEditorInstance->OriginalFunction = MaterialFunction->ParentFunction;
 	}
+	FPropertyEditorModule& Module = FModuleManager::LoadModuleChecked<FPropertyEditorModule>("PropertyEditor");
+	if (!Generator.IsValid())
+	{
+		FPropertyRowGeneratorArgs Args;
+		Generator = Module.CreatePropertyRowGenerator(Args);
+	}
+
 	MaterialParametersOverviewWidget = SNew(SMaterialParametersOverviewPanel)
-		.InMaterialEditorInstance(MaterialEditorInstance);
-	MaterialParametersOverviewWidget->GetGenerator()->OnFinishedChangingProperties().AddSP(this, &FMaterialEditor::OnFinishedChangingParametersFromOverview);
+		.InMaterialEditorInstance(MaterialEditorInstance)
+		.InGenerator(Generator);
+	Generator->OnFinishedChangingProperties().AddSP(this, &FMaterialEditor::OnFinishedChangingParametersFromOverview);
+	Generator->OnRowsRefreshed().AddSP(this, &FMaterialEditor::GeneratorRowsRefreshed);
 	MaterialCustomPrimitiveDataWidget = SNew(SMaterialCustomPrimitiveDataPanel, MaterialEditorInstance);
 
 	IMaterialEditorModule* MaterialEditorModule = &FModuleManager::LoadModuleChecked<IMaterialEditorModule>("MaterialEditor");
 	if (MaterialEditorModule->MaterialLayersEnabled())
 	{
 		MaterialLayersFunctionsInstance = SNew(SMaterialLayersFunctionsMaterialWrapper)
-			.InMaterialEditorInstance(MaterialEditorInstance);
+			.InMaterialEditorInstance(MaterialEditorInstance)
+			.InGenerator(Generator);
 	}
 
 	Palette = SNew(SMaterialPalette, SharedThis(this));
@@ -1053,6 +1072,15 @@ void FMaterialEditor::OnFinishedChangingParametersFromOverview(const FPropertyCh
 	}
 }
 
+void FMaterialEditor::GeneratorRowsRefreshed()
+{
+	MaterialParametersOverviewWidget->Refresh();
+	if (MaterialLayersFunctionsInstance)
+	{
+		MaterialLayersFunctionsInstance->Refresh();
+	}
+}
+
 FName FMaterialEditor::GetToolkitFName() const
 {
 	return FName("MaterialEditor");
@@ -1834,7 +1862,6 @@ void FMaterialEditor::UpdatePreviewMaterial( bool bForce )
 		if (MaterialLayersFunctionsInstance.IsValid())
 		{
 			MaterialLayersFunctionsInstance->SetEditorInstance(MaterialEditorInstance);
-			MaterialLayersFunctionsInstance->Refresh();
 		}
 	}
 
diff --git a/Engine/Source/Editor/MaterialEditor/Private/MaterialEditor.h b/Engine/Source/Editor/MaterialEditor/Private/MaterialEditor.h
--- a/Engine/Source/Editor/MaterialEditor/Private/MaterialEditor.h
+++ b/Engine/Source/Editor/MaterialEditor/Private/MaterialEditor.h
@@ -779,6 +779,7 @@ private:
 
 	void OnFinishedChangingProperties(const FPropertyChangedEvent& PropertyChangedEvent);
 	void OnFinishedChangingParametersFromOverview(const FPropertyChangedEvent& PropertyChangedEvent);
+	void GeneratorRowsRefreshed();
 private:
 	/** List of open tool panels; used to ensure only one exists at any one time */
 	TMap< FName, TWeakPtr<class SDockableTab> > SpawnedToolPanels;
@@ -895,4 +896,6 @@ private:
 
 	/** List of children used to populate the inheritance list chain. */
 	TArray< FAssetData > FunctionChildList;
+
+	TSharedPtr<class IPropertyRowGenerator> Generator;
 };
diff --git a/Engine/Source/Editor/MaterialEditor/Private/SMaterialLayersFunctionsTree.cpp b/Engine/Source/Editor/MaterialEditor/Private/SMaterialLayersFunctionsTree.cpp
--- a/Engine/Source/Editor/MaterialEditor/Private/SMaterialLayersFunctionsTree.cpp
+++ b/Engine/Source/Editor/MaterialEditor/Private/SMaterialLayersFunctionsTree.cpp
@@ -1756,187 +1756,273 @@ void SMaterialLayersFunctionsInstanceWrapper::SetEditorInstance(UMaterialEditorI
 	Refresh();
 }
 
-/////////////// MATERIAL VERSION
+
+#undef LOCTEXT_NAMESPACE
+
 
 
-class SMaterialLayersFunctionsMaterialTreeItem : public STableRow< TSharedPtr<FSortedParamData> >
-{
-public:
 
-	SLATE_BEGIN_ARGS(SMaterialLayersFunctionsMaterialTreeItem)
-		: _StackParameterData(nullptr),
-		_MaterialEditorInstance(nullptr)
-	{}
 
-	/** The item content. */
-	SLATE_ARGUMENT(TSharedPtr<FSortedParamData>, StackParameterData)
-	SLATE_ARGUMENT(UMaterialEditorPreviewParameters*, MaterialEditorInstance)
-	SLATE_ARGUMENT(SMaterialLayersFunctionsMaterialTree*, InTree)
-	SLATE_END_ARGS()
 
-	FMaterialTreeColumnSizeData ColumnSizeData;
 
-	/**
-	* Construct the widget
-	*
-	* @param InArgs   A declaration from which to construct the widget
-	*/
-	FString GetCurvePath(UDEditorScalarParameterValue* Parameter) const;
 
 
-	void Construct(const FArguments& InArgs, const TSharedRef<STableViewBase>& InOwnerTableView)
+/////////////// MATERIAL VERSION
+#define LOCTEXT_NAMESPACE "MaterialLayerDisplay"
+
+FString SMaterialLayersFunctionsMaterialTreeItem::GetCurvePath(UDEditorScalarParameterValue* Parameter) const
+{
+	FString Path = Parameter->AtlasData.Curve->GetPathName();
+	return Path;
+}
+
+const FSlateBrush* SMaterialLayersFunctionsMaterialTreeItem::GetBorderImage() const
+{
+	if (StackParameterData->StackDataType == EStackDataType::Stack)
+	{
+		return FEditorStyle::GetBrush("MaterialInstanceEditor.StackHeader");
+	}
+	else
+	{
+		return FEditorStyle::GetBrush("MaterialInstanceEditor.StackBody");
+	}
+}
+
+void SMaterialLayersFunctionsMaterialTreeItem::RefreshOnRowChange(const FAssetData& AssetData, SMaterialLayersFunctionsInstanceTree* InTree)
+{
+	if (SMaterialLayersFunctionsInstanceWrapper* Wrapper = InTree->GetWrapper())
 	{
-		// todo: add tooltips
-		StackParameterData = InArgs._StackParameterData;
-		MaterialEditorInstance = InArgs._MaterialEditorInstance;
-		SMaterialLayersFunctionsMaterialTree* Tree = InArgs._InTree;
-		ColumnSizeData.LeftColumnWidth = TAttribute<float>(Tree, &SMaterialLayersFunctionsMaterialTree::OnGetLeftColumnWidth);
-		ColumnSizeData.RightColumnWidth = TAttribute<float>(Tree, &SMaterialLayersFunctionsMaterialTree::OnGetRightColumnWidth);
-		ColumnSizeData.OnWidthChanged = SSplitter::FOnSlotResized::CreateSP(Tree, &SMaterialLayersFunctionsMaterialTree::OnSetColumnWidth);
+		InTree->CreateGroupsWidget();
+	}
+}
+
+FText SMaterialLayersFunctionsMaterialTreeItem::GetLayerName(SMaterialLayersFunctionsMaterialTree* InTree, int32 Counter) const
+{
+	return InTree->FunctionInstance->GetLayerName(Counter);
+}
+
+void SMaterialLayersFunctionsMaterialTreeItem::Construct(const FArguments& InArgs, const TSharedRef<STableViewBase>& InOwnerTableView)
+{
+	StackParameterData = InArgs._StackParameterData;
+	MaterialEditorInstance = InArgs._MaterialEditorInstance;
+	Tree = InArgs._InTree;
+	ColumnSizeData.LeftColumnWidth = TAttribute<float>(Tree, &SMaterialLayersFunctionsMaterialTree::OnGetLeftColumnWidth);
+	ColumnSizeData.RightColumnWidth = TAttribute<float>(Tree, &SMaterialLayersFunctionsMaterialTree::OnGetRightColumnWidth);
+	ColumnSizeData.OnWidthChanged = SSplitter::FOnSlotResized::CreateSP(Tree, &SMaterialLayersFunctionsMaterialTree::OnSetColumnWidth);
 
-		TSharedRef<SWidget> LeftSideWidget = SNullWidget::NullWidget;
-		TSharedRef<SWidget> RightSideWidget = SNullWidget::NullWidget;
-		FText NameOverride;
-		TSharedRef<SVerticalBox> WrapperWidget = SNew(SVerticalBox);
+	TSharedRef<SWidget> LeftSideWidget = SNullWidget::NullWidget;
+	TSharedRef<SWidget> RightSideWidget = SNullWidget::NullWidget;
+	FText NameOverride;
+	TSharedRef<SVerticalBox> WrapperWidget = SNew(SVerticalBox);
 
+	// STACK --------------------------------------------------
+	if (StackParameterData->StackDataType == EStackDataType::Stack)
+	{
+		WrapperWidget->AddSlot()
+			.Padding(3.0f)
+			.AutoHeight()
+			[
+				SNullWidget::NullWidget
+			];
+#if WITH_EDITOR
+		NameOverride = Tree->FunctionInstance->GetLayerName(StackParameterData->ParameterInfo.Index);
+#endif
+		TSharedRef<SHorizontalBox> HeaderRowWidget = SNew(SHorizontalBox);
 
-	// STACK ---------------------------------------------------------------------
-		if (StackParameterData->StackDataType == EStackDataType::Stack)
+		const float ThumbnailSize = 24.0f;
+		TArray<TSharedPtr<FSortedParamData>> AssetChildren = StackParameterData->Children;
+		if (AssetChildren.Num() > 0)
 		{
-			WrapperWidget->AddSlot()
-				.Padding(1.0f)
-				.AutoHeight()
+			HeaderRowWidget->AddSlot()
+				.HAlign(HAlign_Center)
+				.VAlign(VAlign_Center)
+				.Padding(2.5f, 0)
+				.AutoWidth()
 				[
 					SNullWidget::NullWidget
 				];
-#if WITH_EDITOR
-			NameOverride = Tree->FunctionInstance->GetLayerName(StackParameterData->ParameterInfo.Index);
-#endif
-			TSharedRef<SHorizontalBox> HeaderRowWidget = SNew(SHorizontalBox);
-			const float ThumbnailSize = 24.0f;
-			TArray<TSharedPtr<FSortedParamData>> AssetChildren = StackParameterData->Children;
-			if (AssetChildren.Num() > 0)
+		}
+		for (TSharedPtr<FSortedParamData> AssetChild : AssetChildren)
+		{
+			TSharedPtr<SBox> ThumbnailBox;
+			FAssetData AssetData;
+			AssetChild->ParameterHandle->GetValue(AssetData);
+			int32 PreviewIndex = INDEX_NONE;
+			int32 ThumbnailIndex = INDEX_NONE;
+			EMaterialParameterAssociation PreviewAssociation = EMaterialParameterAssociation::GlobalParameter;
+			if (UObject* AssetObject = AssetData.GetAsset())
 			{
-				HeaderRowWidget->AddSlot()
-					.HAlign(HAlign_Center)
-					.VAlign(VAlign_Center)
-					.Padding(2.5f, 0)
-					.AutoWidth()
-					[
-						SNullWidget::NullWidget
-					];
+				if (Cast<UMaterialFunctionInterface>(AssetObject)->GetMaterialFunctionUsage() == EMaterialFunctionUsage::MaterialLayer)
+				{
+					PreviewIndex = StackParameterData->ParameterInfo.Index;
+					PreviewAssociation = EMaterialParameterAssociation::LayerParameter;
+					Tree->UpdateThumbnailMaterial(PreviewAssociation, PreviewIndex);
+					ThumbnailIndex = PreviewIndex;
+				}
+				if (Cast<UMaterialFunctionInterface>(AssetObject)->GetMaterialFunctionUsage() == EMaterialFunctionUsage::MaterialLayerBlend)
+				{
+					PreviewIndex = StackParameterData->ParameterInfo.Index;
+					PreviewAssociation = EMaterialParameterAssociation::BlendParameter;
+					Tree->UpdateThumbnailMaterial(PreviewAssociation, PreviewIndex, true);
+					ThumbnailIndex = PreviewIndex - 1;
+				}
 			}
-			for (TSharedPtr<FSortedParamData> AssetChild : AssetChildren)
-			{
-				TSharedPtr<SBox> ThumbnailBox;
-				UObject* AssetObject;
-				AssetChild->ParameterHandle->GetValue(AssetObject);
-				const TSharedPtr<FAssetThumbnail> AssetThumbnail = MakeShareable(new FAssetThumbnail(AssetObject, ThumbnailSize, ThumbnailSize, InArgs._InTree->GetTreeThumbnailPool()));
-				HeaderRowWidget->AddSlot()
-					.AutoWidth()
-					.HAlign(HAlign_Center)
-					.VAlign(VAlign_Center)
-					.Padding(4.0f)
-					.MaxWidth(ThumbnailSize)
+			HeaderRowWidget->AddSlot()
+				.AutoWidth()
+				.HAlign(HAlign_Center)
+				.VAlign(VAlign_Center)
+				.Padding(4.0f)
+				.MaxWidth(ThumbnailSize)
+				[
+					SAssignNew(ThumbnailBox, SBox)
 					[
-						SAssignNew(ThumbnailBox, SBox)
-						[
-							AssetThumbnail->MakeThumbnailWidget()
-						]
-					];
-				ThumbnailBox->SetMaxDesiredHeight(ThumbnailSize);
-				ThumbnailBox->SetMinDesiredHeight(ThumbnailSize);
-			}
+						Tree->CreateThumbnailWidget(PreviewAssociation, ThumbnailIndex, ThumbnailSize)
+					]
+				];
+			ThumbnailBox->SetMaxDesiredHeight(ThumbnailSize);
+			ThumbnailBox->SetMinDesiredHeight(ThumbnailSize);
+			ThumbnailBox->SetMinDesiredWidth(ThumbnailSize);
+			ThumbnailBox->SetMaxDesiredWidth(ThumbnailSize);
+		}
 
+		if (StackParameterData->ParameterInfo.Index != 0)
+		{
 			HeaderRowWidget->AddSlot()
 				.VAlign(VAlign_Center)
 				.AutoWidth()
 				.Padding(5.0f)
 				[
 					SNew(STextBlock)
-					.Text(NameOverride)
-					.TextStyle(FEditorStyle::Get(), "BoldText")
+					.Text(this, &SMaterialLayersFunctionsMaterialTreeItem::GetLayerName, InArgs._InTree, StackParameterData->ParameterInfo.Index)
+					.Font(FEditorStyle::GetFontStyle(TEXT("MaterialEditor.Layers.EditableFontImportant")))
+				];
+				HeaderRowWidget->AddSlot()
+				.FillWidth(1.0f)
+				.VAlign(VAlign_Center)
+				[
+					SNullWidget::NullWidget
 				];
-			LeftSideWidget = HeaderRowWidget;
 		}
-	// END STACK
-	
-	// GROUP ---------------------------------------------------------------------
-		if (StackParameterData->StackDataType == EStackDataType::Group)
+		else
 		{
-			NameOverride = FText::FromName(StackParameterData->Group.GroupName);
-			LeftSideWidget = SNew(STextBlock)
-				.Text(NameOverride)
-				.TextStyle(FEditorStyle::Get(), "TinyText");
+			HeaderRowWidget->AddSlot()
+				.VAlign(VAlign_Center)
+				.AutoWidth()
+				.Padding(5.0f)
+				[
+					SNew(STextBlock)
+					.Text(NameOverride)
+					.TextStyle(FEditorStyle::Get(), "NormalText.Important")
+				];
 		}
+		LeftSideWidget = HeaderRowWidget;
+	}
+	// END STACK
+
+	// GROUP --------------------------------------------------
+	if (StackParameterData->StackDataType == EStackDataType::Group)
+	{
+		NameOverride = FText::FromName(StackParameterData->Group.GroupName);
+		LeftSideWidget = SNew(STextBlock)
+			.Text(NameOverride)
+			.TextStyle(FEditorStyle::Get(), "TinyText");
+	}
 	// END GROUP
 
-	// ASSET ---------------------------------------------------------------------
-		if (StackParameterData->StackDataType == EStackDataType::Asset)
+	// ASSET --------------------------------------------------
+	if (StackParameterData->StackDataType == EStackDataType::Asset)
+	{
+		StackParameterData->ParameterHandle->GetProperty()->SetMetaData(FName(TEXT("DisplayThumbnail")), TEXT("true"));
+		FIntPoint ThumbnailOverride;
+		if (StackParameterData->ParameterInfo.Association == EMaterialParameterAssociation::LayerParameter)
 		{
-			if (StackParameterData->ParameterInfo.Association == EMaterialParameterAssociation::LayerParameter)
-			{
-				NameOverride = FMaterialPropertyHelpers::LayerID;
-				StackParameterData->ParameterHandle->GetProperty()->SetMetaData(FName(TEXT("DisplayThumbnail")), TEXT("true"));
-			}
-			else if (StackParameterData->ParameterInfo.Association == EMaterialParameterAssociation::BlendParameter)
-			{
-				NameOverride = FMaterialPropertyHelpers::BlendID;
-			}
+			NameOverride = FMaterialPropertyHelpers::LayerID;
+			ThumbnailOverride = FIntPoint(64, 64);
+		}
+		else if (StackParameterData->ParameterInfo.Association == EMaterialParameterAssociation::BlendParameter)
+		{
+			NameOverride = FMaterialPropertyHelpers::BlendID;
+			ThumbnailOverride = FIntPoint(32, 32);
+		}
 
-			IDetailTreeNode& Node = *StackParameterData->ParameterNode;
-			FNodeWidgets NodeWidgets = Node.CreateNodeWidgets();
 
-			LeftSideWidget = StackParameterData->ParameterHandle->CreatePropertyNameWidget(NameOverride);
+		const int32 LayerStateIndex = StackParameterData->ParameterInfo.Association == EMaterialParameterAssociation::BlendParameter ? StackParameterData->ParameterInfo.Index + 1 : StackParameterData->ParameterInfo.Index;
 
-			StackParameterData->ParameterHandle->MarkResetToDefaultCustomized(false);
+		IDetailTreeNode& Node = *StackParameterData->ParameterNode;
+		FNodeWidgets NodeWidgets = Node.CreateNodeWidgets();
 
-			EMaterialParameterAssociation InAssociation = StackParameterData->ParameterInfo.Association;
+		LeftSideWidget = StackParameterData->ParameterHandle->CreatePropertyNameWidget(NameOverride);
 
-			RightSideWidget = 
-					SNew(SObjectPropertyEntryBox)
-					.AllowedClass(UMaterialFunctionInterface::StaticClass())
-					.ObjectPath(this, &SMaterialLayersFunctionsMaterialTreeItem::GetInstancePath, Tree)
-					.ThumbnailPool(Tree->GetTreeThumbnailPool())
-					.DisplayCompactSize(true);
-		}
-	// END ASSET
+		EMaterialParameterAssociation InAssociation = StackParameterData->ParameterInfo.Association;
+		UMaterialFunctionInterface* LocalFunction = nullptr;
+		TSharedPtr<SBox> ThumbnailBox;
 
-	// PROPERTY ------------------------------------------------------------------
-		if (StackParameterData->StackDataType == EStackDataType::Property)
+		if (StackParameterData->ParameterInfo.Association == EMaterialParameterAssociation::LayerParameter)
 		{
-			UMaterialFunctionInterface* OwningInterface = nullptr;
-			if (StackParameterData->ParameterInfo.Association == EMaterialParameterAssociation::LayerParameter)
-			{
-				OwningInterface = Tree->FunctionInstance->Layers[StackParameterData->ParameterInfo.Index];
-			}
-			if (StackParameterData->ParameterInfo.Association == EMaterialParameterAssociation::BlendParameter)
-			{
-				OwningInterface = Tree->FunctionInstance->Blends[StackParameterData->ParameterInfo.Index];
-			}
+			LocalFunction = Tree->FunctionInstance->Layers[StackParameterData->ParameterInfo.Index];
+		}
+		else if (StackParameterData->ParameterInfo.Association == EMaterialParameterAssociation::BlendParameter)
+		{
+			LocalFunction = Tree->FunctionInstance->Blends[StackParameterData->ParameterInfo.Index];
+		}
 
-			NameOverride = FText::FromName(StackParameterData->Parameter->ParameterInfo.Name);
+		RightSideWidget = SNew(SVerticalBox)
+			+ SVerticalBox::Slot()
+			[
+				SNew(SHorizontalBox)
+				+ SHorizontalBox::Slot()
+			.AutoWidth()
+			.HAlign(HAlign_Center)
+			.VAlign(VAlign_Center)
+			.Padding(4.0f)
+			.MaxWidth(ThumbnailOverride.X)
+			[
+				SAssignNew(ThumbnailBox, SBox)
+				[
+					Tree->CreateThumbnailWidget(StackParameterData->ParameterInfo.Association, StackParameterData->ParameterInfo.Index, ThumbnailOverride.X)
+				]
+			]
+		+ SHorizontalBox::Slot()
+			.FillWidth(1.0)
+			[
+				SNew(SObjectPropertyEntryBox)
+				.AllowedClass(UMaterialFunctionInterface::StaticClass())
+				.ObjectPath(this, &SMaterialLayersFunctionsMaterialTreeItem::GetInstancePath, Tree)
+				.DisplayCompactSize(true)
+			]
+	
+			];
+		ThumbnailBox->SetMaxDesiredHeight(ThumbnailOverride.Y);
+		ThumbnailBox->SetMinDesiredHeight(ThumbnailOverride.Y);
+		ThumbnailBox->SetMinDesiredWidth(ThumbnailOverride.X);
+		ThumbnailBox->SetMaxDesiredWidth(ThumbnailOverride.X);
+	}
+	// END ASSET
 
-			IDetailTreeNode& Node = *StackParameterData->ParameterNode;
-			TSharedPtr<IDetailPropertyRow> GeneratedRow = StaticCastSharedPtr<IDetailPropertyRow>(Node.GetRow());
-			IDetailPropertyRow& Row = *GeneratedRow.Get();
-			Row
-				.DisplayName(NameOverride)
-				.EditCondition(false, nullptr);
+	// PROPERTY ----------------------------------------------
+	if (StackParameterData->StackDataType == EStackDataType::Property)
+	{
 
+		UDEditorStaticComponentMaskParameterValue* CompMaskParam = Cast<UDEditorStaticComponentMaskParameterValue>(StackParameterData->Parameter);
+		UDEditorVectorParameterValue* VectorParam = Cast<UDEditorVectorParameterValue>(StackParameterData->Parameter);
+		UDEditorScalarParameterValue* ScalarParam = Cast<UDEditorScalarParameterValue>(StackParameterData->Parameter);
+		UDEditorTextureParameterValue* TextureParam = Cast<UDEditorTextureParameterValue>(StackParameterData->Parameter);
+		NameOverride = FText::FromName(StackParameterData->Parameter->ParameterInfo.Name);
 
-			UDEditorStaticComponentMaskParameterValue* CompMaskParam = Cast<UDEditorStaticComponentMaskParameterValue>(StackParameterData->Parameter);
-			UDEditorVectorParameterValue* VectorParam = Cast<UDEditorVectorParameterValue>(StackParameterData->Parameter);
-			UDEditorScalarParameterValue* ScalarParam = Cast<UDEditorScalarParameterValue>(StackParameterData->Parameter);
+		IDetailTreeNode& Node = *StackParameterData->ParameterNode;
+		TSharedPtr<IDetailPropertyRow> GeneratedRow = StaticCastSharedPtr<IDetailPropertyRow>(Node.GetRow());
+		IDetailPropertyRow& Row = *GeneratedRow.Get();
+		Row
+			.DisplayName(NameOverride);
 
-			if (VectorParam && VectorParam->bIsUsedAsChannelMask)
-			{
-				FOnGetPropertyComboBoxStrings GetMaskStrings = FOnGetPropertyComboBoxStrings::CreateStatic(&FMaterialPropertyHelpers::GetVectorChannelMaskComboBoxStrings);
-				FOnGetPropertyComboBoxValue GetMaskValue = FOnGetPropertyComboBoxValue::CreateStatic(&FMaterialPropertyHelpers::GetVectorChannelMaskValue, StackParameterData->Parameter);
-				FOnPropertyComboBoxValueSelected SetMaskValue = FOnPropertyComboBoxValueSelected::CreateStatic(&FMaterialPropertyHelpers::SetVectorChannelMaskValue, StackParameterData->ParameterNode->CreatePropertyHandle(), StackParameterData->Parameter, (UObject*)MaterialEditorInstance);
+		if (VectorParam && VectorParam->bIsUsedAsChannelMask)
+		{
+			FOnGetPropertyComboBoxStrings GetMaskStrings = FOnGetPropertyComboBoxStrings::CreateStatic(&FMaterialPropertyHelpers::GetVectorChannelMaskComboBoxStrings);
+			FOnGetPropertyComboBoxValue GetMaskValue = FOnGetPropertyComboBoxValue::CreateStatic(&FMaterialPropertyHelpers::GetVectorChannelMaskValue, StackParameterData->Parameter);
+			FOnPropertyComboBoxValueSelected SetMaskValue = FOnPropertyComboBoxValueSelected::CreateStatic(&FMaterialPropertyHelpers::SetVectorChannelMaskValue, StackParameterData->ParameterNode->CreatePropertyHandle(), StackParameterData->Parameter, (UObject*)MaterialEditorInstance);
 
-				FDetailWidgetRow& CustomWidget = Row.CustomWidget();
-				CustomWidget
+			FDetailWidgetRow& CustomWidget = Row.CustomWidget();
+			CustomWidget
 				.FilterString(NameOverride)
 				.NameContent()
 				[
@@ -1961,183 +2047,365 @@ public:
 						]
 					]
 				];
-			}
-			if (ScalarParam && ScalarParam->AtlasData.bIsUsedAsAtlasPosition)
-			{
-
-				TAttribute<bool> IsParamEnabled = TAttribute<bool>::Create(TAttribute<bool>::FGetter::CreateStatic(&FMaterialPropertyHelpers::IsOverriddenExpression, StackParameterData->Parameter));
+		}
+		else if (ScalarParam && ScalarParam->AtlasData.bIsUsedAsAtlasPosition)
+		{
+			const FText ParameterName = FText::FromName(StackParameterData->Parameter->ParameterInfo.Name);
 
-				const FText ParameterName = FText::FromName(StackParameterData->Parameter->ParameterInfo.Name);
+			FDetailWidgetRow& CustomWidget = Row.CustomWidget();
+			CustomWidget
+				.FilterString(ParameterName)
+				.NameContent()
+				[
+					SNew(STextBlock)
+					.Text(ParameterName)
+				.ToolTipText(FMaterialPropertyHelpers::GetParameterExpressionDescription(StackParameterData->Parameter, MaterialEditorInstance))
+				.Font(FEditorStyle::GetFontStyle(TEXT("PropertyWindow.NormalFont")))
+				]
+			.ValueContent()
+				.HAlign(HAlign_Fill)
+				.MaxDesiredWidth(400.0f)
+				[
+					SNew(SObjectPropertyEntryBox)
+					.ObjectPath(this, &SMaterialLayersFunctionsMaterialTreeItem::GetCurvePath, ScalarParam)
+					.AllowedClass(UCurveLinearColor::StaticClass())
+					.NewAssetFactories(TArray<UFactory*>())
+					.DisplayThumbnail(true)
+					.ThumbnailPool(InArgs._InTree->GetTreeThumbnailPool())
+					.DisplayCompactSize(true)
+				];
 
-				FDetailWidgetRow& CustomWidget = Row.CustomWidget();
-				CustomWidget
-					.FilterString(ParameterName)
-					.NameContent()
-					[
-						SNew(STextBlock)
-						.Text(ParameterName)
-						.ToolTipText(FMaterialPropertyHelpers::GetParameterExpressionDescription(StackParameterData->Parameter, MaterialEditorInstance))
-						.Font(FEditorStyle::GetFontStyle(TEXT("PropertyWindow.NormalFont")))
-					]
-					.ValueContent()
-					.HAlign(HAlign_Fill)
-					.MaxDesiredWidth(400.0f)
-					[
-						SNew(SObjectPropertyEntryBox)
-						.ObjectPath(this, &SMaterialLayersFunctionsMaterialTreeItem::GetCurvePath, ScalarParam)
-						.AllowedClass(UCurveLinearColor::StaticClass())
-						.NewAssetFactories(TArray<UFactory*>())
-						.DisplayThumbnail(true)
-						.ThumbnailPool(InArgs._InTree->GetTreeThumbnailPool())
-						.OnShouldSetAsset(FOnShouldSetAsset::CreateStatic(&FMaterialPropertyHelpers::OnShouldSetCurveAsset, ScalarParam->AtlasData.Atlas))
-						.OnShouldFilterAsset(FOnShouldFilterAsset::CreateStatic(&FMaterialPropertyHelpers::OnShouldFilterCurveAsset, ScalarParam->AtlasData.Atlas))
-						.OnObjectChanged(FOnSetObject::CreateStatic(&FMaterialPropertyHelpers::SetPositionFromCurveAsset, ScalarParam->AtlasData.Atlas, ScalarParam, StackParameterData->ParameterHandle, (UObject*)MaterialEditorInstance))
-						.DisplayCompactSize(true)
-					];
-
-			}
-			else if (!CompMaskParam)
+		}
+		else if (TextureParam)
+		{
+			UMaterial* Material = MaterialEditorInstance->PreviewMaterial;
+			if (Material != nullptr)
 			{
-				FNodeWidgets StoredNodeWidgets = Node.CreateNodeWidgets();
-				TSharedRef<SWidget> StoredRightSideWidget = StoredNodeWidgets.ValueWidget.ToSharedRef();
-				FDetailWidgetRow& CustomWidget = Row.CustomWidget();
-				CustomWidget
-					.FilterString(NameOverride)
-					.NameContent()
-					[
-						SNew(SHorizontalBox)
-						+ SHorizontalBox::Slot()
-					.VAlign(VAlign_Center)
-					[
-						SNew(STextBlock)
-						.Text(NameOverride)
+				UMaterialExpressionTextureSampleParameter* Expression = Material->FindExpressionByGUID<UMaterialExpressionTextureSampleParameter>(TextureParam->ExpressionId);
+				if (Expression != nullptr)
+				{
+					TWeakObjectPtr<UMaterialExpressionTextureSampleParameter> SamplerExpression = Expression;
+					TSharedPtr<SWidget> NameWidget;
+					TSharedPtr<SWidget> ValueWidget;
+					TSharedPtr<SVerticalBox> NameVerticalBox;
+					const FText ParameterName = FText::FromName(StackParameterData->Parameter->ParameterInfo.Name);
+					FDetailWidgetRow& CustomWidget = Row.CustomWidget();
+					CustomWidget
+						.FilterString(ParameterName)
+						.NameContent()
+						[
+							SAssignNew(NameVerticalBox, SVerticalBox)
+							+ SVerticalBox::Slot()
+						.AutoHeight()
+						[
+							SNew(STextBlock)
+							.Text(ParameterName)
 						.ToolTipText(FMaterialPropertyHelpers::GetParameterExpressionDescription(StackParameterData->Parameter, MaterialEditorInstance))
 						.Font(FEditorStyle::GetFontStyle(TEXT("PropertyWindow.NormalFont")))
-					]
-					]
-				.ValueContent()
-					[
-						StoredRightSideWidget
-					];
+						]
+						];
+					CustomWidget.ValueContent()
+						[
+							SNew(SObjectPropertyEntryBox)
+							.PropertyHandle(StackParameterData->ParameterNode->CreatePropertyHandle())
+							.AllowedClass(UTexture::StaticClass())
+							.ThumbnailPool(Tree->GetTreeThumbnailPool())
+						];
 
-			}
-			else
-			{
-				TSharedPtr<IPropertyHandle> RMaskProperty = StackParameterData->ParameterNode->CreatePropertyHandle()->GetChildHandle("R");
-				TSharedPtr<IPropertyHandle> GMaskProperty = StackParameterData->ParameterNode->CreatePropertyHandle()->GetChildHandle("G");
-				TSharedPtr<IPropertyHandle> BMaskProperty = StackParameterData->ParameterNode->CreatePropertyHandle()->GetChildHandle("B");
-				TSharedPtr<IPropertyHandle> AMaskProperty = StackParameterData->ParameterNode->CreatePropertyHandle()->GetChildHandle("A");
-				FDetailWidgetRow& CustomWidget = Row.CustomWidget();
-				CustomWidget
-					.FilterString(NameOverride)
-					.NameContent()
-					[
-						SNew(SHorizontalBox)
+					static const FName Red("R");
+					static const FName Green("G");
+					static const FName Blue("B");
+					static const FName Alpha("A");
+
+					if (!TextureParam->ChannelNames.R.IsEmpty())
+					{
+						NameVerticalBox->AddSlot()
+							[
+								SNew(SHorizontalBox)
+								+ SHorizontalBox::Slot()
+							.AutoWidth()
+							.Padding(20.0, 2.0, 4.0, 2.0)
+							[
+								SNew(STextBlock)
+								.Text(FText::FromName(Red))
+							.Font(FEditorStyle::GetFontStyle(TEXT("PropertyWindow.BoldFont")))
+							]
 						+ SHorizontalBox::Slot()
-						.VAlign(VAlign_Center)
-						[
-							SNew(STextBlock)
-							.Text(NameOverride)
-						.ToolTipText(FMaterialPropertyHelpers::GetParameterExpressionDescription(StackParameterData->Parameter, MaterialEditorInstance))
+							.HAlign(HAlign_Left)
+							.Padding(4.0, 2.0)
+							[
+								SNew(STextBlock)
+								.Text(TextureParam->ChannelNames.R)
 							.Font(FEditorStyle::GetFontStyle(TEXT("PropertyWindow.NormalFont")))
-						]
-					]
-				.ValueContent()
-					.MaxDesiredWidth(200.0f)
-					[
-						SNew(SHorizontalBox)
+							]
+							];
+					}
+					if (!TextureParam->ChannelNames.G.IsEmpty())
+					{
+						NameVerticalBox->AddSlot()
+							[
+								SNew(SHorizontalBox)
+								+ SHorizontalBox::Slot()
+							.Padding(20.0, 2.0, 4.0, 2.0)
+							.AutoWidth()
+							[
+								SNew(STextBlock)
+								.Text(FText::FromName(Green))
+							.Font(FEditorStyle::GetFontStyle(TEXT("PropertyWindow.BoldFont")))
+							]
 						+ SHorizontalBox::Slot()
-					.FillWidth(1.0f)
-					[
-						SNew(SHorizontalBox)
+							.HAlign(HAlign_Left)
+							.Padding(4.0, 2.0)
+							[
+								SNew(STextBlock)
+								.Text(TextureParam->ChannelNames.G)
+							.Font(FEditorStyle::GetFontStyle(TEXT("PropertyWindow.NormalFont")))
+							]
+							];
+					}
+					if (!TextureParam->ChannelNames.B.IsEmpty())
+					{
+						NameVerticalBox->AddSlot()
+							[
+								SNew(SHorizontalBox)
+								+ SHorizontalBox::Slot()
+							.Padding(20.0, 2.0, 4.0, 2.0)
+							.AutoWidth()
+							[
+								SNew(STextBlock)
+								.Text(FText::FromName(Blue))
+							.Font(FEditorStyle::GetFontStyle(TEXT("PropertyWindow.BoldFont")))
+							]
 						+ SHorizontalBox::Slot()
-					.HAlign(HAlign_Left)
-					.AutoWidth()
-					[
-						RMaskProperty->CreatePropertyNameWidget(FText::GetEmpty(), FText::GetEmpty(), false)
-					]
-				+ SHorizontalBox::Slot()
-					.HAlign(HAlign_Left)
-					.AutoWidth()
-					[
-						RMaskProperty->CreatePropertyValueWidget()
-					]
-				+ SHorizontalBox::Slot()
-					.HAlign(HAlign_Left)
-					.Padding(FMargin(10.0f, 0.0f, 0.0f, 0.0f))
-					.AutoWidth()
-					[
-						GMaskProperty->CreatePropertyNameWidget(FText::GetEmpty(), FText::GetEmpty(), false)
-					]
-				+ SHorizontalBox::Slot()
-					.HAlign(HAlign_Left)
-					.AutoWidth()
-					[
-						GMaskProperty->CreatePropertyValueWidget()
-					]
-				+ SHorizontalBox::Slot()
-					.HAlign(HAlign_Left)
-					.Padding(FMargin(10.0f, 0.0f, 0.0f, 0.0f))
-					.AutoWidth()
-					[
-						BMaskProperty->CreatePropertyNameWidget(FText::GetEmpty(), FText::GetEmpty(), false)
-					]
-				+ SHorizontalBox::Slot()
-					.HAlign(HAlign_Left)
-					.AutoWidth()
-					[
-						BMaskProperty->CreatePropertyValueWidget()
-					]
-				+ SHorizontalBox::Slot()
-					.HAlign(HAlign_Left)
-					.Padding(FMargin(10.0f, 0.0f, 0.0f, 0.0f))
-					.AutoWidth()
-					[
-						AMaskProperty->CreatePropertyNameWidget(FText::GetEmpty(), FText::GetEmpty(), false)
-					]
-				+ SHorizontalBox::Slot()
-					.HAlign(HAlign_Left)
-					.AutoWidth()
+							.HAlign(HAlign_Left)
+							.Padding(4.0, 2.0)
+							[
+								SNew(STextBlock)
+								.Text(TextureParam->ChannelNames.B)
+							.Font(FEditorStyle::GetFontStyle(TEXT("PropertyWindow.NormalFont")))
+							]
+							];
+					}
+					if (!TextureParam->ChannelNames.A.IsEmpty())
+					{
+						NameVerticalBox->AddSlot()
+							[
+								SNew(SHorizontalBox)
+								+ SHorizontalBox::Slot()
+							.Padding(20.0, 2.0, 4.0, 2.0)
+							.AutoWidth()
+							[
+								SNew(STextBlock)
+								.Text(FText::FromName(Alpha))
+							.Font(FEditorStyle::GetFontStyle(TEXT("PropertyWindow.BoldFont")))
+							]
+						+ SHorizontalBox::Slot()
+							.HAlign(HAlign_Left)
+							.Padding(4.0, 2.0)
+							[
+								SNew(STextBlock)
+								.Text(TextureParam->ChannelNames.A)
+							.Font(FEditorStyle::GetFontStyle(TEXT("PropertyWindow.NormalFont")))
+							]
+							];
+					}
+				}
+			}
+		}
+		else if (!CompMaskParam)
+		{
+			FNodeWidgets StoredNodeWidgets = Node.CreateNodeWidgets();
+			TSharedRef<SWidget> StoredRightSideWidget = StoredNodeWidgets.ValueWidget.ToSharedRef();
+			FDetailWidgetRow& CustomWidget = Row.CustomWidget();
+			CustomWidget
+				.FilterString(NameOverride)
+				.NameContent()
+				[
+					SNew(SHorizontalBox)
+					+ SHorizontalBox::Slot()
+					.VAlign(VAlign_Center)
 					[
-						AMaskProperty->CreatePropertyValueWidget()
-					]
+						SNew(STextBlock)
+						.Text(NameOverride)
+						.ToolTipText(FMaterialPropertyHelpers::GetParameterExpressionDescription(StackParameterData->Parameter, MaterialEditorInstance))
+						.Font(FEditorStyle::GetFontStyle(TEXT("PropertyWindow.NormalFont")))
 					]
-					];
-			}
-
-			FNodeWidgets NodeWidgets = Node.CreateNodeWidgets();
-			LeftSideWidget = NodeWidgets.NameWidget.ToSharedRef();
-			RightSideWidget = NodeWidgets.ValueWidget.ToSharedRef();
+				]
+				.ValueContent()
+				[
+					StoredRightSideWidget
+				];
 
 		}
-	// END PROPERTY
-
-	// PROPERTY CHILD ------------------------------------------------------------
-		if (StackParameterData->StackDataType == EStackDataType::PropertyChild)
+		else
 		{
-			FNodeWidgets NodeWidgets = StackParameterData->ParameterNode->CreateNodeWidgets();
-			LeftSideWidget = NodeWidgets.NameWidget.ToSharedRef();
-			RightSideWidget = NodeWidgets.ValueWidget.ToSharedRef();
+			TSharedPtr<IPropertyHandle> RMaskProperty = StackParameterData->ParameterNode->CreatePropertyHandle()->GetChildHandle("R");
+			TSharedPtr<IPropertyHandle> GMaskProperty = StackParameterData->ParameterNode->CreatePropertyHandle()->GetChildHandle("G");
+			TSharedPtr<IPropertyHandle> BMaskProperty = StackParameterData->ParameterNode->CreatePropertyHandle()->GetChildHandle("B");
+			TSharedPtr<IPropertyHandle> AMaskProperty = StackParameterData->ParameterNode->CreatePropertyHandle()->GetChildHandle("A");
+			FDetailWidgetRow& CustomWidget = Row.CustomWidget();
+			CustomWidget
+				.FilterString(NameOverride)
+				.NameContent()
+				[
+					SNew(SHorizontalBox)
+					+ SHorizontalBox::Slot()
+				.VAlign(VAlign_Center)
+				[
+					SNew(STextBlock)
+					.Text(NameOverride)
+				.ToolTipText(FMaterialPropertyHelpers::GetParameterExpressionDescription(StackParameterData->Parameter, MaterialEditorInstance))
+				.Font(FEditorStyle::GetFontStyle(TEXT("PropertyWindow.NormalFont")))
+				]
+				]
+			.ValueContent()
+				.MaxDesiredWidth(200.0f)
+				[
+					SNew(SHorizontalBox)
+					+ SHorizontalBox::Slot()
+				.FillWidth(1.0f)
+				[
+					SNew(SHorizontalBox)
+					+ SHorizontalBox::Slot()
+				.HAlign(HAlign_Left)
+				.AutoWidth()
+				[
+					RMaskProperty->CreatePropertyNameWidget(FText::GetEmpty(), FText::GetEmpty(), false)
+				]
+			+ SHorizontalBox::Slot()
+				.HAlign(HAlign_Left)
+				.AutoWidth()
+				[
+					RMaskProperty->CreatePropertyValueWidget()
+				]
+			+ SHorizontalBox::Slot()
+				.HAlign(HAlign_Left)
+				.Padding(FMargin(10.0f, 0.0f, 0.0f, 0.0f))
+				.AutoWidth()
+				[
+					GMaskProperty->CreatePropertyNameWidget(FText::GetEmpty(), FText::GetEmpty(), false)
+				]
+			+ SHorizontalBox::Slot()
+				.HAlign(HAlign_Left)
+				.AutoWidth()
+				[
+					GMaskProperty->CreatePropertyValueWidget()
+				]
+			+ SHorizontalBox::Slot()
+				.HAlign(HAlign_Left)
+				.Padding(FMargin(10.0f, 0.0f, 0.0f, 0.0f))
+				.AutoWidth()
+				[
+					BMaskProperty->CreatePropertyNameWidget(FText::GetEmpty(), FText::GetEmpty(), false)
+				]
+			+ SHorizontalBox::Slot()
+				.HAlign(HAlign_Left)
+				.AutoWidth()
+				[
+					BMaskProperty->CreatePropertyValueWidget()
+				]
+			+ SHorizontalBox::Slot()
+				.HAlign(HAlign_Left)
+				.Padding(FMargin(10.0f, 0.0f, 0.0f, 0.0f))
+				.AutoWidth()
+				[
+					AMaskProperty->CreatePropertyNameWidget(FText::GetEmpty(), FText::GetEmpty(), false)
+				]
+			+ SHorizontalBox::Slot()
+				.HAlign(HAlign_Left)
+				.AutoWidth()
+				[
+					AMaskProperty->CreatePropertyValueWidget()
+				]
+				]
+				];
 		}
+
+		FNodeWidgets NodeWidgets = Node.CreateNodeWidgets();
+		LeftSideWidget = NodeWidgets.NameWidget.ToSharedRef();
+		RightSideWidget = NodeWidgets.ValueWidget.ToSharedRef();
+	}
+	// END PROPERTY
+
+	// PROPERTY CHILD ----------------------------------------
+	if (StackParameterData->StackDataType == EStackDataType::PropertyChild)
+	{
+		FNodeWidgets NodeWidgets = StackParameterData->ParameterNode->CreateNodeWidgets();
+		LeftSideWidget = NodeWidgets.NameWidget.ToSharedRef();
+		RightSideWidget = NodeWidgets.ValueWidget.ToSharedRef();
+	}
 	// END PROPERTY CHILD
 
-		RightSideWidget->SetEnabled(false);
+	// FINAL WRAPPER
 
-		// Final wrapper
-		if (StackParameterData->StackDataType == EStackDataType::Stack)
+	LeftSideWidget->SetEnabled(false);
+	RightSideWidget->SetEnabled(false);
+	if (StackParameterData->StackDataType == EStackDataType::Stack)
+	{
+		TSharedPtr<SHorizontalBox> FinalStack;
+		WrapperWidget->AddSlot()
+			.AutoHeight()
+			[
+				SNew(SBorder)
+				.BorderImage(this, &SMaterialLayersFunctionsMaterialTreeItem::GetBorderImage)
+				.Padding(0.0f)
+				[
+					SAssignNew(FinalStack, SHorizontalBox)
+				]
+			];
+		FinalStack->AddSlot()
+			.AutoWidth()
+			.VAlign(VAlign_Center)
+			.Padding(FMargin(2.0f))
+			[
+				SNew(SExpanderArrow, SharedThis(this))
+			];
+		FinalStack->AddSlot()
+			.Padding(FMargin(2.0f))
+			.VAlign(VAlign_Center)
+			[
+				LeftSideWidget
+			];
+	}
+	else
+	{
+		FSlateBrush* StackBrush;
+		switch (StackParameterData->ParameterInfo.Association)
+		{
+		case EMaterialParameterAssociation::LayerParameter:
+		{
+			StackBrush = const_cast<FSlateBrush*>(FEditorStyle::GetBrush("MaterialInstanceEditor.StackBody"));
+			break;
+		}
+		case EMaterialParameterAssociation::BlendParameter:
 		{
-			WrapperWidget->AddSlot()
-				.AutoHeight()
+			StackBrush = const_cast<FSlateBrush*>(FEditorStyle::GetBrush("MaterialInstanceEditor.StackBodyBlend"));
+			break;
+		}
+		default:
+			break;
+		}
+		WrapperWidget->AddSlot()
+			.AutoHeight()
+			[
+				SNew(SBorder)
+				.BorderImage(this, &SMaterialLayersFunctionsMaterialTreeItem::GetBorderImage)
+				.Padding(0.0f)
 				[
-					SNew(SBorder)
-					.BorderImage(FEditorStyle::GetBrush("MaterialInstanceEditor.StackHeader"))
-					.Padding(0.0f)
+					SNew(SSplitter)
+					.Style(FEditorStyle::Get(), "DetailsView.Splitter")
+					.PhysicalSplitterHandleSize(1.0f)
+					.HitDetectionSplitterHandleSize(5.0f)
+					+ SSplitter::Slot()
+					.Value(ColumnSizeData.LeftColumnWidth)
+					.OnSlotResized(ColumnSizeData.OnWidthChanged)
+					.Value(0.25f)
 					[
 						SNew(SHorizontalBox)
 						+ SHorizontalBox::Slot()
 						.AutoWidth()
 						.VAlign(VAlign_Center)
-						.Padding(FMargin(2.0f))
+						.Padding(FMargin(3.0f))
 						[
 							SNew(SExpanderArrow, SharedThis(this))
 						]
@@ -2148,242 +2416,81 @@ public:
 							LeftSideWidget
 						]
 					]
-				];
-		}
-		else
-		{
-			if (StackParameterData->StackDataType == EStackDataType::Asset && StackParameterData->ParameterInfo.Association == EMaterialParameterAssociation::BlendParameter)
-			{
-				WrapperWidget->AddSlot()
-					[
-						SNew(SBorder)
-						.BorderImage(FEditorStyle::GetBrush("MaterialInstanceEditor.StackBody"))
-						.Padding(FMargin(5.0f, 2.0f))
-					[
-						SNew(SImage)
-						.Image(FEditorStyle::GetBrush("ThinLine.Horizontal"))
-						.ColorAndOpacity(FLinearColor(.2f, .2f, .2f, 1.0f))
-					]
-					];
-			}
-			WrapperWidget->AddSlot()
-				.AutoHeight()
-				[
-					SNew(SBorder)
-					.BorderImage(FEditorStyle::GetBrush("MaterialInstanceEditor.StackBody"))
-					.Padding(0.0f)
+					+ SSplitter::Slot()
+					.Value(ColumnSizeData.RightColumnWidth)
+					.OnSlotResized(ColumnSizeData.OnWidthChanged)
 					[
-						SNew(SSplitter)
-						.Style(FEditorStyle::Get(), "DetailsView.Splitter")
-						.PhysicalSplitterHandleSize(1.0f)
-						.HitDetectionSplitterHandleSize(5.0f)
-						+ SSplitter::Slot()
-						.Value(ColumnSizeData.LeftColumnWidth)
-						.OnSlotResized(ColumnSizeData.OnWidthChanged)
-						.Value(0.25f)
-						[
-							SNew(SHorizontalBox)
-							+ SHorizontalBox::Slot()
-							.AutoWidth()
-							.VAlign(VAlign_Center)
-							.Padding(FMargin(3.0f))
-							[
-								SNew(SExpanderArrow, SharedThis(this))
-							]
-							+ SHorizontalBox::Slot()
-							.Padding(FMargin(2.0f))
-							.VAlign(VAlign_Center)
-							[
-								LeftSideWidget
-							]
-						]
-						+ SSplitter::Slot()
-						.Value(ColumnSizeData.RightColumnWidth)
-						.OnSlotResized(ColumnSizeData.OnWidthChanged)
+						SNew(SHorizontalBox)
+						+ SHorizontalBox::Slot()
+						.MaxWidth(350.0f)
+						.Padding(FMargin(5.0f, 2.0f, 0.0f, 2.0f))
 						[
-							SNew(SHorizontalBox)
-							+ SHorizontalBox::Slot()
-							.MaxWidth(350.0f)
-							.Padding(FMargin(5.0f, 2.0f, 0.0f, 2.0f))
-							[
-								RightSideWidget
-							]
+							RightSideWidget
 						]
 					]
-				];
-		}
-
-
-		this->ChildSlot
-			[
-				WrapperWidget
+				]
 			];
-
-		STableRow< TSharedPtr<FSortedParamData> >::ConstructInternal(
-			STableRow::FArguments()
-			.Style(FEditorStyle::Get(), "DetailsView.TreeView.TableRow")
-			.ShowSelection(false),
-			InOwnerTableView
-		);
-	}
-
-	FString GetInstancePath(SMaterialLayersFunctionsMaterialTree* Tree) const
-	{
-		FString InstancePath;
-		if (StackParameterData->ParameterInfo.Association == EMaterialParameterAssociation::BlendParameter && Tree->FunctionInstance->Blends.IsValidIndex(StackParameterData->ParameterInfo.Index))
-		{
-			InstancePath = Tree->FunctionInstance->Blends[StackParameterData->ParameterInfo.Index]->GetPathName();
-		}
-		else if (StackParameterData->ParameterInfo.Association == EMaterialParameterAssociation::LayerParameter && Tree->FunctionInstance->Layers.IsValidIndex(StackParameterData->ParameterInfo.Index))
-		{
-			InstancePath = Tree->FunctionInstance->Layers[StackParameterData->ParameterInfo.Index]->GetPathName();
-		}			
-		return InstancePath;
 	}
 
-	// Block double click expansion
-	virtual FReply OnMouseButtonDoubleClick(const FGeometry& InMyGeometry, const FPointerEvent& InMouseEvent) override { return FReply::Handled(); };
 
-	/** The node info to build the tree view row from. */
-	TSharedPtr<FSortedParamData> StackParameterData;
+	this->ChildSlot
+		[
+			WrapperWidget
+		];
 
-	UMaterialEditorPreviewParameters* MaterialEditorInstance;
-};
 
-FString SMaterialLayersFunctionsMaterialTreeItem::GetCurvePath(UDEditorScalarParameterValue* Parameter) const
-{
-	FString Path = Parameter->AtlasData.Curve->GetPathName();
-	return Path;
+	STableRow< TSharedPtr<FSortedParamData> >::ConstructInternal(
+		STableRow::FArguments()
+		.Style(FEditorStyle::Get(), "DetailsView.TreeView.TableRow")
+		.ShowSelection(false),
+		InOwnerTableView
+	);
 }
 
-void SMaterialLayersFunctionsMaterialWrapper::Refresh()
+FString SMaterialLayersFunctionsMaterialTreeItem::GetInstancePath(SMaterialLayersFunctionsMaterialTree* InTree) const
 {
-	LayerParameter = nullptr;
-	NestedTree->CreateGroupsWidget();
-	LayerParameter = NestedTree->FunctionParameter;
-	FOnClicked 	OnChildButtonClicked = FOnClicked::CreateStatic(&FMaterialPropertyHelpers::OnClickedSaveNewMaterialInstance, ImplicitConv<UMaterialInterface*>(MaterialEditorInstance->OriginalMaterial), ImplicitConv<UObject*>(MaterialEditorInstance));
-
-	if (LayerParameter != nullptr)
-	{
-		this->ChildSlot
-		[
-			SNew(SBorder)
-			.BorderImage(FEditorStyle::GetBrush("MaterialInstanceEditor.LayersBorder"))
-			.Padding(FMargin(4.0f))
-			[
-				SNew(SVerticalBox)
-				+ SVerticalBox::Slot()
-				.AutoHeight()
-				[
-					SNew(SHorizontalBox)
-					+SHorizontalBox::Slot()
-					.AutoWidth()
-					[					
-						SNew(STextBlock)
-						.Text(FText::FromName(NestedTree->LayersFunctionsParameterName))
-						.TextStyle(FEditorStyle::Get(), "LargeText")
-					]
-					+ SHorizontalBox::Slot()
-					.FillWidth(1.0f)
-					[
-						SNullWidget::NullWidget
-					]
-					+ SHorizontalBox::Slot()
-					.AutoWidth()
-					.Padding(2.0f)
-					[
-						SNew(SButton)
-						.ButtonStyle(FEditorStyle::Get(), "FlatButton.Dark")
-						.HAlign(HAlign_Center)
-						.OnClicked(OnChildButtonClicked)
-						.ToolTipText(LOCTEXT("SaveToChildInstance", "Save To Child Instance"))
-						.Content()
-						[
-							SNew(SHorizontalBox)
-							+ SHorizontalBox::Slot()
-							.AutoWidth()
-							[
-								SNew(STextBlock)
-								.Font(FEditorStyle::Get().GetFontStyle("FontAwesome.10"))
-								.TextStyle(FEditorStyle::Get(), "NormalText.Important")
-								.Text(FText::FromString(FString(TEXT("\xf0c7 \xf149"))) /*fa-filter*/)
-							]
-							+ SHorizontalBox::Slot()
-							.AutoWidth()
-							[
-								SNew(STextBlock)
-								.TextStyle(FEditorStyle::Get(), "NormalText.Important")
-								.Text(FText::FromString(FString(TEXT(" Save Child"))) /*fa-filter*/)
-							]
-						]
-					]
-				]
-				+ SVerticalBox::Slot()
-				.Padding(FMargin(3.0f, 0.0f))
-				[
-					NestedTree.ToSharedRef()
-				]
-			]
-		];
-	}
-	else if(MaterialEditorInstance->OriginalFunction)
+	FString InstancePath;
+	if (StackParameterData->ParameterInfo.Association == EMaterialParameterAssociation::BlendParameter && InTree->FunctionInstance->Blends.IsValidIndex(StackParameterData->ParameterInfo.Index))
 	{
-		this->ChildSlot
-			[
-				SNew(SBorder)
-				.BorderImage(FEditorStyle::GetBrush("MaterialInstanceEditor.StackBody"))
-				.Padding(FMargin(4.0f))
-				[
-					SNew(STextBlock)
-					.Text(LOCTEXT("NoMaterialAttributeLayersAllowed", "Material Functions, Layers, and Blends cannot contain Material Attribute Layers nodes."))
-					.AutoWrapText(true)
-				]
-			];
+		InstancePath = InTree->FunctionInstance->Blends[StackParameterData->ParameterInfo.Index]->GetPathName();
 	}
-	else
+	else if (StackParameterData->ParameterInfo.Association == EMaterialParameterAssociation::LayerParameter && InTree->FunctionInstance->Layers.IsValidIndex(StackParameterData->ParameterInfo.Index))
 	{
-		this->ChildSlot
-			[
-				SNew(SBorder)
-				.BorderImage(FEditorStyle::GetBrush("MaterialInstanceEditor.StackBody"))
-				.Padding(FMargin(4.0f))
-			[
-				SNew(STextBlock)
-				.Text(LOCTEXT("AddMaterialLayerParameterPrompt", "Add a Material Attribute Layers parameter to see it here."))
-				.AutoWrapText(true)
-			]
-			];
+		InstancePath = InTree->FunctionInstance->Layers[StackParameterData->ParameterInfo.Index]->GetPathName();
 	}
+	return InstancePath;
 }
 
-
-void SMaterialLayersFunctionsMaterialWrapper::Construct(const FArguments& InArgs)
-{
-	NestedTree = SNew(SMaterialLayersFunctionsMaterialTree)
-		.InMaterialEditorInstance(InArgs._InMaterialEditorInstance)
-		.InShowHiddenDelegate(InArgs._InShowHiddenDelegate);
-
-	LayerParameter = NestedTree->FunctionParameter;
-	MaterialEditorInstance = InArgs._InMaterialEditorInstance;
-	Refresh();
-
-}
-
-void SMaterialLayersFunctionsMaterialWrapper::SetEditorInstance(UMaterialEditorPreviewParameters* InMaterialEditorInstance)
-{
-	NestedTree->MaterialEditorInstance = InMaterialEditorInstance;
-	NestedTree->CreateGroupsWidget();
-	Refresh();
-}
-
-
 void SMaterialLayersFunctionsMaterialTree::Construct(const FArguments& InArgs)
 {
 	ColumnWidth = 0.5f;
 	MaterialEditorInstance = InArgs._InMaterialEditorInstance;
-	ShowHiddenDelegate = InArgs._InShowHiddenDelegate;
+	Wrapper = InArgs._InWrapper;
 	CreateGroupsWidget();
 
+#ifdef WITH_EDITOR
+	//Fixup for adding new bool arrays to the class
+	if (FunctionInstance)
+	{
+		if (FunctionInstance->Layers.Num() != FunctionInstance->RestrictToLayerRelatives.Num())
+		{
+			int32 OriginalSize = FunctionInstance->RestrictToLayerRelatives.Num();
+			for (int32 LayerIt = 0; LayerIt < FunctionInstance->Layers.Num() - OriginalSize; LayerIt++)
+			{
+				FunctionInstance->RestrictToLayerRelatives.Add(false);
+			}
+		}
+		if (FunctionInstance->Blends.Num() != FunctionInstance->RestrictToBlendRelatives.Num())
+		{
+			int32 OriginalSize = FunctionInstance->RestrictToBlendRelatives.Num();
+			for (int32 BlendIt = 0; BlendIt < FunctionInstance->Blends.Num() - OriginalSize; BlendIt++)
+			{
+				FunctionInstance->RestrictToBlendRelatives.Add(false);
+			}
+		}
+	}
+#endif
+
 	STreeView<TSharedPtr<FSortedParamData>>::Construct(
 		STreeView::FArguments()
 		.TreeItemsSource(&LayerProperties)
@@ -2410,21 +2517,22 @@ void SMaterialLayersFunctionsMaterialTree::OnGetChildrenMaterialLayersFunctionsT
 	OutChildren = InParent->Children;
 }
 
+
 void SMaterialLayersFunctionsMaterialTree::OnExpansionChanged(TSharedPtr<FSortedParamData> Item, bool bIsExpanded)
 {
-	bool* ExpansionValue = MaterialEditorInstance->OriginalMaterial->LayerParameterExpansion.Find(Item->NodeKey);
+	bool* ExpansionValue = MaterialEditorInstance->PreviewMaterial->LayerParameterExpansion.Find(Item->NodeKey);
 	if (ExpansionValue == nullptr)
 	{
-		MaterialEditorInstance->OriginalMaterial->LayerParameterExpansion.Add(Item->NodeKey, bIsExpanded);
+		MaterialEditorInstance->PreviewMaterial->LayerParameterExpansion.Add(Item->NodeKey, bIsExpanded);
 	}
 	else if (*ExpansionValue != bIsExpanded)
 	{
-		MaterialEditorInstance->OriginalMaterial->LayerParameterExpansion.Emplace(Item->NodeKey, bIsExpanded);
+		MaterialEditorInstance->PreviewMaterial->LayerParameterExpansion.Emplace(Item->NodeKey, bIsExpanded);
 	}
 	// Expand any children that are also expanded
 	for (auto Child : Item->Children)
 	{
-		bool* ChildExpansionValue = MaterialEditorInstance->OriginalMaterial->LayerParameterExpansion.Find(Child->NodeKey);
+		bool* ChildExpansionValue = MaterialEditorInstance->PreviewMaterial->LayerParameterExpansion.Find(Child->NodeKey);
 		if (ChildExpansionValue != nullptr && *ChildExpansionValue == true)
 		{
 			SetItemExpansion(Child, true);
@@ -2438,7 +2546,7 @@ void SMaterialLayersFunctionsMaterialTree::SetParentsExpansionState()
 	{
 		if (Pair->Children.Num())
 		{
-			bool* bIsExpanded = MaterialEditorInstance->OriginalMaterial->LayerParameterExpansion.Find(Pair->NodeKey);
+			bool* bIsExpanded = MaterialEditorInstance->PreviewMaterial->LayerParameterExpansion.Find(Pair->NodeKey);
 			if (bIsExpanded)
 			{
 				SetItemExpansion(Pair, *bIsExpanded);
@@ -2449,32 +2557,24 @@ void SMaterialLayersFunctionsMaterialTree::SetParentsExpansionState()
 
 TSharedPtr<class FAssetThumbnailPool> SMaterialLayersFunctionsMaterialTree::GetTreeThumbnailPool()
 {
-	return Generator->GetGeneratedThumbnailPool();
+	return Wrapper->GetGenerator()->GetGeneratedThumbnailPool();
 }
 
 void SMaterialLayersFunctionsMaterialTree::CreateGroupsWidget()
 {
 	check(MaterialEditorInstance);
-	FunctionParameter = nullptr;
+	MaterialEditorInstance->RegenerateArrays();
 	NonLayerProperties.Empty();
 	LayerProperties.Empty();
-	FPropertyEditorModule& Module = FModuleManager::LoadModuleChecked<FPropertyEditorModule>("PropertyEditor");
-	if (!Generator.IsValid())
-	{
-		FPropertyRowGeneratorArgs Args;
-		Generator = Module.CreatePropertyRowGenerator(Args);
+	FunctionParameter = nullptr;
 
-		TArray<UObject*> Objects;
-		Objects.Add(MaterialEditorInstance);
-		Generator->SetObjects(Objects);
-	}
-	else
+	const TArray<TSharedRef<IDetailTreeNode>> TestData = Wrapper->GetGenerator()->GetRootTreeNodes();
+
+	if (TestData.Num() == 0)
 	{
-		TArray<UObject*> Objects;
-		Objects.Add(MaterialEditorInstance);
-		Generator->SetObjects(Objects);
+		return;
 	}
-	const TArray<TSharedRef<IDetailTreeNode>> TestData = Generator->GetRootTreeNodes();
+
 	TSharedPtr<IDetailTreeNode> Category = TestData[0];
 	TSharedPtr<IDetailTreeNode> ParameterGroups;
 	TArray<TSharedRef<IDetailTreeNode>> Children;
@@ -2528,23 +2628,48 @@ void SMaterialLayersFunctionsMaterialTree::CreateGroupsWidget()
 				LayerHandle->GetNumChildren(LayerChildren);
 				uint32 BlendChildren;
 				BlendHandle->GetNumChildren(BlendChildren);
+				if (MaterialEditorInstance->StoredLayerPreviews.Num() != LayerChildren)
+				{
+					MaterialEditorInstance->StoredLayerPreviews.Empty();
+					MaterialEditorInstance->StoredLayerPreviews.AddDefaulted(LayerChildren);
+				}
+				if (MaterialEditorInstance->StoredBlendPreviews.Num() != BlendChildren)
+				{
+					MaterialEditorInstance->StoredBlendPreviews.Empty();
+					MaterialEditorInstance->StoredBlendPreviews.AddDefaulted(BlendChildren);
+				}
 
-				TSharedPtr<FSortedParamData> StackProperty(new FSortedParamData());
+				TSharedRef<FSortedParamData> StackProperty = MakeShared<FSortedParamData>();
 				StackProperty->StackDataType = EStackDataType::Stack;
 				StackProperty->Parameter = Parameter;
 				StackProperty->ParameterInfo.Index = LayerChildren - 1;
 				StackProperty->NodeKey = FString::FromInt(StackProperty->ParameterInfo.Index);
 
 
-				TSharedPtr<FSortedParamData> ChildProperty(new FSortedParamData());
+				TSharedRef<FSortedParamData> ChildProperty = MakeShared<FSortedParamData>();
 				ChildProperty->StackDataType = EStackDataType::Asset;
 				ChildProperty->Parameter = Parameter;
 				ChildProperty->ParameterHandle = LayerHandle->AsArray()->GetElement(LayerChildren - 1);
-				ChildProperty->ParameterNode = Generator->FindTreeNode(ChildProperty->ParameterHandle);
+				ChildProperty->ParameterNode = Wrapper->GetGenerator()->FindTreeNode(ChildProperty->ParameterHandle);
 				ChildProperty->ParameterInfo.Index = LayerChildren - 1;
 				ChildProperty->ParameterInfo.Association = EMaterialParameterAssociation::LayerParameter;
 				ChildProperty->NodeKey = FString::FromInt(ChildProperty->ParameterInfo.Index) + FString::FromInt(ChildProperty->ParameterInfo.Association);
 
+				UObject* AssetObject;
+				ChildProperty->ParameterHandle->GetValue(AssetObject);
+				if (AssetObject)
+				{
+					if (MaterialEditorInstance->StoredLayerPreviews[LayerChildren - 1] == nullptr)
+					{
+						MaterialEditorInstance->StoredLayerPreviews[LayerChildren - 1] = (NewObject<UMaterialInstanceConstant>(MaterialEditorInstance, NAME_None));
+					}
+					UMaterialInterface* EditedMaterial = Cast<UMaterialFunctionInterface>(AssetObject)->GetPreviewMaterial();
+					if (MaterialEditorInstance->StoredLayerPreviews[LayerChildren - 1] && MaterialEditorInstance->StoredLayerPreviews[LayerChildren - 1]->Parent != EditedMaterial)
+					{
+						MaterialEditorInstance->StoredLayerPreviews[LayerChildren - 1]->SetParentEditorOnly(EditedMaterial);
+					}
+				}
+
 				StackProperty->Children.Add(ChildProperty);
 				LayerProperties.Add(StackProperty);
 
@@ -2552,32 +2677,57 @@ void SMaterialLayersFunctionsMaterialTree::CreateGroupsWidget()
 				{
 					for (int32 Counter = BlendChildren - 1; Counter >= 0; Counter--)
 					{
-						ChildProperty = MakeShareable(new FSortedParamData());
+						ChildProperty = MakeShared<FSortedParamData>();
 						ChildProperty->StackDataType = EStackDataType::Asset;
 						ChildProperty->Parameter = Parameter;
 						ChildProperty->ParameterHandle = BlendHandle->AsArray()->GetElement(Counter);
-						ChildProperty->ParameterNode = Generator->FindTreeNode(ChildProperty->ParameterHandle);
+						ChildProperty->ParameterNode = Wrapper->GetGenerator()->FindTreeNode(ChildProperty->ParameterHandle);
 						ChildProperty->ParameterInfo.Index = Counter;
 						ChildProperty->ParameterInfo.Association = EMaterialParameterAssociation::BlendParameter;
 						ChildProperty->NodeKey = FString::FromInt(ChildProperty->ParameterInfo.Index) + FString::FromInt(ChildProperty->ParameterInfo.Association);
+						ChildProperty->ParameterHandle->GetValue(AssetObject);
+						if (AssetObject)
+						{
+							if (MaterialEditorInstance->StoredBlendPreviews[Counter] == nullptr)
+							{
+								MaterialEditorInstance->StoredBlendPreviews[Counter] = (NewObject<UMaterialInstanceConstant>(MaterialEditorInstance, NAME_None));
+							}
+							UMaterialInterface* EditedMaterial = Cast<UMaterialFunctionInterface>(AssetObject)->GetPreviewMaterial();
+							if (MaterialEditorInstance->StoredBlendPreviews[Counter] && MaterialEditorInstance->StoredBlendPreviews[Counter]->Parent != EditedMaterial)
+							{
+								MaterialEditorInstance->StoredBlendPreviews[Counter]->SetParentEditorOnly(EditedMaterial);
+							}
+						}
 						LayerProperties.Last()->Children.Add(ChildProperty);
 
-						StackProperty = MakeShareable(new FSortedParamData());
+						StackProperty = MakeShared<FSortedParamData>();
 						StackProperty->StackDataType = EStackDataType::Stack;
 						StackProperty->Parameter = Parameter;
 						StackProperty->ParameterInfo.Index = Counter;
 						StackProperty->NodeKey = FString::FromInt(StackProperty->ParameterInfo.Index);
 						LayerProperties.Add(StackProperty);
 
-						ChildProperty = MakeShareable(new FSortedParamData());
+						ChildProperty = MakeShared<FSortedParamData>();
 						ChildProperty->StackDataType = EStackDataType::Asset;
 						ChildProperty->Parameter = Parameter;
 						ChildProperty->ParameterHandle = LayerHandle->AsArray()->GetElement(Counter);
-						ChildProperty->ParameterNode = Generator->FindTreeNode(ChildProperty->ParameterHandle);
+						ChildProperty->ParameterNode = Wrapper->GetGenerator()->FindTreeNode(ChildProperty->ParameterHandle);
 						ChildProperty->ParameterInfo.Index = Counter;
 						ChildProperty->ParameterInfo.Association = EMaterialParameterAssociation::LayerParameter;
 						ChildProperty->NodeKey = FString::FromInt(ChildProperty->ParameterInfo.Index) + FString::FromInt(ChildProperty->ParameterInfo.Association);
-
+						ChildProperty->ParameterHandle->GetValue(AssetObject);
+						if (AssetObject)
+						{
+							if (MaterialEditorInstance->StoredLayerPreviews[Counter] == nullptr)
+							{
+								MaterialEditorInstance->StoredLayerPreviews[Counter] = (NewObject<UMaterialInstanceConstant>(MaterialEditorInstance, NAME_None));
+							}
+							UMaterialInterface* EditedMaterial = Cast<UMaterialFunctionInterface>(AssetObject)->GetPreviewMaterial();
+							if (MaterialEditorInstance->StoredLayerPreviews[Counter] && MaterialEditorInstance->StoredLayerPreviews[Counter]->Parent != EditedMaterial)
+							{
+								MaterialEditorInstance->StoredLayerPreviews[Counter]->SetParentEditorOnly(EditedMaterial);
+							}
+						}
 						LayerProperties.Last()->Children.Add(ChildProperty);
 					}
 				}
@@ -2595,7 +2745,7 @@ void SMaterialLayersFunctionsMaterialTree::CreateGroupsWidget()
 
 				NonLayerProperty.Parameter = Parameter;
 				NonLayerProperty.ParameterGroup = ParameterGroup;
-				NonLayerProperty.ParameterNode = Generator->FindTreeNode(ParameterValueProperty);
+				NonLayerProperty.ParameterNode = Wrapper->GetGenerator()->FindTreeNode(ParameterValueProperty);
 				NonLayerProperty.ParameterHandle = NonLayerProperty.ParameterNode->CreatePropertyHandle();
 				NonLayerProperty.UnsortedName = Parameter->ParameterInfo.Name;
 
@@ -2614,6 +2764,73 @@ void SMaterialLayersFunctionsMaterialTree::CreateGroupsWidget()
 	SetParentsExpansionState();
 }
 
+
+TSharedRef<SWidget> SMaterialLayersFunctionsMaterialTree::CreateThumbnailWidget(EMaterialParameterAssociation InAssociation, int32 InIndex, float InThumbnailSize)
+{
+	UObject* ThumbnailObject = nullptr;
+	if (InAssociation == EMaterialParameterAssociation::LayerParameter)
+	{
+		ThumbnailObject = MaterialEditorInstance->StoredLayerPreviews[InIndex];
+	}
+	else if (InAssociation == EMaterialParameterAssociation::BlendParameter)
+	{
+		ThumbnailObject = MaterialEditorInstance->StoredBlendPreviews[InIndex];
+	}
+	const TSharedPtr<FAssetThumbnail> AssetThumbnail = MakeShareable(new FAssetThumbnail(ThumbnailObject, InThumbnailSize, InThumbnailSize, GetTreeThumbnailPool()));
+	TSharedRef<SWidget> ThumbnailWidget = AssetThumbnail->MakeThumbnailWidget();
+	return ThumbnailWidget;
+}
+
+void SMaterialLayersFunctionsMaterialTree::UpdateThumbnailMaterial(TEnumAsByte<EMaterialParameterAssociation> InAssociation, int32 InIndex, bool bAlterBlendIndex)
+{
+	// Need to invert index b/c layer properties is generated in reverse order
+	TArray<TSharedPtr<FSortedParamData>> AssetChildren = LayerProperties[LayerProperties.Num() - 1 - InIndex]->Children;
+	UMaterialInstanceConstant* MaterialToUpdate = nullptr;
+	int32 ParameterIndex = InIndex;
+	if (InAssociation == EMaterialParameterAssociation::LayerParameter)
+	{
+		MaterialToUpdate = MaterialEditorInstance->StoredLayerPreviews[ParameterIndex];
+	}
+	if (InAssociation == EMaterialParameterAssociation::BlendParameter)
+	{
+		if (bAlterBlendIndex)
+		{
+			ParameterIndex--;
+		}
+		MaterialToUpdate = MaterialEditorInstance->StoredBlendPreviews[ParameterIndex];
+	}
+
+	TArray<FEditorParameterGroup> ParameterGroups;
+	for (TSharedPtr<FSortedParamData> AssetChild : AssetChildren)
+	{
+		for (TSharedPtr<FSortedParamData> Group : AssetChild->Children)
+		{
+			if (Group->ParameterInfo.Association == InAssociation)
+			{
+				FEditorParameterGroup DuplicatedGroup = FEditorParameterGroup();
+				DuplicatedGroup.GroupAssociation = Group->Group.GroupAssociation;
+				DuplicatedGroup.GroupName = Group->Group.GroupName;
+				DuplicatedGroup.GroupSortPriority = Group->Group.GroupSortPriority;
+				for (UDEditorParameterValue* Parameter : Group->Group.Parameters)
+				{
+					if (Parameter->ParameterInfo.Index == ParameterIndex)
+					{
+						DuplicatedGroup.Parameters.Add(Parameter);
+					}
+				}
+				ParameterGroups.Add(DuplicatedGroup);
+			}
+		}
+
+
+
+	}
+	if (MaterialToUpdate != nullptr)
+	{
+		FMaterialPropertyHelpers::TransitionAndCopyParameters(MaterialToUpdate, ParameterGroups, true);
+	}
+}
+
 void SMaterialLayersFunctionsMaterialTree::ShowSubParameters(TSharedPtr<FSortedParamData> ParentParameter)
 {
 	for (FUnsortedParamData Property : NonLayerProperties)
@@ -2641,6 +2858,7 @@ void SMaterialLayersFunctionsMaterialTree::ShowSubParameters(TSharedPtr<FSortedP
 			{
 				ParentParameter->Children.Add(GroupProperty);
 			}
+
 			TSharedPtr<FSortedParamData> ChildProperty(new FSortedParamData());
 			ChildProperty->StackDataType = EStackDataType::Property;
 			ChildProperty->Parameter = Parameter;
@@ -2664,23 +2882,110 @@ void SMaterialLayersFunctionsMaterialTree::ShowSubParameters(TSharedPtr<FSortedP
 					ParamChildProperty->ParameterHandle = ParamChildProperty->ParameterNode->CreatePropertyHandle();
 					ParamChildProperty->ParameterInfo.Index = Parameter->ParameterInfo.Index;
 					ParamChildProperty->ParameterInfo.Association = Parameter->ParameterInfo.Association;
+					ParamChildProperty->Parameter = ChildProperty->Parameter;
 					ChildProperty->Children.Add(ParamChildProperty);
 				}
 			}
-
 			for (TSharedPtr<struct FSortedParamData> GroupChild : ParentParameter->Children)
 			{
 				if (GroupChild->Group.GroupName == Property.ParameterGroup.GroupName
 					&& GroupChild->ParameterInfo.Association == ChildProperty->ParameterInfo.Association
-					&&  GroupChild->ParameterInfo.Index == ChildProperty->ParameterInfo.Index)
+					&& GroupChild->ParameterInfo.Index == ChildProperty->ParameterInfo.Index)
 				{
 					GroupChild->Children.Add(ChildProperty);
 				}
 			}
+
 		}
 	}
 }
 
+void SMaterialLayersFunctionsMaterialWrapper::Refresh()
+{
+	LayerParameter = nullptr;
+	TSharedPtr<SHorizontalBox> HeaderBox;
+	NestedTree->CreateGroupsWidget();
+	LayerParameter = NestedTree->FunctionParameter;
+
+	if (LayerParameter != nullptr)
+	{
+		this->ChildSlot
+			[
+				SNew(SBorder)
+				.BorderImage(FEditorStyle::GetBrush("MaterialInstanceEditor.LayersBorder"))
+				.Padding(FMargin(4.0f))
+				[
+					SNew(SVerticalBox)
+					+ SVerticalBox::Slot()
+					.AutoHeight()
+					[
+						SAssignNew(HeaderBox, SHorizontalBox)
+						+ SHorizontalBox::Slot()
+						.Padding(FMargin(3.0f, 1.0f))
+						.HAlign(HAlign_Left)
+						.AutoWidth()
+						.VAlign(VAlign_Center)
+						[
+							SNew(STextBlock)
+							.Text(FText::FromName(NestedTree->LayersFunctionsParameterName))
+							.TextStyle(FEditorStyle::Get(), "LargeText")
+						]
+					]
+					+ SVerticalBox::Slot()
+					.Padding(FMargin(3.0f, 0.0f))
+					[
+						NestedTree.ToSharedRef()
+					]
+				]
+			];
+		HeaderBox->AddSlot()
+			.FillWidth(1.0f)
+			[
+				SNullWidget::NullWidget
+			];
+	}
+	else
+	{
+		this->ChildSlot
+			[
+				SNew(SBorder)
+				.BorderImage(FEditorStyle::GetBrush("MaterialInstanceEditor.StackBody"))
+				.Padding(FMargin(4.0f))
+				[
+					SNew(STextBlock)
+					.Text(LOCTEXT("AddLayerParameterPrompt", "Add a Material Attribute Layers parameter to see it here."))
+				]
+			];
+	}
+}
+
+
+void SMaterialLayersFunctionsMaterialWrapper::Construct(const FArguments& InArgs)
+{
+	TSharedPtr<IPropertyRowGenerator> InGenerator = InArgs._InGenerator;
+	Generator = InGenerator;
+
+	NestedTree = SNew(SMaterialLayersFunctionsMaterialTree)
+		.InMaterialEditorInstance(InArgs._InMaterialEditorInstance)
+		.InWrapper(this);
+
+	LayerParameter = NestedTree->FunctionParameter;
+
+	MaterialEditorInstance = InArgs._InMaterialEditorInstance;
+	Refresh();
+
+}
+
+void SMaterialLayersFunctionsMaterialWrapper::SetEditorInstance(UMaterialEditorPreviewParameters* InMaterialEditorInstance)
+{
+	NestedTree->MaterialEditorInstance = InMaterialEditorInstance;
+	Refresh();
+}
+
 
+TSharedPtr<class IPropertyRowGenerator> SMaterialLayersFunctionsMaterialWrapper::GetGenerator()
+{
+	return Generator.Pin();
+}
 
 #undef LOCTEXT_NAMESPACE
\ No newline at end of file
diff --git a/Engine/Source/Editor/MaterialEditor/Private/SMaterialLayersFunctionsTree.h b/Engine/Source/Editor/MaterialEditor/Private/SMaterialLayersFunctionsTree.h
--- a/Engine/Source/Editor/MaterialEditor/Private/SMaterialLayersFunctionsTree.h
+++ b/Engine/Source/Editor/MaterialEditor/Private/SMaterialLayersFunctionsTree.h
@@ -191,68 +191,115 @@ private:
 };
 
 class UMaterialEditorPreviewParameters;
+class SMaterialLayersFunctionsMaterialTree;
+//////////////////////////////// Material version
+
+class SMaterialLayersFunctionsMaterialTreeItem : public STableRow< TSharedPtr<FSortedParamData> >
+{
+public:
+
+	SLATE_BEGIN_ARGS(SMaterialLayersFunctionsMaterialTreeItem)
+		: _StackParameterData(nullptr),
+		_MaterialEditorInstance(nullptr)
+	{}
+
+	/** The item content. */
+	SLATE_ARGUMENT(TSharedPtr<FSortedParamData>, StackParameterData)
+	SLATE_ARGUMENT(UMaterialEditorPreviewParameters*, MaterialEditorInstance)
+	SLATE_ARGUMENT(SMaterialLayersFunctionsMaterialTree*, InTree)
+	SLATE_END_ARGS()
+
+	FMaterialTreeColumnSizeData ColumnSizeData;
+	bool bIsBeingDragged;
+
+private:
+	bool bIsHoveredDragTarget;
+	FString GetCurvePath(class UDEditorScalarParameterValue* Parameter) const;
+	const FSlateBrush* GetBorderImage() const;
+
+public:
+
+	void RefreshOnRowChange(const FAssetData& AssetData, SMaterialLayersFunctionsInstanceTree* InTree);
+	FText GetLayerName(SMaterialLayersFunctionsMaterialTree* InTree, int32 Counter) const;
+
+	/**
+	* Construct the widget
+	*
+	* @param InArgs   A declaration from which to construct the widget
+	*/
+	void Construct(const FArguments& InArgs, const TSharedRef<STableViewBase>& InOwnerTableView);
+
+	/** The node info to build the tree view row from. */
+	TSharedPtr<FSortedParamData> StackParameterData;
+
+	SMaterialLayersFunctionsMaterialTree* Tree;
+
+	UMaterialEditorPreviewParameters* MaterialEditorInstance;
+
+	FString GetInstancePath(SMaterialLayersFunctionsMaterialTree* InTree) const;
+};
 
 class SMaterialLayersFunctionsMaterialWrapper : public SCompoundWidget
 {
 public:
 	SLATE_BEGIN_ARGS(SMaterialLayersFunctionsMaterialWrapper)
-		: _InMaterialEditorInstance(nullptr),
-		_InShowHiddenDelegate()
+		: _InMaterialEditorInstance(nullptr)
+		, _InGenerator()
 	{}
 
 	SLATE_ARGUMENT(UMaterialEditorPreviewParameters*, InMaterialEditorInstance)
-	SLATE_ARGUMENT(FGetShowHiddenParameters, InShowHiddenDelegate)
-
+	SLATE_ARGUMENT(TSharedPtr<class IPropertyRowGenerator>, InGenerator)
 	SLATE_END_ARGS()
 	void Refresh();
 	void Construct(const FArguments& InArgs);
 	void SetEditorInstance(UMaterialEditorPreviewParameters* InMaterialEditorInstance);
-
+	TSharedPtr<class IPropertyRowGenerator> GetGenerator();
 	class UDEditorParameterValue* LayerParameter;
-	UMaterialEditorPreviewParameters* MaterialEditorInstance;
+	class UMaterialEditorPreviewParameters* MaterialEditorInstance;
 	TSharedPtr<class SMaterialLayersFunctionsMaterialTree> NestedTree;
-};
 
+private:
+	TWeakPtr<class IPropertyRowGenerator> Generator;
+};
 
 class SMaterialLayersFunctionsMaterialTree : public STreeView<TSharedPtr<FSortedParamData>>
 {
 	friend class SMaterialLayersFunctionsMaterialTreeItem;
 public:
 	SLATE_BEGIN_ARGS(SMaterialLayersFunctionsMaterialTree)
-		: _InMaterialEditorInstance(nullptr),
-		_InShowHiddenDelegate()
+		: _InMaterialEditorInstance(nullptr)
 	{}
 
 	SLATE_ARGUMENT(UMaterialEditorPreviewParameters*, InMaterialEditorInstance)
-	SLATE_ARGUMENT(FGetShowHiddenParameters, InShowHiddenDelegate)
-
-	SLATE_END_ARGS()
-
-	/** Constructs this widget with InArgs */
-	void Construct(const FArguments& InArgs);
-	TSharedRef< ITableRow > OnGenerateRowMaterialLayersFunctionsTreeView(TSharedPtr<FSortedParamData> Item, const TSharedRef< STableViewBase >& OwnerTable);
-	void OnGetChildrenMaterialLayersFunctionsTreeView(TSharedPtr<FSortedParamData> InParent, TArray< TSharedPtr<FSortedParamData> >& OutChildren);
-	void OnExpansionChanged(TSharedPtr<FSortedParamData> Item, bool bIsExpanded);
-	void SetParentsExpansionState();
-
-	float OnGetLeftColumnWidth() const { return 1.0f - ColumnWidth; }
-	float OnGetRightColumnWidth() const { return ColumnWidth; }
-	void OnSetColumnWidth(float InWidth) { ColumnWidth = InWidth; }
-	FName LayersFunctionsParameterName;
-	class UDEditorParameterValue* FunctionParameter;
-	struct FMaterialLayersFunctions* FunctionInstance;
-	TSharedPtr<IPropertyHandle> FunctionInstanceHandle;
-	TSharedPtr<class FAssetThumbnailPool> GetTreeThumbnailPool();
-
-	/** Object that stores all of the possible parameters we can edit */
-	UMaterialEditorPreviewParameters* MaterialEditorInstance;
-
-	/** Delegate to call to determine if hidden parameters should be shown */
-	FGetShowHiddenParameters ShowHiddenDelegate;
-
-	/** Builds the custom parameter groups category */
-	void CreateGroupsWidget();
-
+		SLATE_ARGUMENT(SMaterialLayersFunctionsMaterialWrapper*, InWrapper)
+		SLATE_END_ARGS()
+
+		/** Constructs this widget with InArgs */
+		void Construct(const FArguments& InArgs);
+		TSharedRef< ITableRow > OnGenerateRowMaterialLayersFunctionsTreeView(TSharedPtr<FSortedParamData> Item, const TSharedRef< STableViewBase >& OwnerTable);
+		void OnGetChildrenMaterialLayersFunctionsTreeView(TSharedPtr<FSortedParamData> InParent, TArray< TSharedPtr<FSortedParamData> >& OutChildren);
+		void OnExpansionChanged(TSharedPtr<FSortedParamData> Item, bool bIsExpanded);
+		void SetParentsExpansionState();
+		float OnGetLeftColumnWidth() const { return 1.0f - ColumnWidth; }
+		float OnGetRightColumnWidth() const { return ColumnWidth; }
+		void OnSetColumnWidth(float InWidth) { ColumnWidth = InWidth; }
+		void ShowHiddenValues(bool& bShowHiddenParameters) { bShowHiddenParameters = true; }
+		FName LayersFunctionsParameterName;
+		class UDEditorParameterValue* FunctionParameter;
+		struct FMaterialLayersFunctions* FunctionInstance;
+		TSharedPtr<IPropertyHandle> FunctionInstanceHandle;
+		TSharedPtr<class FAssetThumbnailPool> GetTreeThumbnailPool();
+
+		/** Object that stores all of the possible parameters we can edit */
+		UMaterialEditorPreviewParameters* MaterialEditorInstance;
+
+		/** Builds the custom parameter groups category */
+		void CreateGroupsWidget();
+
+		SMaterialLayersFunctionsMaterialWrapper* GetWrapper() { return Wrapper; }
+
+		TSharedRef<SWidget> CreateThumbnailWidget(EMaterialParameterAssociation InAssociation, int32 InIndex, float InThumbnailSize);
+		void UpdateThumbnailMaterial(TEnumAsByte<EMaterialParameterAssociation> InAssociation, int32 InIndex, bool bAlterBlendIndex = false);
 protected:
 
 	void ShowSubParameters(TSharedPtr<FSortedParamData> ParentParameter);
@@ -265,6 +312,5 @@ private:
 	/** The actual width of the right column.  The left column is 1-ColumnWidth */
 	float ColumnWidth;
 
-	TSharedPtr<class IPropertyRowGenerator> Generator;
-
-};
\ No newline at end of file
+	SMaterialLayersFunctionsMaterialWrapper* Wrapper;
+};
diff --git a/Engine/Source/Editor/MaterialEditor/Private/SMaterialParametersOverviewWidget.cpp b/Engine/Source/Editor/MaterialEditor/Private/SMaterialParametersOverviewWidget.cpp
--- a/Engine/Source/Editor/MaterialEditor/Private/SMaterialParametersOverviewWidget.cpp
+++ b/Engine/Source/Editor/MaterialEditor/Private/SMaterialParametersOverviewWidget.cpp
@@ -334,7 +334,10 @@ void SMaterialParametersOverviewTreeItem::Construct(const FArguments& InArgs, co
 		{
 			FNodeWidgets StoredNodeWidgets = Node.CreateNodeWidgets();
 			TSharedRef<SWidget> StoredRightSideWidget = StoredNodeWidgets.ValueWidget.ToSharedRef();
-			StackParameterData->ParameterNode->CreatePropertyHandle()->MarkResetToDefaultCustomized(true);
+			if (TSharedPtr<IPropertyHandle> PropertyHandle = StackParameterData->ParameterNode->CreatePropertyHandle())
+			{
+				PropertyHandle->MarkResetToDefaultCustomized(true);
+			}
 			FDetailWidgetRow& CustomWidget = Row.CustomWidget();
 			CustomWidget
 			.FilterString(NameOverride)
@@ -596,7 +599,7 @@ void SMaterialParametersOverviewTree::SetParentsExpansionState()
 
 TSharedPtr<class FAssetThumbnailPool> SMaterialParametersOverviewTree::GetTreeThumbnailPool()
 {
-	return Generator->GetGeneratedThumbnailPool();
+	return GetOwner().Pin()->GetGenerator()->GetGeneratedThumbnailPool();
 }
 
 void SMaterialParametersOverviewTree::CreateGroupsWidget()
@@ -605,23 +608,12 @@ void SMaterialParametersOverviewTree::CreateGroupsWidget()
 	MaterialEditorInstance->RegenerateArrays();
 	UnsortedParameters.Empty();
 	SortedParameters.Empty();
-	FPropertyEditorModule& Module = FModuleManager::LoadModuleChecked<FPropertyEditorModule>("PropertyEditor");
-	if (!Generator.IsValid())
-	{
-		FPropertyRowGeneratorArgs Args;
-		Generator = Module.CreatePropertyRowGenerator(Args);
 
-		TArray<UObject*> Objects;
-		Objects.Add(MaterialEditorInstance);
-		Generator->SetObjects(Objects);
-	}
-	else
+	const TArray<TSharedRef<IDetailTreeNode>> TestData = GetOwner().Pin()->GetGenerator()->GetRootTreeNodes();
+	if (TestData.Num() == 0)
 	{
-		TArray<UObject*> Objects;
-		Objects.Add(MaterialEditorInstance);
-		Generator->SetObjects(Objects);
+		return;
 	}
-	const TArray<TSharedRef<IDetailTreeNode>> TestData = Generator->GetRootTreeNodes();
 	TSharedPtr<IDetailTreeNode> Category = TestData[0];
 	TSharedPtr<IDetailTreeNode> ParameterGroups;
 	TArray<TSharedRef<IDetailTreeNode>> Children;
@@ -652,48 +644,51 @@ void SMaterialParametersOverviewTree::CreateGroupsWidget()
 		const FEditorParameterGroup* ParameterGroupPtr = reinterpret_cast<FEditorParameterGroup*>(*GroupIt);
 		const FEditorParameterGroup& ParameterGroup = *ParameterGroupPtr;
 
-		for (int32 ParamIdx = 0; ParamIdx < ParameterGroup.Parameters.Num(); ParamIdx++)
+		// Don't create or show the material layer parameter info in this UI
+		if (ParameterGroup.GroupName != FMaterialPropertyHelpers::LayerParamName)
 		{
-			UDEditorParameterValue* Parameter = ParameterGroup.Parameters[ParamIdx];
-
-			TSharedPtr<IPropertyHandle> ParametersArrayProperty = ChildHandle->GetChildHandle("Parameters");
-			TSharedPtr<IPropertyHandle> ParameterProperty = ParametersArrayProperty->GetChildHandle(ParamIdx);
-			TSharedPtr<IPropertyHandle> ParameterValueProperty = ParameterProperty->GetChildHandle("ParameterValue");
-
+			for (int32 ParamIdx = 0; ParamIdx < ParameterGroup.Parameters.Num(); ParamIdx++)
 			{
-				FUnsortedParamData NonLayerProperty;
-				UDEditorScalarParameterValue* ScalarParam = Cast<UDEditorScalarParameterValue>(Parameter);
-				UDEditorVectorParameterValue* VectorParam = Cast<UDEditorVectorParameterValue>(Parameter);
+				UDEditorParameterValue* Parameter = ParameterGroup.Parameters[ParamIdx];
 
-				if (ScalarParam && ScalarParam->SliderMax > ScalarParam->SliderMin)
-				{
-					ParameterValueProperty->SetInstanceMetaData("UIMin", FString::Printf(TEXT("%f"), ScalarParam->SliderMin));
-					ParameterValueProperty->SetInstanceMetaData("UIMax", FString::Printf(TEXT("%f"), ScalarParam->SliderMax));
-				}
+				TSharedPtr<IPropertyHandle> ParametersArrayProperty = ChildHandle->GetChildHandle("Parameters");
+				TSharedPtr<IPropertyHandle> ParameterProperty = ParametersArrayProperty->GetChildHandle(ParamIdx);
+				TSharedPtr<IPropertyHandle> ParameterValueProperty = ParameterProperty->GetChildHandle("ParameterValue");
 
-				if (VectorParam)
 				{
-					static const FName Red("R");
-					static const FName Green("G");
-					static const FName Blue("B");
-					static const FName Alpha("A");
-					if (!VectorParam->ChannelNames.R.IsEmpty())
-					{
-						ParameterProperty->GetChildHandle(Red)->SetPropertyDisplayName(VectorParam->ChannelNames.R);
-					}
-					if (!VectorParam->ChannelNames.G.IsEmpty())
-					{
-						ParameterProperty->GetChildHandle(Green)->SetPropertyDisplayName(VectorParam->ChannelNames.G);
-					}
-					if (!VectorParam->ChannelNames.B.IsEmpty())
+					FUnsortedParamData NonLayerProperty;
+					UDEditorScalarParameterValue* ScalarParam = Cast<UDEditorScalarParameterValue>(Parameter);
+					UDEditorVectorParameterValue* VectorParam = Cast<UDEditorVectorParameterValue>(Parameter);
+
+					if (ScalarParam && ScalarParam->SliderMax > ScalarParam->SliderMin)
 					{
-						ParameterProperty->GetChildHandle(Blue)->SetPropertyDisplayName(VectorParam->ChannelNames.B);
+						ParameterValueProperty->SetInstanceMetaData("UIMin", FString::Printf(TEXT("%f"), ScalarParam->SliderMin));
+						ParameterValueProperty->SetInstanceMetaData("UIMax", FString::Printf(TEXT("%f"), ScalarParam->SliderMax));
 					}
-					if (!VectorParam->ChannelNames.A.IsEmpty())
+
+					if (VectorParam)
 					{
-						ParameterProperty->GetChildHandle(Alpha)->SetPropertyDisplayName(VectorParam->ChannelNames.A);
+						static const FName Red("R");
+						static const FName Green("G");
+						static const FName Blue("B");
+						static const FName Alpha("A");
+						if (!VectorParam->ChannelNames.R.IsEmpty())
+						{
+							ParameterProperty->GetChildHandle(Red)->SetPropertyDisplayName(VectorParam->ChannelNames.R);
+						}
+						if (!VectorParam->ChannelNames.G.IsEmpty())
+						{
+							ParameterProperty->GetChildHandle(Green)->SetPropertyDisplayName(VectorParam->ChannelNames.G);
+						}
+						if (!VectorParam->ChannelNames.B.IsEmpty())
+						{
+							ParameterProperty->GetChildHandle(Blue)->SetPropertyDisplayName(VectorParam->ChannelNames.B);
+						}
+						if (!VectorParam->ChannelNames.A.IsEmpty())
+						{
+							ParameterProperty->GetChildHandle(Alpha)->SetPropertyDisplayName(VectorParam->ChannelNames.A);
+						}
 					}
-				}
 
 				NonLayerProperty.Parameter = Parameter;
 				NonLayerProperty.ParameterGroup = ParameterGroup;
@@ -706,7 +701,7 @@ void SMaterialParametersOverviewTree::CreateGroupsWidget()
 	}
 
 	checkf(UnsortedParameters.Num() == DeferredSearches.Num(), TEXT("Internal inconsistency: number of node searches does not match the number of properties"));
-	TArray<TSharedPtr<IDetailTreeNode>> DeferredResults = Generator->FindTreeNodes(DeferredSearches);
+	TArray<TSharedPtr<IDetailTreeNode>> DeferredResults = GetOwner().Pin()->GetGenerator()->FindTreeNode(ParameterValueProperty);->FindTreeNodes(DeferredSearches);
 	checkf(UnsortedParameters.Num() == DeferredResults.Num(), TEXT("Internal inconsistency: number of node search results does not match the number of properties"));
 
 	for (int Idx = 0, NumUnsorted = UnsortedParameters.Num(); Idx < NumUnsorted; ++Idx)
@@ -950,6 +945,8 @@ void SMaterialParametersOverviewPanel::Refresh()
 void SMaterialParametersOverviewPanel::Construct(const FArguments& InArgs)
 {
 	ExternalScrollbar = SNew(SScrollBar);
+	TSharedPtr<IPropertyRowGenerator> InGenerator = InArgs._InGenerator;
+	Generator = InGenerator;
 
 	NestedTree = SNew(SMaterialParametersOverviewTree)
 		.InMaterialEditorInstance(InArgs._InMaterialEditorInstance)
@@ -969,7 +966,7 @@ void SMaterialParametersOverviewPanel::UpdateEditorInstance(UMaterialEditorPrevi
 
 TSharedPtr<class IPropertyRowGenerator> SMaterialParametersOverviewPanel::GetGenerator()
 {
-	 return NestedTree->GetGenerator();
+	 return Generator.Pin();
 }
 
-#undef LOCTEXT_NAMESPACE
\ No newline at end of file
+#undef LOCTEXT_NAMESPACE
diff --git a/Engine/Source/Editor/MaterialEditor/Private/SMaterialParametersOverviewWidget.h b/Engine/Source/Editor/MaterialEditor/Private/SMaterialParametersOverviewWidget.h
--- a/Engine/Source/Editor/MaterialEditor/Private/SMaterialParametersOverviewWidget.h
+++ b/Engine/Source/Editor/MaterialEditor/Private/SMaterialParametersOverviewWidget.h
@@ -68,9 +68,11 @@ class SMaterialParametersOverviewPanel : public SCompoundWidget
 public:
 	SLATE_BEGIN_ARGS(SMaterialParametersOverviewPanel)
 		: _InMaterialEditorInstance(nullptr)
+		, _InGenerator()
 	{}
 
 	SLATE_ARGUMENT(UMaterialEditorPreviewParameters*, InMaterialEditorInstance)
+	SLATE_ARGUMENT(TSharedPtr<class IPropertyRowGenerator>, InGenerator)
 	SLATE_END_ARGS()
 	void Refresh();
 	void Construct(const FArguments& InArgs);
@@ -90,6 +92,7 @@ private:
 	TSharedPtr<class SMaterialParametersOverviewTree> NestedTree;
 
 	TSharedPtr<class SScrollBar> ExternalScrollbar;
+	TWeakPtr<class IPropertyRowGenerator> Generator;
 };
 
 // ********* SMaterialParametersOverviewTree *******
@@ -129,9 +132,6 @@ public:
 	void CreateGroupsWidget();
 
 	TWeakPtr<SMaterialParametersOverviewPanel> GetOwner() { return Owner; }
-
-	TSharedPtr<class IPropertyRowGenerator> GetGenerator() { return Generator; }
-
 	bool HasAnyParameters() const { return bHasAnyParameters; }
 
 protected:
@@ -149,8 +149,6 @@ private:
 
 	TWeakPtr<SMaterialParametersOverviewPanel> Owner;
 
-	TSharedPtr<class IPropertyRowGenerator> Generator;
-
 	bool bHasAnyParameters;
 
 };
diff --git a/Engine/Source/Editor/MaterialEditor/Public/MaterialPropertyHelpers.h b/Engine/Source/Editor/MaterialEditor/Public/MaterialPropertyHelpers.h
--- a/Engine/Source/Editor/MaterialEditor/Public/MaterialPropertyHelpers.h
+++ b/Engine/Source/Editor/MaterialEditor/Public/MaterialPropertyHelpers.h
@@ -41,6 +41,7 @@ struct MATERIALEDITOR_API FSortedParamData
 public:
 	EStackDataType StackDataType;
 
+	UPROPERTY(Transient)
 	class UDEditorParameterValue* Parameter;
 
 	FName PropertyName;
@@ -58,8 +59,11 @@ public:
 	FString NodeKey;
 };
 
+USTRUCT()
 struct FUnsortedParamData
 {
+	GENERATED_USTRUCT_BODY()
+	UPROPERTY(Transient)
 	class UDEditorParameterValue* Parameter;
 	FEditorParameterGroup ParameterGroup;
 	TSharedPtr<class IDetailTreeNode> ParameterNode;
diff --git a/Engine/Source/Editor/PropertyEditor/Private/DetailItemNode.cpp b/Engine/Source/Editor/PropertyEditor/Private/DetailItemNode.cpp
--- a/Engine/Source/Editor/PropertyEditor/Private/DetailItemNode.cpp
+++ b/Engine/Source/Editor/PropertyEditor/Private/DetailItemNode.cpp
@@ -96,7 +96,7 @@ EDetailNodeType FDetailItemNode::GetNodeType() const
 
 TSharedPtr<IPropertyHandle> FDetailItemNode::CreatePropertyHandle() const
 {
-	if (Customization.HasPropertyNode())
+	if (Customization.HasPropertyNode() && ParentCategory.IsValid())
 	{
 		return ParentCategory.Pin()->GetParentLayoutImpl().GetPropertyHandle(Customization.GetPropertyNode());
 	}
diff --git a/Engine/Source/Editor/PropertyEditor/Private/PropertyHandleImpl.cpp b/Engine/Source/Editor/PropertyEditor/Private/PropertyHandleImpl.cpp
--- a/Engine/Source/Editor/PropertyEditor/Private/PropertyHandleImpl.cpp
+++ b/Engine/Source/Editor/PropertyEditor/Private/PropertyHandleImpl.cpp
@@ -2551,8 +2551,10 @@ FProperty* FPropertyHandleBase::GetMetaDataProperty() const
 		if( PropertyNode->GetArrayIndex() != INDEX_NONE )
 		{
 			TSharedPtr<FPropertyNode> ParentNode = PropertyNode->GetParentNodeSharedPtr();
-			check(ParentNode.IsValid());
-			MetaDataProperty = ParentNode->GetProperty();
+			if (ParentNode)
+			{
+				MetaDataProperty = ParentNode->GetProperty();
+			}
 		}
 	}
 
diff --git a/Engine/Source/Editor/UnrealEd/Classes/MaterialEditor/MaterialEditorPreviewParameters.h b/Engine/Source/Editor/UnrealEd/Classes/MaterialEditor/MaterialEditorPreviewParameters.h
--- a/Engine/Source/Editor/UnrealEd/Classes/MaterialEditor/MaterialEditorPreviewParameters.h
+++ b/Engine/Source/Editor/UnrealEd/Classes/MaterialEditor/MaterialEditorPreviewParameters.h
@@ -39,6 +39,14 @@ public:
 	UPROPERTY()
 	class UMaterial* OriginalMaterial;
 
+#if WITH_EDITORONLY_DATA
+	UPROPERTY()
+	TArray<class UMaterialInstanceConstant*> StoredLayerPreviews;
+
+	UPROPERTY()
+	TArray<class UMaterialInstanceConstant*> StoredBlendPreviews;
+#endif
+
 	//~ Begin UObject Interface.
 	virtual void PostEditChangeProperty(FPropertyChangedEvent& PropertyChangedEvent) override;
 #if WITH_EDITOR
