Fixing incorrectly merged line

diff --git a/Engine/Source/Editor/MaterialEditor/Private/SMaterialParametersOverviewWidget.cpp b/Engine/Source/Editor/MaterialEditor/Private/SMaterialParametersOverviewWidget.cpp
--- a/Engine/Source/Editor/MaterialEditor/Private/SMaterialParametersOverviewWidget.cpp
+++ b/Engine/Source/Editor/MaterialEditor/Private/SMaterialParametersOverviewWidget.cpp
@@ -701,7 +701,7 @@ void SMaterialParametersOverviewTree::CreateGroupsWidget()
 	}
 
 	checkf(UnsortedParameters.Num() == DeferredSearches.Num(), TEXT("Internal inconsistency: number of node searches does not match the number of properties"));
-	TArray<TSharedPtr<IDetailTreeNode>> DeferredResults = GetOwner().Pin()->GetGenerator()->FindTreeNode(ParameterValueProperty);->FindTreeNodes(DeferredSearches);
+	TArray<TSharedPtr<IDetailTreeNode>> DeferredResults = GetOwner().Pin()->GetGenerator()->FindTreeNodes(DeferredSearches);
 	checkf(UnsortedParameters.Num() == DeferredResults.Num(), TEXT("Internal inconsistency: number of node search results does not match the number of properties"));
 
 	for (int Idx = 0, NumUnsorted = UnsortedParameters.Num(); Idx < NumUnsorted; ++Idx)
