Fixing a compile issue in material editor
#rnx

diff --git a/Engine/Source/Editor/MaterialEditor/Private/SMaterialParametersOverviewWidget.cpp b/Engine/Source/Editor/MaterialEditor/Private/SMaterialParametersOverviewWidget.cpp
--- a/Engine/Source/Editor/MaterialEditor/Private/SMaterialParametersOverviewWidget.cpp
+++ b/Engine/Source/Editor/MaterialEditor/Private/SMaterialParametersOverviewWidget.cpp
@@ -696,6 +696,7 @@ void SMaterialParametersOverviewTree::CreateGroupsWidget()
 
 				DeferredSearches.Add(ParameterValueProperty);
 				UnsortedParameters.Add(NonLayerProperty);
+				}
 			}
 		}
 	}
