Fixing infinitely refreshing UI in material editor
#rnx

diff --git a/Engine/Source/Editor/MaterialEditor/Private/MaterialEditor.cpp b/Engine/Source/Editor/MaterialEditor/Private/MaterialEditor.cpp
--- a/Engine/Source/Editor/MaterialEditor/Private/MaterialEditor.cpp
+++ b/Engine/Source/Editor/MaterialEditor/Private/MaterialEditor.cpp
@@ -733,12 +733,6 @@ void FMaterialEditor::InitMaterialEditor( const EToolkitMode::Type Mode, const T
 		Objects.Add(MaterialEditorInstance);
 		Generator->SetObjects(Objects);
 	}
-	// Update the parameter list now that the material has been fully initialized
-	MaterialParametersOverviewWidget->UpdateEditorInstance(MaterialEditorInstance);
-	if (MaterialLayersFunctionsInstance.IsValid())
-	{
-		MaterialLayersFunctionsInstance->SetEditorInstance(MaterialEditorInstance);
-	}
 
 	if (OriginalMaterial->bUsedAsSpecialEngineMaterial)
 	{
@@ -1064,7 +1058,6 @@ void FMaterialEditor::OnFinishedChangingParametersFromOverview(const FPropertyCh
 	{
 		RefreshExpressionPreviews(true);
 		RefreshPreviewViewport();
-		UpdatePreviewMaterial();
 		TArray<TWeakObjectPtr<UObject>> SelectedObjects = MaterialDetailsView->GetSelectedObjects();
 		MaterialDetailsView->SetObjects(SelectedObjects, true);
 		Material->MarkPackageDirty();
@@ -1074,10 +1067,10 @@ void FMaterialEditor::OnFinishedChangingParametersFromOverview(const FPropertyCh
 
 void FMaterialEditor::GeneratorRowsRefreshed()
 {
-	MaterialParametersOverviewWidget->Refresh();
+	MaterialParametersOverviewWidget->Refresh(false);
 	if (MaterialLayersFunctionsInstance)
 	{
-		MaterialLayersFunctionsInstance->Refresh();
+		MaterialLayersFunctionsInstance->Refresh(false);
 	}
 }
 
diff --git a/Engine/Source/Editor/MaterialEditor/Private/SMaterialLayersFunctionsTree.cpp b/Engine/Source/Editor/MaterialEditor/Private/SMaterialLayersFunctionsTree.cpp
--- a/Engine/Source/Editor/MaterialEditor/Private/SMaterialLayersFunctionsTree.cpp
+++ b/Engine/Source/Editor/MaterialEditor/Private/SMaterialLayersFunctionsTree.cpp
@@ -1788,11 +1788,11 @@ const FSlateBrush* SMaterialLayersFunctionsMaterialTreeItem::GetBorderImage() co
 	}
 }
 
-void SMaterialLayersFunctionsMaterialTreeItem::RefreshOnRowChange(const FAssetData& AssetData, SMaterialLayersFunctionsInstanceTree* InTree)
+void SMaterialLayersFunctionsMaterialTreeItem::RefreshOnRowChange(const FAssetData& AssetData, SMaterialLayersFunctionsMaterialTree* InTree)
 {
-	if (SMaterialLayersFunctionsInstanceWrapper* Wrapper = InTree->GetWrapper())
+	if (SMaterialLayersFunctionsMaterialWrapper* Wrapper = InTree->GetWrapper())
 	{
-		InTree->CreateGroupsWidget();
+		InTree->CreateGroupsWidget(false);
 	}
 }
 
@@ -2466,7 +2466,7 @@ void SMaterialLayersFunctionsMaterialTree::Construct(const FArguments& InArgs)
 	ColumnWidth = 0.5f;
 	MaterialEditorInstance = InArgs._InMaterialEditorInstance;
 	Wrapper = InArgs._InWrapper;
-	CreateGroupsWidget();
+	CreateGroupsWidget(true);
 
 #ifdef WITH_EDITOR
 	//Fixup for adding new bool arrays to the class
@@ -2560,10 +2560,13 @@ TSharedPtr<class FAssetThumbnailPool> SMaterialLayersFunctionsMaterialTree::GetT
 	return Wrapper->GetGenerator()->GetGeneratedThumbnailPool();
 }
 
-void SMaterialLayersFunctionsMaterialTree::CreateGroupsWidget()
+void SMaterialLayersFunctionsMaterialTree::CreateGroupsWidget(const bool bRegenerateArray)
 {
 	check(MaterialEditorInstance);
-	MaterialEditorInstance->RegenerateArrays();
+	if (bRegenerateArray)
+	{
+		MaterialEditorInstance->RegenerateArrays();
+	}
 	NonLayerProperties.Empty();
 	LayerProperties.Empty();
 	FunctionParameter = nullptr;
@@ -2900,11 +2903,11 @@ void SMaterialLayersFunctionsMaterialTree::ShowSubParameters(TSharedPtr<FSortedP
 	}
 }
 
-void SMaterialLayersFunctionsMaterialWrapper::Refresh()
+void SMaterialLayersFunctionsMaterialWrapper::Refresh(const bool bRegenerateArray)
 {
 	LayerParameter = nullptr;
 	TSharedPtr<SHorizontalBox> HeaderBox;
-	NestedTree->CreateGroupsWidget();
+	NestedTree->CreateGroupsWidget(bRegenerateArray);
 	LayerParameter = NestedTree->FunctionParameter;
 
 	if (LayerParameter != nullptr)
diff --git a/Engine/Source/Editor/MaterialEditor/Private/SMaterialLayersFunctionsTree.h b/Engine/Source/Editor/MaterialEditor/Private/SMaterialLayersFunctionsTree.h
--- a/Engine/Source/Editor/MaterialEditor/Private/SMaterialLayersFunctionsTree.h
+++ b/Engine/Source/Editor/MaterialEditor/Private/SMaterialLayersFunctionsTree.h
@@ -219,7 +219,7 @@ private:
 
 public:
 
-	void RefreshOnRowChange(const FAssetData& AssetData, SMaterialLayersFunctionsInstanceTree* InTree);
+	void RefreshOnRowChange(const FAssetData& AssetData, SMaterialLayersFunctionsMaterialTree* InTree);
 	FText GetLayerName(SMaterialLayersFunctionsMaterialTree* InTree, int32 Counter) const;
 
 	/**
@@ -250,7 +250,7 @@ public:
 	SLATE_ARGUMENT(UMaterialEditorPreviewParameters*, InMaterialEditorInstance)
 	SLATE_ARGUMENT(TSharedPtr<class IPropertyRowGenerator>, InGenerator)
 	SLATE_END_ARGS()
-	void Refresh();
+	void Refresh(const bool bRegenerateArray = false);
 	void Construct(const FArguments& InArgs);
 	void SetEditorInstance(UMaterialEditorPreviewParameters* InMaterialEditorInstance);
 	TSharedPtr<class IPropertyRowGenerator> GetGenerator();
@@ -294,7 +294,7 @@ public:
 		UMaterialEditorPreviewParameters* MaterialEditorInstance;
 
 		/** Builds the custom parameter groups category */
-		void CreateGroupsWidget();
+		void CreateGroupsWidget(const bool bRegenerateArray);
 
 		SMaterialLayersFunctionsMaterialWrapper* GetWrapper() { return Wrapper; }
 
diff --git a/Engine/Source/Editor/MaterialEditor/Private/SMaterialParametersOverviewWidget.cpp b/Engine/Source/Editor/MaterialEditor/Private/SMaterialParametersOverviewWidget.cpp
--- a/Engine/Source/Editor/MaterialEditor/Private/SMaterialParametersOverviewWidget.cpp
+++ b/Engine/Source/Editor/MaterialEditor/Private/SMaterialParametersOverviewWidget.cpp
@@ -67,7 +67,7 @@ void SMaterialParametersOverviewTreeItem::RefreshOnRowChange(const FAssetData& A
 {
 	if (InTree.IsValid())
 	{
-		InTree->CreateGroupsWidget();
+		InTree->CreateGroupsWidget(false);
 	}
 }
 
@@ -528,7 +528,7 @@ void SMaterialParametersOverviewTree::Construct(const FArguments& InArgs)
 	ColumnWidth = 0.5f;
 	MaterialEditorInstance = InArgs._InMaterialEditorInstance;
 	Owner = InArgs._InOwner;
-	CreateGroupsWidget();
+	CreateGroupsWidget(true);
 
 
 	STreeView<TSharedPtr<FSortedParamData>>::Construct(
@@ -602,10 +602,13 @@ TSharedPtr<class FAssetThumbnailPool> SMaterialParametersOverviewTree::GetTreeTh
 	return GetOwner().Pin()->GetGenerator()->GetGeneratedThumbnailPool();
 }
 
-void SMaterialParametersOverviewTree::CreateGroupsWidget()
+void SMaterialParametersOverviewTree::CreateGroupsWidget(const bool bRegenerateArray)
 {
 	check(MaterialEditorInstance);
-	MaterialEditorInstance->RegenerateArrays();
+	if (bRegenerateArray)
+	{
+		MaterialEditorInstance->RegenerateArrays();
+	}
 	UnsortedParameters.Empty();
 	SortedParameters.Empty();
 
@@ -636,7 +639,6 @@ void SMaterialParametersOverviewTree::CreateGroupsWidget()
 	TArray<TSharedPtr<IPropertyHandle>> DeferredSearches;
 	for (int32 GroupIdx = 0; GroupIdx < Children.Num(); ++GroupIdx)
 	{
-		bHasAnyParameters = true;
 		TArray<void*> GroupPtrs;
 		TSharedPtr<IPropertyHandle> ChildHandle = Children[GroupIdx]->CreatePropertyHandle();
 		ChildHandle->AccessRawData(GroupPtrs);
@@ -650,45 +652,47 @@ void SMaterialParametersOverviewTree::CreateGroupsWidget()
 			for (int32 ParamIdx = 0; ParamIdx < ParameterGroup.Parameters.Num(); ParamIdx++)
 			{
 				UDEditorParameterValue* Parameter = ParameterGroup.Parameters[ParamIdx];
-
-				TSharedPtr<IPropertyHandle> ParametersArrayProperty = ChildHandle->GetChildHandle("Parameters");
-				TSharedPtr<IPropertyHandle> ParameterProperty = ParametersArrayProperty->GetChildHandle(ParamIdx);
-				TSharedPtr<IPropertyHandle> ParameterValueProperty = ParameterProperty->GetChildHandle("ParameterValue");
-
+				if(Parameter->ParameterInfo.Association == EMaterialParameterAssociation::GlobalParameter)
 				{
-					FUnsortedParamData NonLayerProperty;
-					UDEditorScalarParameterValue* ScalarParam = Cast<UDEditorScalarParameterValue>(Parameter);
-					UDEditorVectorParameterValue* VectorParam = Cast<UDEditorVectorParameterValue>(Parameter);
+					bHasAnyParameters = true;
+					TSharedPtr<IPropertyHandle> ParametersArrayProperty = ChildHandle->GetChildHandle("Parameters");
+					TSharedPtr<IPropertyHandle> ParameterProperty = ParametersArrayProperty->GetChildHandle(ParamIdx);
+					TSharedPtr<IPropertyHandle> ParameterValueProperty = ParameterProperty->GetChildHandle("ParameterValue");
 
-					if (ScalarParam && ScalarParam->SliderMax > ScalarParam->SliderMin)
 					{
-						ParameterValueProperty->SetInstanceMetaData("UIMin", FString::Printf(TEXT("%f"), ScalarParam->SliderMin));
-						ParameterValueProperty->SetInstanceMetaData("UIMax", FString::Printf(TEXT("%f"), ScalarParam->SliderMax));
-					}
+						FUnsortedParamData NonLayerProperty;
+						UDEditorScalarParameterValue* ScalarParam = Cast<UDEditorScalarParameterValue>(Parameter);
+						UDEditorVectorParameterValue* VectorParam = Cast<UDEditorVectorParameterValue>(Parameter);
 
-					if (VectorParam)
-					{
-						static const FName Red("R");
-						static const FName Green("G");
-						static const FName Blue("B");
-						static const FName Alpha("A");
-						if (!VectorParam->ChannelNames.R.IsEmpty())
-						{
-							ParameterProperty->GetChildHandle(Red)->SetPropertyDisplayName(VectorParam->ChannelNames.R);
-						}
-						if (!VectorParam->ChannelNames.G.IsEmpty())
+						if (ScalarParam && ScalarParam->SliderMax > ScalarParam->SliderMin)
 						{
-							ParameterProperty->GetChildHandle(Green)->SetPropertyDisplayName(VectorParam->ChannelNames.G);
+							ParameterValueProperty->SetInstanceMetaData("UIMin", FString::Printf(TEXT("%f"), ScalarParam->SliderMin));
+							ParameterValueProperty->SetInstanceMetaData("UIMax", FString::Printf(TEXT("%f"), ScalarParam->SliderMax));
 						}
-						if (!VectorParam->ChannelNames.B.IsEmpty())
-						{
-							ParameterProperty->GetChildHandle(Blue)->SetPropertyDisplayName(VectorParam->ChannelNames.B);
-						}
-						if (!VectorParam->ChannelNames.A.IsEmpty())
+
+						if (VectorParam)
 						{
-							ParameterProperty->GetChildHandle(Alpha)->SetPropertyDisplayName(VectorParam->ChannelNames.A);
+							static const FName Red("R");
+							static const FName Green("G");
+							static const FName Blue("B");
+							static const FName Alpha("A");
+							if (!VectorParam->ChannelNames.R.IsEmpty())
+							{
+								ParameterProperty->GetChildHandle(Red)->SetPropertyDisplayName(VectorParam->ChannelNames.R);
+							}
+							if (!VectorParam->ChannelNames.G.IsEmpty())
+							{
+								ParameterProperty->GetChildHandle(Green)->SetPropertyDisplayName(VectorParam->ChannelNames.G);
+							}
+							if (!VectorParam->ChannelNames.B.IsEmpty())
+							{
+								ParameterProperty->GetChildHandle(Blue)->SetPropertyDisplayName(VectorParam->ChannelNames.B);
+							}
+							if (!VectorParam->ChannelNames.A.IsEmpty())
+							{
+								ParameterProperty->GetChildHandle(Alpha)->SetPropertyDisplayName(VectorParam->ChannelNames.A);
+							}
 						}
-					}
 
 					NonLayerProperty.Parameter = Parameter;
 					NonLayerProperty.ParameterGroup = ParameterGroup;
@@ -805,10 +809,10 @@ int32 SMaterialParametersOverviewPanel::GetPanelIndex() const
 	return NestedTree && NestedTree->HasAnyParameters() ? 1 : 0;
 }
 
-void SMaterialParametersOverviewPanel::Refresh()
+void SMaterialParametersOverviewPanel::Refresh(const bool bRegenerateArray)
 {
 	TSharedPtr<SHorizontalBox> HeaderBox;
-	NestedTree->CreateGroupsWidget();
+	NestedTree->CreateGroupsWidget(bRegenerateArray);
 
 	FOnClicked 	OnChildButtonClicked = FOnClicked();
 	if (MaterialEditorInstance->OriginalFunction)
diff --git a/Engine/Source/Editor/MaterialEditor/Private/SMaterialParametersOverviewWidget.h b/Engine/Source/Editor/MaterialEditor/Private/SMaterialParametersOverviewWidget.h
--- a/Engine/Source/Editor/MaterialEditor/Private/SMaterialParametersOverviewWidget.h
+++ b/Engine/Source/Editor/MaterialEditor/Private/SMaterialParametersOverviewWidget.h
@@ -74,7 +74,7 @@ public:
 	SLATE_ARGUMENT(UMaterialEditorPreviewParameters*, InMaterialEditorInstance)
 	SLATE_ARGUMENT(TSharedPtr<class IPropertyRowGenerator>, InGenerator)
 	SLATE_END_ARGS()
-	void Refresh();
+	void Refresh(const bool bRegenerateArray = true);
 	void Construct(const FArguments& InArgs);
 	void UpdateEditorInstance(UMaterialEditorPreviewParameters* InMaterialEditorInstance);
 	TSharedPtr<class IPropertyRowGenerator> GetGenerator();
@@ -129,7 +129,7 @@ public:
 	UMaterialEditorPreviewParameters* MaterialEditorInstance;
 
 	/** Builds the custom parameter groups category */
-	void CreateGroupsWidget();
+	void CreateGroupsWidget(const bool bRegenerateArray);
 
 	TWeakPtr<SMaterialParametersOverviewPanel> GetOwner() { return Owner; }
 	bool HasAnyParameters() const { return bHasAnyParameters; }
