Fix for extra brace
#rnx

diff --git a/Engine/Source/Editor/MaterialEditor/Private/SMaterialParametersOverviewWidget.cpp b/Engine/Source/Editor/MaterialEditor/Private/SMaterialParametersOverviewWidget.cpp
--- a/Engine/Source/Editor/MaterialEditor/Private/SMaterialParametersOverviewWidget.cpp
+++ b/Engine/Source/Editor/MaterialEditor/Private/SMaterialParametersOverviewWidget.cpp
@@ -659,40 +659,40 @@ void SMaterialParametersOverviewTree::CreateGroupsWidget(const bool bRegenerateA
 					TSharedPtr<IPropertyHandle> ParameterProperty = ParametersArrayProperty->GetChildHandle(ParamIdx);
 					TSharedPtr<IPropertyHandle> ParameterValueProperty = ParameterProperty->GetChildHandle("ParameterValue");
 
+
+					FUnsortedParamData NonLayerProperty;
+					UDEditorScalarParameterValue* ScalarParam = Cast<UDEditorScalarParameterValue>(Parameter);
+					UDEditorVectorParameterValue* VectorParam = Cast<UDEditorVectorParameterValue>(Parameter);
+
+					if (ScalarParam && ScalarParam->SliderMax > ScalarParam->SliderMin)
 					{
-						FUnsortedParamData NonLayerProperty;
-						UDEditorScalarParameterValue* ScalarParam = Cast<UDEditorScalarParameterValue>(Parameter);
-						UDEditorVectorParameterValue* VectorParam = Cast<UDEditorVectorParameterValue>(Parameter);
+						ParameterValueProperty->SetInstanceMetaData("UIMin", FString::Printf(TEXT("%f"), ScalarParam->SliderMin));
+						ParameterValueProperty->SetInstanceMetaData("UIMax", FString::Printf(TEXT("%f"), ScalarParam->SliderMax));
+					}
 
-						if (ScalarParam && ScalarParam->SliderMax > ScalarParam->SliderMin)
+					if (VectorParam)
+					{
+						static const FName Red("R");
+						static const FName Green("G");
+						static const FName Blue("B");
+						static const FName Alpha("A");
+						if (!VectorParam->ChannelNames.R.IsEmpty())
 						{
-							ParameterValueProperty->SetInstanceMetaData("UIMin", FString::Printf(TEXT("%f"), ScalarParam->SliderMin));
-							ParameterValueProperty->SetInstanceMetaData("UIMax", FString::Printf(TEXT("%f"), ScalarParam->SliderMax));
+							ParameterProperty->GetChildHandle(Red)->SetPropertyDisplayName(VectorParam->ChannelNames.R);
 						}
-
-						if (VectorParam)
+						if (!VectorParam->ChannelNames.G.IsEmpty())
 						{
-							static const FName Red("R");
-							static const FName Green("G");
-							static const FName Blue("B");
-							static const FName Alpha("A");
-							if (!VectorParam->ChannelNames.R.IsEmpty())
-							{
-								ParameterProperty->GetChildHandle(Red)->SetPropertyDisplayName(VectorParam->ChannelNames.R);
-							}
-							if (!VectorParam->ChannelNames.G.IsEmpty())
-							{
-								ParameterProperty->GetChildHandle(Green)->SetPropertyDisplayName(VectorParam->ChannelNames.G);
-							}
-							if (!VectorParam->ChannelNames.B.IsEmpty())
-							{
-								ParameterProperty->GetChildHandle(Blue)->SetPropertyDisplayName(VectorParam->ChannelNames.B);
-							}
-							if (!VectorParam->ChannelNames.A.IsEmpty())
-							{
-								ParameterProperty->GetChildHandle(Alpha)->SetPropertyDisplayName(VectorParam->ChannelNames.A);
-							}
+							ParameterProperty->GetChildHandle(Green)->SetPropertyDisplayName(VectorParam->ChannelNames.G);
+						}
+						if (!VectorParam->ChannelNames.B.IsEmpty())
+						{
+							ParameterProperty->GetChildHandle(Blue)->SetPropertyDisplayName(VectorParam->ChannelNames.B);
 						}
+						if (!VectorParam->ChannelNames.A.IsEmpty())
+						{
+							ParameterProperty->GetChildHandle(Alpha)->SetPropertyDisplayName(VectorParam->ChannelNames.A);
+						}
+					}
 
 					NonLayerProperty.Parameter = Parameter;
 					NonLayerProperty.ParameterGroup = ParameterGroup;
