Removing redundant UI refreshes

diff --git a/Engine/Source/Editor/MaterialEditor/Private/MaterialEditor.cpp b/Engine/Source/Editor/MaterialEditor/Private/MaterialEditor.cpp
--- a/Engine/Source/Editor/MaterialEditor/Private/MaterialEditor.cpp
+++ b/Engine/Source/Editor/MaterialEditor/Private/MaterialEditor.cpp
@@ -1046,7 +1046,6 @@ void FMaterialEditor::OnFinishedChangingProperties(const FPropertyChangedEvent&
 		}
 		RefreshPreviewViewport();
 		UpdatePreviewMaterial();
-		MaterialParametersOverviewWidget->UpdateEditorInstance(MaterialEditorInstance);
 		MaterialCustomPrimitiveDataWidget->UpdateEditorInstance(MaterialEditorInstance);
 	}
 }
@@ -1850,13 +1849,6 @@ void FMaterialEditor::UpdatePreviewMaterial( bool bForce )
 		// Null out the expression preview material so they can be GC'ed
 		ExpressionPreviewMaterial = NULL;
 	}
-	if (MaterialEditorInstance != nullptr)
-	{
-		if (MaterialLayersFunctionsInstance.IsValid())
-		{
-			MaterialLayersFunctionsInstance->SetEditorInstance(MaterialEditorInstance);
-		}
-	}
 
 	MaterialStatsManager->SetMaterial(bStatsFromPreviewMaterial ? Material : OriginalMaterial);
 	MaterialStatsManager->SignalMaterialChanged();
@@ -3138,11 +3130,6 @@ void FMaterialEditor::OnConvertObjects()
 		{
 			GraphEditor->SetNodeSelection(*NodeIter, true);
 		}
-
-		if (MaterialEditorInstance != nullptr)
-		{
-			MaterialParametersOverviewWidget->UpdateEditorInstance(MaterialEditorInstance);
-		}
 	}
 }
 
@@ -3508,7 +3495,6 @@ void FMaterialEditor::OnPromoteToParameter()
 	}
 	if (MaterialEditorInstance != nullptr)
 	{
-		MaterialParametersOverviewWidget->UpdateEditorInstance(MaterialEditorInstance);
 		MaterialCustomPrimitiveDataWidget->UpdateEditorInstance(MaterialEditorInstance);
 	}
 }
@@ -4472,7 +4458,6 @@ void FMaterialEditor::UpdateMaterialAfterGraphChange()
 	RegenerateCodeView();
 	RefreshExpressionPreviews();
 	SetMaterialDirty();
-	MaterialParametersOverviewWidget->UpdateEditorInstance(MaterialEditorInstance);
 	if (bHideUnrelatedNodes && !bLockNodeFadeState && bSelectRegularNode)
 	{
 		GraphEditor->ResetAllNodesUnrelatedStates();
@@ -4522,7 +4507,6 @@ void FMaterialEditor::UndoGraphAction()
 	{
 		Material->BuildEditorParameterList();
 	}
-	MaterialParametersOverviewWidget->UpdateEditorInstance(MaterialEditorInstance);
 }
 
 void FMaterialEditor::RedoGraphAction()
@@ -4536,7 +4520,6 @@ void FMaterialEditor::RedoGraphAction()
 	{
 		Material->BuildEditorParameterList();
 	}
-	MaterialParametersOverviewWidget->UpdateEditorInstance(MaterialEditorInstance);
 }
 
 void FMaterialEditor::OnAlignTop()
@@ -4646,11 +4629,6 @@ void FMaterialEditor::NotifyPostChange( const FPropertyChangedEvent& PropertyCha
 
 	if ( PropertyThatChanged )
 	{
-		if (MaterialLayersFunctionsInstance.IsValid())
-		{
-			MaterialLayersFunctionsInstance->SetEditorInstance(MaterialEditorInstance);
-		}
-		MaterialParametersOverviewWidget->UpdateEditorInstance(MaterialEditorInstance);
 		MaterialCustomPrimitiveDataWidget->UpdateEditorInstance(MaterialEditorInstance);
 
 		const FName NameOfPropertyThatChanged( *PropertyThatChanged->GetName() );
@@ -4910,7 +4888,6 @@ void FMaterialEditor::OnColorPickerCommitted(FLinearColor LinearColor)
 	}
 
 	RefreshExpressionPreviews();
-	MaterialParametersOverviewWidget->UpdateEditorInstance(MaterialEditorInstance);
 	MaterialCustomPrimitiveDataWidget->UpdateEditorInstance(MaterialEditorInstance);
 }
 
@@ -5403,7 +5380,6 @@ void FMaterialEditor::OnNodeTitleCommitted(const FText& NewText, ETextCommit::Ty
 		const FScopedTransaction Transaction( LOCTEXT( "RenameNode", "Rename Node" ) );
 		NodeBeingChanged->Modify();
 		NodeBeingChanged->OnRenameNode(NewText.ToString());
-		MaterialParametersOverviewWidget->UpdateEditorInstance(MaterialEditorInstance);
 		MaterialCustomPrimitiveDataWidget->UpdateEditorInstance(MaterialEditorInstance);
 	}
 }
