Moving the texture parameter's value property to the top of the details panel

diff --git a/Engine/Source/Editor/MaterialEditor/Private/MaterialEditorDetailCustomization.cpp b/Engine/Source/Editor/MaterialEditor/Private/MaterialEditorDetailCustomization.cpp
--- a/Engine/Source/Editor/MaterialEditor/Private/MaterialEditorDetailCustomization.cpp
+++ b/Engine/Source/Editor/MaterialEditor/Private/MaterialEditorDetailCustomization.cpp
@@ -250,6 +250,8 @@ void FMaterialExpressionParameterDetails::CustomizeDetails( IDetailLayoutBuilder
 
 		if (TextureParameter)
 		{
+			IDetailCategoryBuilder& TextureCategory = DetailLayout.EditCategory(TEXT("MaterialExpressionTextureBase"), FText::GetEmpty(), ECategoryPriority::Important);
+
 			TSharedPtr<IPropertyHandle> ChannelHandle = DetailLayout.GetProperty("ChannelNames", UMaterialExpressionTextureSampleParameter::StaticClass());
 			TSharedPtr<IPropertyHandle> ValueHandle = DetailLayout.GetProperty("Texture", UMaterialExpressionTextureBase::StaticClass());
 			if (TextureParameter->GetOutputType(0) != MCT_Texture)
