Update cached material data before updating expression previews

diff --git a/Engine/Source/Editor/MaterialEditor/Private/MaterialEditor.cpp b/Engine/Source/Editor/MaterialEditor/Private/MaterialEditor.cpp
--- a/Engine/Source/Editor/MaterialEditor/Private/MaterialEditor.cpp
+++ b/Engine/Source/Editor/MaterialEditor/Private/MaterialEditor.cpp
@@ -4744,6 +4744,8 @@ void FMaterialEditor::RefreshExpressionPreviews(bool bForceRefreshAll /*= false*
 {
 	const FScopedBusyCursor BusyCursor;
 
+	Material->UpdateCachedExpressionData();
+
 	if ( bAlwaysRefreshAllPreviews || bForceRefreshAll)
 	{
 		// we need to make sure the rendering thread isn't drawing these tiles
