Add UI to disconnect/reconnect MI layers from parent material

diff --git a/Engine/Source/Editor/MaterialEditor/Private/SMaterialLayersFunctionsTree.cpp b/Engine/Source/Editor/MaterialEditor/Private/SMaterialLayersFunctionsTree.cpp
--- a/Engine/Source/Editor/MaterialEditor/Private/SMaterialLayersFunctionsTree.cpp
+++ b/Engine/Source/Editor/MaterialEditor/Private/SMaterialLayersFunctionsTree.cpp
@@ -40,6 +40,7 @@
 #include "Curves/CurveLinearColorAtlas.h"
 #include "Subsystems/AssetEditorSubsystem.h"
 #include "Editor.h"
+#include "EditorFontGlyphs.h"
 
 
 #define LOCTEXT_NAMESPACE "MaterialLayerCustomization"
@@ -334,6 +335,24 @@ void SMaterialLayersFunctionsInstanceTreeItem::Construct(const FArguments& InArg
 					.Font(FEditorStyle::GetFontStyle(TEXT("MaterialEditor.Layers.EditableFontImportant")))
 					.ForegroundColor(FLinearColor::White)
 				];
+		}
+		else
+		{
+			HeaderRowWidget->AddSlot()
+				.VAlign(VAlign_Center)
+				.AutoWidth()
+				.Padding(5.0f)
+				[
+					SNew(STextBlock)
+					.Text(NameOverride)
+					.TextStyle(FEditorStyle::Get(), "NormalText.Important")
+				];
+		}
+
+		// Can only unlink/remove layers that aren't the base layer.
+		if (StackParameterData->ParameterInfo.Index != 0)
+		{
+			// Unlink UI
 			HeaderRowWidget->AddSlot()
 				.FillWidth(1.0f)
 				.VAlign(VAlign_Center)
@@ -343,21 +362,28 @@ void SMaterialLayersFunctionsInstanceTreeItem::Construct(const FArguments& InArg
 			HeaderRowWidget->AddSlot()
 				.AutoWidth()
 				.VAlign(VAlign_Center)
-				.Padding(0.0f, 0.0f, 5.0f, 0.0f)
+				.Padding(0.0f, 0.0f, 0.0f, 0.0f)
 				[
-					PropertyCustomizationHelpers::MakeClearButton(FSimpleDelegate::CreateSP(InArgs._InTree, &SMaterialLayersFunctionsInstanceTree::RemoveLayer, StackParameterData->ParameterInfo.Index))
+					SNew(SButton)
+					.ButtonStyle(FEditorStyle::Get(), "HoverHintOnly")
+					.OnClicked(Tree, &SMaterialLayersFunctionsInstanceTree::UnlinkLayer, StackParameterData->ParameterInfo.Index)
+					.Visibility(Tree, &SMaterialLayersFunctionsInstanceTree::GetUnlinkLayerVisibility, StackParameterData->ParameterInfo.Index)
+					.ToolTipText(LOCTEXT("UnlinkLayer", "Whether or not to unlink this layer/blend combination from the parent."))
+					.Content()
+					[
+						SNew(STextBlock)
+						.TextStyle(FEditorStyle::Get(), "ContentBrowser.TopBar.Font")
+						.Font(FEditorStyle::Get().GetFontStyle("FontAwesome.10"))
+						.Text(FEditorFontGlyphs::Chain_Broken) /*fa-filter*/
+					]
 				];
-		}
-		else
-		{
+
 			HeaderRowWidget->AddSlot()
-				.VAlign(VAlign_Center)
 				.AutoWidth()
-				.Padding(5.0f)
+				.VAlign(VAlign_Center)
+				.Padding(0.0f, 0.0f, 5.0f, 0.0f)
 				[
-					SNew(STextBlock)
-					.Text(NameOverride)
-					.TextStyle(FEditorStyle::Get(), "NormalText.Important")
+					PropertyCustomizationHelpers::MakeClearButton(FSimpleDelegate::CreateSP(InArgs._InTree, &SMaterialLayersFunctionsInstanceTree::RemoveLayer, StackParameterData->ParameterInfo.Index))
 				];
 		}
 		LeftSideWidget = HeaderRowWidget;
@@ -1190,6 +1216,45 @@ void SMaterialLayersFunctionsInstanceTree::RemoveLayer(int32 Index)
 	RequestTreeRefresh();
 }
 
+FReply SMaterialLayersFunctionsInstanceTree::UnlinkLayer(int32 Index)
+{
+	const FScopedTransaction Transaction(LOCTEXT("UnlinkLayer", "Unlink a layer from the parent"));
+	FunctionInstanceHandle->NotifyPreChange();
+	FunctionInstance->UnlinkLayerFromParent(Index);
+	FunctionInstanceHandle->NotifyPostChange();
+	CreateGroupsWidget();
+	return FReply::Handled();
+}
+
+FReply SMaterialLayersFunctionsInstanceTree::RelinkLayersToParent()
+{
+	const FScopedTransaction Transaction(LOCTEXT("RelinkLayersToParent", "Relink layers to parent"));
+	FunctionInstanceHandle->NotifyPreChange();
+	FunctionInstance->RelinkLayersToParent();
+	FunctionInstanceHandle->NotifyPostChange();
+	MaterialEditorInstance->RegenerateArrays();
+	CreateGroupsWidget();
+	return FReply::Handled();
+}
+
+EVisibility SMaterialLayersFunctionsInstanceTree::GetUnlinkLayerVisibility(int32 Index) const
+{
+	if (FunctionInstance->IsLayerLinkedToParent(Index))
+	{
+		return EVisibility::Visible;
+	}
+	return EVisibility::Collapsed;
+}
+
+EVisibility SMaterialLayersFunctionsInstanceTree::GetRelinkLayersToParentVisibility() const
+{
+	if (FunctionInstance->HasAnyUnlinkedLayers())
+	{
+		return EVisibility::Visible;
+	}
+	return EVisibility::Collapsed;
+}
+
 FReply SMaterialLayersFunctionsInstanceTree::ToggleLayerVisibility(int32 Index)
 {
 	if (!FSlateApplication::Get().GetModifierKeys().AreModifersDown(EModifierKey::Alt))
@@ -1617,6 +1682,8 @@ void SMaterialLayersFunctionsInstanceWrapper::Refresh()
 
 	if (LayerParameter != nullptr)
 	{
+		FOnClicked OnRelinkToParent = FOnClicked::CreateSP(NestedTree.ToSharedRef(), &SMaterialLayersFunctionsInstanceTree::RelinkLayersToParent);
+
 		this->ChildSlot
 			[
 				SNew(SBorder)
@@ -1661,6 +1728,36 @@ void SMaterialLayersFunctionsInstanceWrapper::Refresh()
 			[
 				SNullWidget::NullWidget
 			];
+		HeaderBox->AddSlot()
+			.AutoWidth()
+			.Padding(2.0f)
+			[
+				SNew(SButton)
+				.ButtonStyle(FEditorStyle::Get(), "FlatButton.DarkGrey")
+				.HAlign(HAlign_Center)
+				.OnClicked(OnRelinkToParent)
+				.ToolTipText(LOCTEXT("RelinkToParentLayers", "Relink to Parent Layers and Blends"))
+				.Visibility(NestedTree.Get(), &SMaterialLayersFunctionsInstanceTree::GetRelinkLayersToParentVisibility)
+				.Content()
+				[
+					SNew(SHorizontalBox)
+					+ SHorizontalBox::Slot()
+					.AutoWidth()
+					[
+						SNew(STextBlock)
+						.Font(FEditorStyle::Get().GetFontStyle("FontAwesome.10"))
+						.TextStyle(FEditorStyle::Get(), "NormalText.Important")
+						.Text(FEditorFontGlyphs::Link)
+					]
+					+ SHorizontalBox::Slot()
+					.AutoWidth()
+					[
+						SNew(STextBlock)
+						.TextStyle(FEditorStyle::Get(), "NormalText.Important")
+						.Text(FText::FromString(FString(TEXT(" Relink"))) /*fa-filter*/)
+					]
+				]
+			];
 		HeaderBox->AddSlot()
 				.AutoWidth()
 				.Padding(2.0f)
diff --git a/Engine/Source/Editor/MaterialEditor/Private/SMaterialLayersFunctionsTree.h b/Engine/Source/Editor/MaterialEditor/Private/SMaterialLayersFunctionsTree.h
--- a/Engine/Source/Editor/MaterialEditor/Private/SMaterialLayersFunctionsTree.h
+++ b/Engine/Source/Editor/MaterialEditor/Private/SMaterialLayersFunctionsTree.h
@@ -146,6 +146,10 @@ public:
 	void ResetAssetToDefault(TSharedPtr<IPropertyHandle> InHandle, TSharedPtr<FSortedParamData> InData);
 	void AddLayer();
 	void RemoveLayer(int32 Index);
+	FReply UnlinkLayer(int32 Index);
+	FReply RelinkLayersToParent();
+	EVisibility GetUnlinkLayerVisibility(int32 Index) const;
+	EVisibility GetRelinkLayersToParentVisibility() const;
 	FReply ToggleLayerVisibility(int32 Index);
 	TSharedPtr<class FAssetThumbnailPool> GetTreeThumbnailPool();
 
diff --git a/Engine/Source/Editor/MaterialEditor/Private/SMaterialParametersOverviewWidget.cpp b/Engine/Source/Editor/MaterialEditor/Private/SMaterialParametersOverviewWidget.cpp
--- a/Engine/Source/Editor/MaterialEditor/Private/SMaterialParametersOverviewWidget.cpp
+++ b/Engine/Source/Editor/MaterialEditor/Private/SMaterialParametersOverviewWidget.cpp
@@ -644,63 +644,69 @@ void SMaterialParametersOverviewTree::CreateGroupsWidget(const bool bRegenerateA
 		ChildHandle->AccessRawData(GroupPtrs);
 		auto GroupIt = GroupPtrs.CreateConstIterator();
 		const FEditorParameterGroup* ParameterGroupPtr = reinterpret_cast<FEditorParameterGroup*>(*GroupIt);
+		if (!ParameterGroupPtr)
+		{
+			continue;
+		}
+
 		const FEditorParameterGroup& ParameterGroup = *ParameterGroupPtr;
+		if (ParameterGroup.GroupName == FMaterialPropertyHelpers::LayerParamName)
+		{
+			// Don't create or show the material layer parameter info in this UI
+			continue;
+		}
 
-		// Don't create or show the material layer parameter info in this UI
-		if (ParameterGroup.GroupName != FMaterialPropertyHelpers::LayerParamName)
+		for (int32 ParamIdx = 0; ParamIdx < ParameterGroup.Parameters.Num(); ParamIdx++)
 		{
-			for (int32 ParamIdx = 0; ParamIdx < ParameterGroup.Parameters.Num(); ParamIdx++)
+			UDEditorParameterValue* Parameter = ParameterGroup.Parameters[ParamIdx];
+			if (Parameter->ParameterInfo.Association == EMaterialParameterAssociation::GlobalParameter)
 			{
-				UDEditorParameterValue* Parameter = ParameterGroup.Parameters[ParamIdx];
-				if(Parameter->ParameterInfo.Association == EMaterialParameterAssociation::GlobalParameter)
-				{
-					bHasAnyParameters = true;
-					TSharedPtr<IPropertyHandle> ParametersArrayProperty = ChildHandle->GetChildHandle("Parameters");
-					TSharedPtr<IPropertyHandle> ParameterProperty = ParametersArrayProperty->GetChildHandle(ParamIdx);
-					TSharedPtr<IPropertyHandle> ParameterValueProperty = ParameterProperty->GetChildHandle("ParameterValue");
+				bHasAnyParameters = true;
+				TSharedPtr<IPropertyHandle> ParametersArrayProperty = ChildHandle->GetChildHandle("Parameters");
+				TSharedPtr<IPropertyHandle> ParameterProperty = ParametersArrayProperty->GetChildHandle(ParamIdx);
+				TSharedPtr<IPropertyHandle> ParameterValueProperty = ParameterProperty->GetChildHandle("ParameterValue");
 
 
-					FUnsortedParamData NonLayerProperty;
-					UDEditorScalarParameterValue* ScalarParam = Cast<UDEditorScalarParameterValue>(Parameter);
-					UDEditorVectorParameterValue* VectorParam = Cast<UDEditorVectorParameterValue>(Parameter);
+				FUnsortedParamData NonLayerProperty;
+				UDEditorScalarParameterValue* ScalarParam = Cast<UDEditorScalarParameterValue>(Parameter);
+				UDEditorVectorParameterValue* VectorParam = Cast<UDEditorVectorParameterValue>(Parameter);
 
-					if (ScalarParam && ScalarParam->SliderMax > ScalarParam->SliderMin)
+				if (ScalarParam && ScalarParam->SliderMax > ScalarParam->SliderMin)
+				{
+					ParameterValueProperty->SetInstanceMetaData("UIMin", FString::Printf(TEXT("%f"), ScalarParam->SliderMin));
+					ParameterValueProperty->SetInstanceMetaData("UIMax", FString::Printf(TEXT("%f"), ScalarParam->SliderMax));
+				}
+
+				if (VectorParam)
+				{
+					static const FName Red("R");
+					static const FName Green("G");
+					static const FName Blue("B");
+					static const FName Alpha("A");
+					if (!VectorParam->ChannelNames.R.IsEmpty())
 					{
-						ParameterValueProperty->SetInstanceMetaData("UIMin", FString::Printf(TEXT("%f"), ScalarParam->SliderMin));
-						ParameterValueProperty->SetInstanceMetaData("UIMax", FString::Printf(TEXT("%f"), ScalarParam->SliderMax));
+						ParameterProperty->GetChildHandle(Red)->SetPropertyDisplayName(VectorParam->ChannelNames.R);
 					}
-
-					if (VectorParam)
+					if (!VectorParam->ChannelNames.G.IsEmpty())
+					{
+						ParameterProperty->GetChildHandle(Green)->SetPropertyDisplayName(VectorParam->ChannelNames.G);
+					}
+					if (!VectorParam->ChannelNames.B.IsEmpty())
+					{
+						ParameterProperty->GetChildHandle(Blue)->SetPropertyDisplayName(VectorParam->ChannelNames.B);
+					}
+					if (!VectorParam->ChannelNames.A.IsEmpty())
 					{
-						static const FName Red("R");
-						static const FName Green("G");
-						static const FName Blue("B");
-						static const FName Alpha("A");
-						if (!VectorParam->ChannelNames.R.IsEmpty())
-						{
-							ParameterProperty->GetChildHandle(Red)->SetPropertyDisplayName(VectorParam->ChannelNames.R);
-						}
-						if (!VectorParam->ChannelNames.G.IsEmpty())
-						{
-							ParameterProperty->GetChildHandle(Green)->SetPropertyDisplayName(VectorParam->ChannelNames.G);
-						}
-						if (!VectorParam->ChannelNames.B.IsEmpty())
-						{
-							ParameterProperty->GetChildHandle(Blue)->SetPropertyDisplayName(VectorParam->ChannelNames.B);
-						}
-						if (!VectorParam->ChannelNames.A.IsEmpty())
-						{
-							ParameterProperty->GetChildHandle(Alpha)->SetPropertyDisplayName(VectorParam->ChannelNames.A);
-						}
+						ParameterProperty->GetChildHandle(Alpha)->SetPropertyDisplayName(VectorParam->ChannelNames.A);
 					}
+				}
 
-					NonLayerProperty.Parameter = Parameter;
-					NonLayerProperty.ParameterGroup = ParameterGroup;
-					NonLayerProperty.UnsortedName = Parameter->ParameterInfo.Name;
+				NonLayerProperty.Parameter = Parameter;
+				NonLayerProperty.ParameterGroup = ParameterGroup;
+				NonLayerProperty.UnsortedName = Parameter->ParameterInfo.Name;
 
-					DeferredSearches.Add(ParameterValueProperty);
-					UnsortedParameters.Add(NonLayerProperty);
-				}
+				DeferredSearches.Add(ParameterValueProperty);
+				UnsortedParameters.Add(NonLayerProperty);
 			}
 		}
 	}
diff --git a/Engine/Source/Runtime/Engine/Classes/Materials/MaterialLayersFunctions.h b/Engine/Source/Runtime/Engine/Classes/Materials/MaterialLayersFunctions.h
--- a/Engine/Source/Runtime/Engine/Classes/Materials/MaterialLayersFunctions.h
+++ b/Engine/Source/Runtime/Engine/Classes/Materials/MaterialLayersFunctions.h
@@ -253,6 +253,13 @@ struct ENGINE_API FMaterialLayersFunctions
 
 	void MoveBlendedLayer(int32 SrcLayerIndex, int32 DstLayerIndex);
 
+#if WITH_EDITOR
+	void UnlinkLayerFromParent(int32 Index);
+	bool IsLayerLinkedToParent(int32 Index) const;
+	void RelinkLayersToParent();
+	bool HasAnyUnlinkedLayers() const;
+#endif // WITH_EDITOR
+
 	void ToggleBlendedLayerVisibility(int32 Index)
 	{
 		check(LayerStates.IsValidIndex(Index));
@@ -299,12 +306,22 @@ struct ENGINE_API FMaterialLayersFunctions
 
 	FORCEINLINE bool operator==(const FMaterialLayersFunctions& Other) const
 	{
-		return Layers == Other.Layers && Blends == Other.Blends && LayerStates == Other.LayerStates;
+		if (Layers != Other.Layers || Blends != Other.Blends || LayerStates != Other.LayerStates)
+		{
+			return false;
+		}
+#if WITH_EDITORONLY_DATA
+		if (ParentLayerGuids != Other.ParentLayerGuids || DeletedParentLayerGuids != Other.DeletedParentLayerGuids)
+		{
+			return false;
+		}
+#endif // WITH_EDITORONLY_DATA
+		return true;
 	}
 
 	FORCEINLINE bool operator!=(const FMaterialLayersFunctions& Other) const
 	{
-		return Layers != Other.Layers || Blends != Other.Blends || LayerStates != Other.LayerStates;
+		return !operator==(Other);
 	}
 };
 
diff --git a/Engine/Source/Runtime/Engine/Private/Materials/MaterialExpressions.cpp b/Engine/Source/Runtime/Engine/Private/Materials/MaterialExpressions.cpp
--- a/Engine/Source/Runtime/Engine/Private/Materials/MaterialExpressions.cpp
+++ b/Engine/Source/Runtime/Engine/Private/Materials/MaterialExpressions.cpp
@@ -12658,21 +12658,53 @@ void FMaterialLayersFunctions::MoveBlendedLayer(int32 SrcLayerIndex, int32 DstLa
 		RestrictToBlendRelatives.Swap(SrcLayerIndex - 1, DstLayerIndex - 1);
 		LayerGuids.Swap(SrcLayerIndex, DstLayerIndex);
 		// Disconnect layers from parent when they're moved at the instance level
-		if (ParentLayerGuids[SrcLayerIndex] != NoParentGuid)
-		{
-			check(ParentLayerGuids[SrcLayerIndex] != UninitializedParentGuid);
-			DeletedParentLayerGuids.Add(ParentLayerGuids[SrcLayerIndex]);
-			ParentLayerGuids[SrcLayerIndex] = NoParentGuid;
-		}
-		if (ParentLayerGuids[DstLayerIndex] != NoParentGuid)
+		UnlinkLayerFromParent(SrcLayerIndex);
+		UnlinkLayerFromParent(DstLayerIndex);
+#endif //WITH_EDITOR
+	}
+}
+
+#if WITH_EDITOR
+void FMaterialLayersFunctions::UnlinkLayerFromParent(int32 Index)
+{
+	check(Index > 0);
+	if (ParentLayerGuids[Index] != NoParentGuid)
+	{
+		check(ParentLayerGuids[Index] != UninitializedParentGuid);
+		DeletedParentLayerGuids.Add(ParentLayerGuids[Index]);
+		ParentLayerGuids[Index] = NoParentGuid;
+	}
+}
+
+bool FMaterialLayersFunctions::IsLayerLinkedToParent(int32 Index) const
+{
+	if (Index > 0 && Index < ParentLayerGuids.Num())
+	{
+		const FGuid ParentGuid = ParentLayerGuids[Index];
+		check(ParentGuid != BackgroundGuid);
+		return ParentGuid != UninitializedParentGuid && ParentGuid != NoParentGuid;
+	}
+	return false;
+}
+
+void FMaterialLayersFunctions::RelinkLayersToParent()
+{
+	for (const FGuid& ParentGuid : DeletedParentLayerGuids)
+	{
+		const int32 LayerIndex = LayerGuids.Find(ParentGuid);
+		if (LayerIndex != INDEX_NONE && ParentLayerGuids[LayerIndex] == NoParentGuid)
 		{
-			check(ParentLayerGuids[DstLayerIndex] != UninitializedParentGuid);
-			DeletedParentLayerGuids.Add(ParentLayerGuids[DstLayerIndex]);
-			ParentLayerGuids[DstLayerIndex] = NoParentGuid;
+			ParentLayerGuids[LayerIndex] = ParentGuid;
 		}
-#endif //WITH_EDITOR
 	}
+	DeletedParentLayerGuids.Empty();
+}
+
+bool FMaterialLayersFunctions::HasAnyUnlinkedLayers() const
+{
+	return DeletedParentLayerGuids.Num() > 0;
 }
+#endif // WITH_EDITOR
 
 #if WITH_EDITORONLY_DATA
 void FMaterialLayersFunctions::CopyGuidsToParent()
