Fix duplicate LOCTEXT

diff --git a/Engine/Source/Editor/MaterialEditor/Private/SMaterialLayersFunctionsTree.cpp b/Engine/Source/Editor/MaterialEditor/Private/SMaterialLayersFunctionsTree.cpp
--- a/Engine/Source/Editor/MaterialEditor/Private/SMaterialLayersFunctionsTree.cpp
+++ b/Engine/Source/Editor/MaterialEditor/Private/SMaterialLayersFunctionsTree.cpp
@@ -1218,7 +1218,7 @@ void SMaterialLayersFunctionsInstanceTree::RemoveLayer(int32 Index)
 
 FReply SMaterialLayersFunctionsInstanceTree::UnlinkLayer(int32 Index)
 {
-	const FScopedTransaction Transaction(LOCTEXT("UnlinkLayer", "Unlink a layer from the parent"));
+	const FScopedTransaction Transaction(LOCTEXT("UnlinkLayerFromParent", "Unlink a layer from the parent"));
 	FunctionInstanceHandle->NotifyPreChange();
 	FunctionInstance->UnlinkLayerFromParent(Index);
 	FunctionInstanceHandle->NotifyPostChange();
