Make some BP utility functions static so they can actually be called

diff --git a/Engine/Source/Editor/MaterialEditor/Public/MaterialEditingLibrary.h b/Engine/Source/Editor/MaterialEditor/Public/MaterialEditingLibrary.h
--- a/Engine/Source/Editor/MaterialEditor/Public/MaterialEditingLibrary.h
+++ b/Engine/Source/Editor/MaterialEditor/Public/MaterialEditingLibrary.h
@@ -250,7 +250,7 @@ public:
 	*	@return	Whether or not the parameter was found in this material
 	*/
 	UFUNCTION(BlueprintCallable, Category = "MaterialEditing")
-	bool GetScalarParameterSource(UMaterialInterface* Material, const FName ParameterName, FSoftObjectPath& ParameterSource);
+	static bool GetScalarParameterSource(UMaterialInterface* Material, const FName ParameterName, FSoftObjectPath& ParameterSource);
 
 	/**
 	*	Returns the path of the asset where the parameter originated, as well as true/false if it was found
@@ -260,7 +260,7 @@ public:
 	*	@return	Whether or not the parameter was found in this material
 	*/
 	UFUNCTION(BlueprintCallable, Category = "MaterialEditing")
-	bool GetVectorParameterSource(UMaterialInterface* Material, const FName ParameterName, FSoftObjectPath& ParameterSource);
+	static bool GetVectorParameterSource(UMaterialInterface* Material, const FName ParameterName, FSoftObjectPath& ParameterSource);
 
 	/**
 	*	Returns the path of the asset where the parameter originated, as well as true/false if it was found
@@ -270,7 +270,7 @@ public:
 	*	@return	Whether or not the parameter was found in this material
 	*/
 	UFUNCTION(BlueprintCallable, Category = "MaterialEditing")
-	bool GetTextureParameterSource(UMaterialInterface* Material, const FName ParameterName, FSoftObjectPath& ParameterSource);
+	static bool GetTextureParameterSource(UMaterialInterface* Material, const FName ParameterName, FSoftObjectPath& ParameterSource);
 
 	/**
 	*	Returns the path of the asset where the parameter originated, as well as true/false if it was found
@@ -280,5 +280,5 @@ public:
 	*	@return	Whether or not the parameter was found in this material
 	*/
 	UFUNCTION(BlueprintCallable, Category = "MaterialEditing")
-	bool GetStaticSwitchParameterSource(UMaterialInterface* Material, const FName ParameterName, FSoftObjectPath& ParameterSource);
+	static bool GetStaticSwitchParameterSource(UMaterialInterface* Material, const FName ParameterName, FSoftObjectPath& ParameterSource);
 };
\ No newline at end of file
