Removing unneeded material object regeneration and added more clear calls to the material editor

diff --git a/Engine/Source/Editor/MaterialEditor/Private/MaterialEditor.cpp b/Engine/Source/Editor/MaterialEditor/Private/MaterialEditor.cpp
--- a/Engine/Source/Editor/MaterialEditor/Private/MaterialEditor.cpp
+++ b/Engine/Source/Editor/MaterialEditor/Private/MaterialEditor.cpp
@@ -1067,10 +1067,10 @@ void FMaterialEditor::OnFinishedChangingParametersFromOverview(const FPropertyCh
 
 void FMaterialEditor::GeneratorRowsRefreshed()
 {
-	MaterialParametersOverviewWidget->Refresh(false);
+	MaterialParametersOverviewWidget->Refresh();
 	if (MaterialLayersFunctionsInstance)
 	{
-		MaterialLayersFunctionsInstance->Refresh(false);
+		MaterialLayersFunctionsInstance->Refresh();
 	}
 }
 
@@ -3113,6 +3113,10 @@ void FMaterialEditor::OnConvertObjects()
 							// Refresh the expression preview if we changed its properties after it was created
 							NewExpression->bNeedToUpdatePreview = true;
 							RefreshExpressionPreview( NewExpression, true );
+							if (MaterialEditorInstance)
+							{
+								MaterialEditorInstance->RegenerateArrays();
+							}
 						}
 
 						NodesToDelete.AddUnique(GraphNode);
@@ -4520,6 +4524,11 @@ void FMaterialEditor::RedoGraphAction()
 	{
 		Material->BuildEditorParameterList();
 	}
+
+	if (MaterialEditorInstance)
+	{
+		MaterialEditorInstance->RegenerateArrays();
+	}
 }
 
 void FMaterialEditor::OnAlignTop()
@@ -4612,6 +4621,11 @@ void FMaterialEditor::PostUndo(bool bSuccess)
 		GraphEditor->NotifyGraphChanged();
 		SetMaterialDirty();
 
+		if (MaterialEditorInstance)
+		{
+			MaterialEditorInstance->RegenerateArrays();
+		}
+
 		FSlateApplication::Get().DismissAllMenus();
 	}
 }
@@ -4720,6 +4734,11 @@ void FMaterialEditor::NotifyPostChange( const FPropertyChangedEvent& PropertyCha
 
 	Material->MarkPackageDirty();
 	SetMaterialDirty();
+
+	if (MaterialEditorInstance)
+	{
+		MaterialEditorInstance->RegenerateArrays();
+	}
 }
 
 void FMaterialEditor::ToggleCollapsed(UMaterialExpression* MaterialExpression)
@@ -5382,6 +5401,10 @@ void FMaterialEditor::OnNodeTitleCommitted(const FText& NewText, ETextCommit::Ty
 		const FScopedTransaction Transaction( LOCTEXT( "RenameNode", "Rename Node" ) );
 		NodeBeingChanged->Modify();
 		NodeBeingChanged->OnRenameNode(NewText.ToString());
+		if (MaterialEditorInstance)
+		{
+			MaterialEditorInstance->RegenerateArrays();
+		}
 		MaterialCustomPrimitiveDataWidget->UpdateEditorInstance(MaterialEditorInstance);
 	}
 }
diff --git a/Engine/Source/Editor/MaterialEditor/Private/SMaterialLayersFunctionsTree.cpp b/Engine/Source/Editor/MaterialEditor/Private/SMaterialLayersFunctionsTree.cpp
--- a/Engine/Source/Editor/MaterialEditor/Private/SMaterialLayersFunctionsTree.cpp
+++ b/Engine/Source/Editor/MaterialEditor/Private/SMaterialLayersFunctionsTree.cpp
@@ -1889,7 +1889,7 @@ void SMaterialLayersFunctionsMaterialTreeItem::RefreshOnRowChange(const FAssetDa
 {
 	if (SMaterialLayersFunctionsMaterialWrapper* Wrapper = InTree->GetWrapper())
 	{
-		InTree->CreateGroupsWidget(false);
+		InTree->CreateGroupsWidget();
 	}
 }
 
@@ -2563,7 +2563,7 @@ void SMaterialLayersFunctionsMaterialTree::Construct(const FArguments& InArgs)
 	ColumnWidth = 0.5f;
 	MaterialEditorInstance = InArgs._InMaterialEditorInstance;
 	Wrapper = InArgs._InWrapper;
-	CreateGroupsWidget(true);
+	CreateGroupsWidget();
 
 #ifdef WITH_EDITOR
 	//Fixup for adding new bool arrays to the class
@@ -2657,13 +2657,10 @@ TSharedPtr<class FAssetThumbnailPool> SMaterialLayersFunctionsMaterialTree::GetT
 	return Wrapper->GetGenerator()->GetGeneratedThumbnailPool();
 }
 
-void SMaterialLayersFunctionsMaterialTree::CreateGroupsWidget(const bool bRegenerateArray)
+void SMaterialLayersFunctionsMaterialTree::CreateGroupsWidget()
 {
 	check(MaterialEditorInstance);
-	if (bRegenerateArray)
-	{
-		MaterialEditorInstance->RegenerateArrays();
-	}
+
 	NonLayerProperties.Empty();
 	LayerProperties.Empty();
 	FunctionParameter = nullptr;
@@ -3000,11 +2997,11 @@ void SMaterialLayersFunctionsMaterialTree::ShowSubParameters(TSharedPtr<FSortedP
 	}
 }
 
-void SMaterialLayersFunctionsMaterialWrapper::Refresh(const bool bRegenerateArray)
+void SMaterialLayersFunctionsMaterialWrapper::Refresh()
 {
 	LayerParameter = nullptr;
 	TSharedPtr<SHorizontalBox> HeaderBox;
-	NestedTree->CreateGroupsWidget(bRegenerateArray);
+	NestedTree->CreateGroupsWidget();
 	LayerParameter = NestedTree->FunctionParameter;
 
 	if (LayerParameter != nullptr)
diff --git a/Engine/Source/Editor/MaterialEditor/Private/SMaterialLayersFunctionsTree.h b/Engine/Source/Editor/MaterialEditor/Private/SMaterialLayersFunctionsTree.h
--- a/Engine/Source/Editor/MaterialEditor/Private/SMaterialLayersFunctionsTree.h
+++ b/Engine/Source/Editor/MaterialEditor/Private/SMaterialLayersFunctionsTree.h
@@ -254,7 +254,7 @@ public:
 	SLATE_ARGUMENT(UMaterialEditorPreviewParameters*, InMaterialEditorInstance)
 	SLATE_ARGUMENT(TSharedPtr<class IPropertyRowGenerator>, InGenerator)
 	SLATE_END_ARGS()
-	void Refresh(const bool bRegenerateArray = false);
+	void Refresh();
 	void Construct(const FArguments& InArgs);
 	void SetEditorInstance(UMaterialEditorPreviewParameters* InMaterialEditorInstance);
 	TSharedPtr<class IPropertyRowGenerator> GetGenerator();
@@ -298,7 +298,7 @@ public:
 		UMaterialEditorPreviewParameters* MaterialEditorInstance;
 
 		/** Builds the custom parameter groups category */
-		void CreateGroupsWidget(const bool bRegenerateArray);
+		void CreateGroupsWidget();
 
 		SMaterialLayersFunctionsMaterialWrapper* GetWrapper() { return Wrapper; }
 
diff --git a/Engine/Source/Editor/MaterialEditor/Private/SMaterialParametersOverviewWidget.cpp b/Engine/Source/Editor/MaterialEditor/Private/SMaterialParametersOverviewWidget.cpp
--- a/Engine/Source/Editor/MaterialEditor/Private/SMaterialParametersOverviewWidget.cpp
+++ b/Engine/Source/Editor/MaterialEditor/Private/SMaterialParametersOverviewWidget.cpp
@@ -67,7 +67,7 @@ void SMaterialParametersOverviewTreeItem::RefreshOnRowChange(const FAssetData& A
 {
 	if (InTree.IsValid())
 	{
-		InTree->CreateGroupsWidget(false);
+		InTree->CreateGroupsWidget();
 	}
 }
 
@@ -528,7 +528,7 @@ void SMaterialParametersOverviewTree::Construct(const FArguments& InArgs)
 	ColumnWidth = 0.5f;
 	MaterialEditorInstance = InArgs._InMaterialEditorInstance;
 	Owner = InArgs._InOwner;
-	CreateGroupsWidget(true);
+	CreateGroupsWidget();
 
 
 	STreeView<TSharedPtr<FSortedParamData>>::Construct(
@@ -602,13 +602,9 @@ TSharedPtr<class FAssetThumbnailPool> SMaterialParametersOverviewTree::GetTreeTh
 	return GetOwner().Pin()->GetGenerator()->GetGeneratedThumbnailPool();
 }
 
-void SMaterialParametersOverviewTree::CreateGroupsWidget(const bool bRegenerateArray)
+void SMaterialParametersOverviewTree::CreateGroupsWidget()
 {
 	check(MaterialEditorInstance);
-	if (bRegenerateArray)
-	{
-		MaterialEditorInstance->RegenerateArrays();
-	}
 	UnsortedParameters.Empty();
 	SortedParameters.Empty();
 
@@ -815,10 +811,10 @@ int32 SMaterialParametersOverviewPanel::GetPanelIndex() const
 	return NestedTree && NestedTree->HasAnyParameters() ? 1 : 0;
 }
 
-void SMaterialParametersOverviewPanel::Refresh(const bool bRegenerateArray)
+void SMaterialParametersOverviewPanel::Refresh()
 {
 	TSharedPtr<SHorizontalBox> HeaderBox;
-	NestedTree->CreateGroupsWidget(bRegenerateArray);
+	NestedTree->CreateGroupsWidget();
 
 	FOnClicked 	OnChildButtonClicked = FOnClicked();
 	if (MaterialEditorInstance->OriginalFunction)
diff --git a/Engine/Source/Editor/MaterialEditor/Private/SMaterialParametersOverviewWidget.h b/Engine/Source/Editor/MaterialEditor/Private/SMaterialParametersOverviewWidget.h
--- a/Engine/Source/Editor/MaterialEditor/Private/SMaterialParametersOverviewWidget.h
+++ b/Engine/Source/Editor/MaterialEditor/Private/SMaterialParametersOverviewWidget.h
@@ -74,7 +74,7 @@ public:
 	SLATE_ARGUMENT(UMaterialEditorPreviewParameters*, InMaterialEditorInstance)
 	SLATE_ARGUMENT(TSharedPtr<class IPropertyRowGenerator>, InGenerator)
 	SLATE_END_ARGS()
-	void Refresh(const bool bRegenerateArray = true);
+	void Refresh();
 	void Construct(const FArguments& InArgs);
 	void UpdateEditorInstance(UMaterialEditorPreviewParameters* InMaterialEditorInstance);
 	TSharedPtr<class IPropertyRowGenerator> GetGenerator();
@@ -129,7 +129,7 @@ public:
 	UMaterialEditorPreviewParameters* MaterialEditorInstance;
 
 	/** Builds the custom parameter groups category */
-	void CreateGroupsWidget(const bool bRegenerateArray);
+	void CreateGroupsWidget();
 
 	TWeakPtr<SMaterialParametersOverviewPanel> GetOwner() { return Owner; }
 	bool HasAnyParameters() const { return bHasAnyParameters; }
