Adding a material parameter refresh at the end of init
#rnx

diff --git a/Engine/Source/Editor/MaterialEditor/Private/MaterialEditor.cpp b/Engine/Source/Editor/MaterialEditor/Private/MaterialEditor.cpp
--- a/Engine/Source/Editor/MaterialEditor/Private/MaterialEditor.cpp
+++ b/Engine/Source/Editor/MaterialEditor/Private/MaterialEditor.cpp
@@ -727,9 +727,9 @@ void FMaterialEditor::InitMaterialEditor( const EToolkitMode::Type Mode, const T
 	RegenerateCodeView(true);
 
 	ForceRefreshExpressionPreviews();
-
-	if (Generator.IsValid())
+	if (Generator.IsValid() && MaterialEditorInstance)
 	{
+		MaterialEditorInstance->RegenerateArrays();
 		TArray<UObject*> Objects;
 		Objects.Add(MaterialEditorInstance);
 		Generator->SetObjects(Objects);
