Refreshing the material parameter object now sends it to the generator as well.
#rnx

diff --git a/Engine/Source/Editor/MaterialEditor/Private/MaterialEditor.cpp b/Engine/Source/Editor/MaterialEditor/Private/MaterialEditor.cpp
--- a/Engine/Source/Editor/MaterialEditor/Private/MaterialEditor.cpp
+++ b/Engine/Source/Editor/MaterialEditor/Private/MaterialEditor.cpp
@@ -730,10 +730,8 @@ void FMaterialEditor::InitMaterialEditor( const EToolkitMode::Type Mode, const T
 	ForceRefreshExpressionPreviews();
 	if (Generator.IsValid() && MaterialEditorInstance)
 	{
-		MaterialEditorInstance->RegenerateArrays();
-		TArray<UObject*> Objects;
-		Objects.Add(MaterialEditorInstance);
-		Generator->SetObjects(Objects);
+		UpdateGenerator();
+
 	}
 
 	if (OriginalMaterial->bUsedAsSpecialEngineMaterial)
@@ -753,6 +751,17 @@ void FMaterialEditor::InitMaterialEditor( const EToolkitMode::Type Mode, const T
 	}
 }
 
+void FMaterialEditor::UpdateGenerator()
+{
+	if (MaterialEditorInstance && Generator.IsValid())
+	{
+		MaterialEditorInstance->RegenerateArrays();
+		TArray<UObject*> Objects;
+		Objects.Add(MaterialEditorInstance);
+		Generator->SetObjects(Objects);
+	}
+}
+
 FMaterialEditor::FMaterialEditor()
 	: bMaterialDirty(false)
 	, bStatsFromPreviewMaterial(false)
@@ -3114,10 +3123,8 @@ void FMaterialEditor::OnConvertObjects()
 							// Refresh the expression preview if we changed its properties after it was created
 							NewExpression->bNeedToUpdatePreview = true;
 							RefreshExpressionPreview( NewExpression, true );
-							if (MaterialEditorInstance)
-							{
-								MaterialEditorInstance->RegenerateArrays();
-							}
+							
+							UpdateGenerator();
 						}
 
 						NodesToDelete.AddUnique(GraphNode);
@@ -4526,10 +4533,7 @@ void FMaterialEditor::RedoGraphAction()
 		Material->BuildEditorParameterList();
 	}
 
-	if (MaterialEditorInstance)
-	{
-		MaterialEditorInstance->RegenerateArrays();
-	}
+	UpdateGenerator();
 }
 
 void FMaterialEditor::OnAlignTop()
@@ -4622,10 +4626,7 @@ void FMaterialEditor::PostUndo(bool bSuccess)
 		GraphEditor->NotifyGraphChanged();
 		SetMaterialDirty();
 
-		if (MaterialEditorInstance)
-		{
-			MaterialEditorInstance->RegenerateArrays();
-		}
+		UpdateGenerator();
 
 		FSlateApplication::Get().DismissAllMenus();
 	}
@@ -4736,10 +4737,7 @@ void FMaterialEditor::NotifyPostChange( const FPropertyChangedEvent& PropertyCha
 	Material->MarkPackageDirty();
 	SetMaterialDirty();
 
-	if (MaterialEditorInstance)
-	{
-		MaterialEditorInstance->RegenerateArrays();
-	}
+	UpdateGenerator();
 }
 
 void FMaterialEditor::ToggleCollapsed(UMaterialExpression* MaterialExpression)
@@ -5402,10 +5400,7 @@ void FMaterialEditor::OnNodeTitleCommitted(const FText& NewText, ETextCommit::Ty
 		const FScopedTransaction Transaction( LOCTEXT( "RenameNode", "Rename Node" ) );
 		NodeBeingChanged->Modify();
 		NodeBeingChanged->OnRenameNode(NewText.ToString());
-		if (MaterialEditorInstance)
-		{
-			MaterialEditorInstance->RegenerateArrays();
-		}
+		UpdateGenerator();
 		MaterialCustomPrimitiveDataWidget->UpdateEditorInstance(MaterialEditorInstance);
 	}
 }
diff --git a/Engine/Source/Editor/MaterialEditor/Private/MaterialEditor.h b/Engine/Source/Editor/MaterialEditor/Private/MaterialEditor.h
--- a/Engine/Source/Editor/MaterialEditor/Private/MaterialEditor.h
+++ b/Engine/Source/Editor/MaterialEditor/Private/MaterialEditor.h
@@ -780,6 +780,7 @@ private:
 	void OnFinishedChangingProperties(const FPropertyChangedEvent& PropertyChangedEvent);
 	void OnFinishedChangingParametersFromOverview(const FPropertyChangedEvent& PropertyChangedEvent);
 	void GeneratorRowsRefreshed();
+	void UpdateGenerator();
 private:
 	/** List of open tool panels; used to ensure only one exists at any one time */
 	TMap< FName, TWeakPtr<class SDockableTab> > SpawnedToolPanels;
