add platform extension shader formats to material editor stats

diff --git a/Engine/Source/Editor/MaterialEditor/Private/MaterialStats.cpp b/Engine/Source/Editor/MaterialEditor/Private/MaterialStats.cpp
--- a/Engine/Source/Editor/MaterialEditor/Private/MaterialStats.cpp
+++ b/Engine/Source/Editor/MaterialEditor/Private/MaterialStats.cpp
@@ -446,6 +446,21 @@ void FMaterialStats::BuildShaderPlatformDB()
 			AddShaderPlatform(EPlatformCategoryType::Console, ShaderPlatform, FName(*PlatformName), true, true, PlatformName);
 		}
 	}
+
+	// Add platform extensions
+	for (int32 StaticPlatform = SP_StaticPlatform_First; StaticPlatform <= SP_StaticPlatform_Last; ++StaticPlatform)
+	{
+		EShaderPlatform ShaderPlatform = (EShaderPlatform)StaticPlatform;
+		if (FDataDrivenShaderPlatformInfo::IsValid(ShaderPlatform))
+		{
+			bool bIsConsole = FDataDrivenShaderPlatformInfo::GetIsConsole(ShaderPlatform);
+
+			const FName ShaderFormat = LegacyShaderPlatformToShaderFormat(ShaderPlatform);
+			const FName PlatformName = ShaderPlatformToPlatformName(ShaderPlatform);
+			AddShaderPlatform(bIsConsole ? EPlatformCategoryType::Console : EPlatformCategoryType::Desktop, ShaderPlatform, PlatformName, true, true, PlatformName.ToString());
+		}
+	}
+
 }
 
 TSharedPtr<FShaderPlatformSettings> FMaterialStats::AddShaderPlatform(const EPlatformCategoryType PlatformType, const EShaderPlatform PlatformID, const FName PlatformName,
