Fixed crash when updating a material function used by a MID currently being edited

diff --git a/Engine/Source/Editor/MaterialEditor/Private/MaterialEditingLibrary.cpp b/Engine/Source/Editor/MaterialEditor/Private/MaterialEditingLibrary.cpp
--- a/Engine/Source/Editor/MaterialEditor/Private/MaterialEditingLibrary.cpp
+++ b/Engine/Source/Editor/MaterialEditor/Private/MaterialEditingLibrary.cpp
@@ -265,6 +265,35 @@ namespace MaterialEditingLibraryImpl
 			ColumnHeight = MaterialExpression->MaterialExpressionEditorY + MaterialExpression->GetHeight() + ME_STD_HPADDING;
 		}
 	}
+
+	IMaterialEditor* FindMaterialEditorForAsset(UObject* InAsset)
+	{
+		if (IAssetEditorInstance* AssetEditorInstance = (InAsset != nullptr) ? GEditor->GetEditorSubsystem<UAssetEditorSubsystem>()->FindEditorForAsset(InAsset, false) : nullptr)
+		{
+			// Ensure this is not a UMaterialInstanceDynamic, as that doesn't use IMaterialEditor as its editor
+			if (!InAsset->IsA(UMaterialInstanceDynamic::StaticClass()))
+			{
+				return static_cast<IMaterialEditor*>(AssetEditorInstance);
+			}
+		}
+
+		return nullptr;
+	}
+
+	FMaterialInstanceEditor* FindMaterialInstanceEditorForAsset(UObject* InAsset)
+	{
+		if (IAssetEditorInstance* AssetEditorInstance = (InAsset != nullptr) ? GEditor->GetEditorSubsystem<UAssetEditorSubsystem>()->FindEditorForAsset(InAsset, false) : nullptr)
+		{
+			// Ensure this is not a UMaterialInstanceDynamic, as that doesn't use FMaterialInstanceEditor as its editor
+			if (!InAsset->IsA(UMaterialInstanceDynamic::StaticClass()))
+			{
+				return static_cast<FMaterialInstanceEditor*>(AssetEditorInstance);
+			}
+		}
+
+		return nullptr;
+	}
+	
 }
 
 void UMaterialEditingLibrary::RebuildMaterialInstanceEditors(UMaterial* BaseMaterial)
@@ -277,7 +306,6 @@ void UMaterialEditingLibrary::RebuildMaterialInstanceEditors(UMaterial* BaseMate
 		UObject* EditedAsset = EditedAssets[AssetIdx];
 
 		UMaterialInstance* SourceInstance = Cast<UMaterialInstance>(EditedAsset);
-
 		if (!SourceInstance)
 		{
 			// Check to see if the EditedAssets are from material instance editor
@@ -288,17 +316,14 @@ void UMaterialEditingLibrary::RebuildMaterialInstanceEditors(UMaterial* BaseMate
 			}
 		}
 
-		// Ensure the material instance is valid and not a UMaterialInstanceDynamic, as that doesn't use FMaterialInstanceEditor as its editor
-		if (SourceInstance != nullptr && !SourceInstance->IsA(UMaterialInstanceDynamic::StaticClass()))
+		if (SourceInstance != nullptr)
 		{
 			UMaterial* MICOriginalMaterial = SourceInstance->GetMaterial();
 			if (MICOriginalMaterial == BaseMaterial)
 			{
-				IAssetEditorInstance* EditorInstance = AssetEditorSubsystem->FindEditorForAsset(EditedAsset, false);
-				if (EditorInstance != nullptr)
+				if (FMaterialInstanceEditor* MaterialInstanceEditor = MaterialEditingLibraryImpl::FindMaterialInstanceEditorForAsset(SourceInstance))
 				{
-					FMaterialInstanceEditor* OtherEditor = static_cast<FMaterialInstanceEditor*>(EditorInstance);
-					OtherEditor->RebuildMaterialInstanceEditor();
+					MaterialInstanceEditor->RebuildMaterialInstanceEditor();
 				}
 			}
 		}
@@ -322,11 +347,9 @@ void UMaterialEditingLibrary::RebuildMaterialInstanceEditors(UMaterialFunction*
 			// Update function instances that are children of this material function	
 			if (BaseFunction && BaseFunction == FunctionInstance->GetBaseFunction())
 			{
-				IAssetEditorInstance* EditorInstance = AssetEditorSubsystem->FindEditorForAsset(EditedAsset, false);
-				if (EditorInstance)
+				if (FMaterialInstanceEditor* MaterialInstanceEditor = MaterialEditingLibraryImpl::FindMaterialInstanceEditorForAsset(EditedAsset))
 				{
-					FMaterialInstanceEditor* OtherEditor = static_cast<FMaterialInstanceEditor*>(EditorInstance);
-					OtherEditor->RebuildMaterialInstanceEditor();
+					MaterialInstanceEditor->RebuildMaterialInstanceEditor();
 				}
 			}
 		}
@@ -350,11 +373,9 @@ void UMaterialEditingLibrary::RebuildMaterialInstanceEditors(UMaterialFunction*
 
 				if (BaseFunction && (DependentFunctions.Contains(BaseFunction) || DependentFunctions.Contains(BaseFunction->ParentFunction)))
 				{
-					IAssetEditorInstance* EditorInstance = AssetEditorSubsystem->FindEditorForAsset(EditedAsset, false);
-					if (EditorInstance != nullptr)
+					if (FMaterialInstanceEditor* MaterialInstanceEditor = MaterialEditingLibraryImpl::FindMaterialInstanceEditorForAsset(EditedAsset))
 					{
-						FMaterialInstanceEditor* OtherEditor = static_cast<FMaterialInstanceEditor*>(EditorInstance);
-						OtherEditor->RebuildMaterialInstanceEditor();
+						MaterialInstanceEditor->RebuildMaterialInstanceEditor();
 					}
 				}
 			}
@@ -714,8 +735,7 @@ bool UMaterialEditingLibrary::GetMaterialDefaultStaticSwitchParameterValue(UMate
 
 TSet<UObject*> UMaterialEditingLibrary::GetMaterialSelectedNodes(UMaterial* Material)
 {
-	auto* MaterialEditor = (IMaterialEditor*)GEditor->GetEditorSubsystem<UAssetEditorSubsystem>()->FindEditorForAsset(Material, false);
-	if (MaterialEditor)
+	if (IMaterialEditor* MaterialEditor = MaterialEditingLibraryImpl::FindMaterialEditorForAsset(Material))
 	{
 		TSet<UObject*> SelectedMaterialObjects;
 		for (const FFieldVariant& SelectedNode : MaterialEditor->GetSelectedNodes())
@@ -837,8 +857,7 @@ void UMaterialEditingLibrary::UpdateMaterialFunction(UMaterialFunctionInterface*
 						}
 
 						// if this instance was opened in an editor notify the change
-						auto* MaterialEditor = (IMaterialEditor*)GEditor->GetEditorSubsystem<UAssetEditorSubsystem>()->FindEditorForAsset(CurrentMaterial, false);
-						if (MaterialEditor)
+						if (IMaterialEditor* MaterialEditor = MaterialEditingLibraryImpl::FindMaterialEditorForAsset(CurrentMaterial))
 						{
 							MaterialEditor->NotifyExternalMaterialChange();
 						}
@@ -861,8 +880,7 @@ void UMaterialEditingLibrary::UpdateMaterialFunction(UMaterialFunctionInterface*
 						CurrentInstance->PostEditChange();
 
 						// if this instance was opened in an editor notify the change
-						auto* MaterialEditor = (IMaterialEditor*)GEditor->GetEditorSubsystem<UAssetEditorSubsystem>()->FindEditorForAsset(CurrentInstance, false);
-						if (MaterialEditor)
+						if (IMaterialEditor* MaterialEditor = MaterialEditingLibraryImpl::FindMaterialEditorForAsset(CurrentInstance))
 						{
 							MaterialEditor->NotifyExternalMaterialChange();
 						}
@@ -1156,4 +1174,4 @@ bool UMaterialEditingLibrary::GetStaticSwitchParameterSource(UMaterialInterface*
 		}
 	}
 	return false;
-}
\ No newline at end of file
+}
