Adding a material parameter panel update to when a parameter is wired in

diff --git a/Engine/Source/Editor/MaterialEditor/Private/MaterialEditor.cpp b/Engine/Source/Editor/MaterialEditor/Private/MaterialEditor.cpp
--- a/Engine/Source/Editor/MaterialEditor/Private/MaterialEditor.cpp
+++ b/Engine/Source/Editor/MaterialEditor/Private/MaterialEditor.cpp
@@ -4470,6 +4470,7 @@ void FMaterialEditor::UpdateMaterialAfterGraphChange()
 	RegenerateCodeView();
 	RefreshExpressionPreviews();
 	SetMaterialDirty();
+	UpdateGenerator();
 	if (bHideUnrelatedNodes && !bLockNodeFadeState && bSelectRegularNode)
 	{
 		GraphEditor->ResetAllNodesUnrelatedStates();
