Added functionality in the material editor to rename parameters by pressing F2.

diff --git a/Engine/Source/Editor/MaterialEditor/Private/MaterialEditor.cpp b/Engine/Source/Editor/MaterialEditor/Private/MaterialEditor.cpp
--- a/Engine/Source/Editor/MaterialEditor/Private/MaterialEditor.cpp
+++ b/Engine/Source/Editor/MaterialEditor/Private/MaterialEditor.cpp
@@ -1278,6 +1278,8 @@ void FMaterialEditor::AddInheritanceMenuEntry(FToolMenuSection& Section, const F
 	FSlateIcon OpenIcon(FEditorStyle::GetStyleSetName(), "ContentBrowser.AssetActions.OpenInExternalEditor");
 	FSlateIcon FindInContentBrowserIcon(FEditorStyle::GetStyleSetName(), "SystemWideCommands.FindInContentBrowser");
 
+	const float MenuIconSize = FCoreStyle::Get().GetFloat("Menu.MenuIconSize", nullptr, 16.f);
+
 	TSharedRef<SWidget> EntryWidget =
 		SNew(SHorizontalBox)
 		.ToolTipText(LOCTEXT("OpenInEditor", "Open In Editor"))
@@ -1287,14 +1289,14 @@ void FMaterialEditor::AddInheritanceMenuEntry(FToolMenuSection& Section, const F
 		.Padding(FMargin(2, 0, 2, 0))
 		[
 			SNew( SBox )
-			.WidthOverride( MultiBoxConstants::MenuIconSize + 2 )
-			.HeightOverride( MultiBoxConstants::MenuIconSize )
+			.WidthOverride(MenuIconSize + 2 )
+			.HeightOverride(MenuIconSize)
 			.HAlign(HAlign_Center)
 			.VAlign(VAlign_Center)
 			[
 				SNew( SBox )
-				.WidthOverride( MultiBoxConstants::MenuIconSize )
-				.HeightOverride( MultiBoxConstants::MenuIconSize )
+				.WidthOverride(MenuIconSize)
+				.HeightOverride(MenuIconSize)
 				[
 					SNew(SImage)
 					.Image(OpenIcon.GetIcon())
@@ -1322,14 +1324,14 @@ void FMaterialEditor::AddInheritanceMenuEntry(FToolMenuSection& Section, const F
 			.OnClicked_Lambda([FindInContentBrowserAction]() { FindInContentBrowserAction.ExecuteIfBound(); return FReply::Handled(); })
 			[
 				SNew( SBox )
-				.WidthOverride( MultiBoxConstants::MenuIconSize + 2 )
-				.HeightOverride( MultiBoxConstants::MenuIconSize )
+				.WidthOverride(MenuIconSize + 2 )
+				.HeightOverride(MenuIconSize)
 				.HAlign(HAlign_Center)
 				.VAlign(VAlign_Center)
 				[
 					SNew( SBox )
-					.WidthOverride( MultiBoxConstants::MenuIconSize )
-					.HeightOverride( MultiBoxConstants::MenuIconSize )
+					.WidthOverride(MenuIconSize)
+					.HeightOverride(MenuIconSize)
 					[
 						SNew(SImage)
 						.Image(FindInContentBrowserIcon.GetIcon())
@@ -4946,6 +4948,10 @@ TSharedRef<SGraphEditor> FMaterialEditor::CreateGraphEditorWidget()
 			FExecuteAction::CreateSP( this, &FMaterialEditor::DuplicateNodes ),
 			FCanExecuteAction::CreateSP( this, &FMaterialEditor::CanDuplicateNodes )
 			);
+		GraphEditorCommands->MapAction(FGenericCommands::Get().Rename,
+			FExecuteAction::CreateSP(this, &FMaterialEditor::OnRenameNode),
+			FCanExecuteAction::CreateSP(this, &FMaterialEditor::CanRenameNodes)
+		);
 
 		// Graph Editor Commands
 		GraphEditorCommands->MapAction( FGraphEditorCommands::Get().CreateComment,
@@ -5394,6 +5400,35 @@ void FMaterialEditor::OnNodeDoubleClicked(class UEdGraphNode* Node)
 	}
 }
 
+void FMaterialEditor::OnRenameNode()
+{
+	const FGraphPanelSelectionSet SelectedNodes = GetSelectedNodes();
+	for (FGraphPanelSelectionSet::TConstIterator NodeIt(SelectedNodes); NodeIt; ++NodeIt)
+	{
+		UEdGraphNode* SelectedNode = Cast<UEdGraphNode>(*NodeIt);
+		if (SelectedNode != nullptr && SelectedNode->GetCanRenameNode())
+		{
+			bool ToRename = true;
+			GraphEditor->IsNodeTitleVisible(SelectedNode, ToRename);
+			break;
+		}
+	}
+}
+
+bool FMaterialEditor::CanRenameNodes() const
+{
+	const FGraphPanelSelectionSet SelectedNodes = GetSelectedNodes();
+	for (FGraphPanelSelectionSet::TConstIterator SelectedIter(SelectedNodes); SelectedIter; ++SelectedIter)
+	{
+		UEdGraphNode* Node = Cast<UEdGraphNode>(*SelectedIter);
+		if ((Node != nullptr) && Node->GetCanRenameNode())
+		{
+			return true;
+		}
+	}
+	return false;
+}
+
 void FMaterialEditor::OnNodeTitleCommitted(const FText& NewText, ETextCommit::Type CommitInfo, UEdGraphNode* NodeBeingChanged)
 {
 	if (NodeBeingChanged)
diff --git a/Engine/Source/Editor/MaterialEditor/Private/MaterialEditor.h b/Engine/Source/Editor/MaterialEditor/Private/MaterialEditor.h
--- a/Engine/Source/Editor/MaterialEditor/Private/MaterialEditor.h
+++ b/Engine/Source/Editor/MaterialEditor/Private/MaterialEditor.h
@@ -536,6 +536,12 @@ protected:
 	/** Called to redo the last undone action */
 	void RedoGraphAction();
 
+	/** On starting to rename node */
+	void OnRenameNode();
+
+	/** Check if node can be renamed */
+	bool CanRenameNodes() const;
+
 	void OnAlignTop();
 	void OnAlignMiddle();
 	void OnAlignBottom();
