Changing a PlatformNeedsOfflineCompiler case to use data driver platform info.

diff --git a/Engine/Source/Editor/MaterialEditor/Private/MaterialStatsCommon.cpp b/Engine/Source/Editor/MaterialEditor/Private/MaterialStatsCommon.cpp
--- a/Engine/Source/Editor/MaterialEditor/Private/MaterialStatsCommon.cpp
+++ b/Engine/Source/Editor/MaterialEditor/Private/MaterialStatsCommon.cpp
@@ -217,7 +217,6 @@ bool FMaterialStatsUtils::PlatformNeedsOfflineCompiler(const EShaderPlatform Sha
 {
 	switch (ShaderPlatform)
 	{
-		case SP_PS4:
 		case SP_OPENGL_SM5:
 		case SP_OPENGL_ES31_EXT:
 		case SP_OPENGL_PCES3_1:
