Fixing shader source showing 'duplicate' entries. The shader lookups were only using the shader type name and not considering that they are unique per vertex factory.

diff --git a/Engine/Source/Editor/MaterialEditor/Private/MaterialStats.cpp b/Engine/Source/Editor/MaterialEditor/Private/MaterialStats.cpp
--- a/Engine/Source/Editor/MaterialEditor/Private/MaterialStats.cpp
+++ b/Engine/Source/Editor/MaterialEditor/Private/MaterialStats.cpp
@@ -57,7 +57,7 @@ void FShaderPlatformSettings::ClearResources()
 			PlatformData[i].MaterialResourcesStats = nullptr;
 		}
 
-		PlatformData[i].ArrShaderNames.Empty();
+		PlatformData[i].ArrShaderEntries.Empty();
 
 		PlatformData[i].bCompilingShaders = false;
 		PlatformData[i].bNeedShaderRecompilation = true;
@@ -66,19 +66,19 @@ void FShaderPlatformSettings::ClearResources()
 
 FText FShaderPlatformSettings::GetSelectedShaderViewComboText(EMaterialQualityLevel::Type QualityLevel) const
 {
-	if (PlatformData[QualityLevel].ArrShaderNames.Num() == 0)
+	if (PlatformData[QualityLevel].ArrShaderEntries.Num() == 0)
 	{
 		return FText::FromString(TEXT("-Compiling-Shaders-"));
 	}
 
-	return FText::FromName(PlatformData[QualityLevel].ComboBoxSelectedName);
+	return FText::FromString(PlatformData[QualityLevel].ComboBoxSelectedEntry.Text);
 }
 
-void FShaderPlatformSettings::OnShaderViewComboSelectionChanged(TSharedPtr<FName> Item, EMaterialQualityLevel::Type QualityType)
+void FShaderPlatformSettings::OnShaderViewComboSelectionChanged(TSharedPtr<FMaterialShaderEntry> Item, EMaterialQualityLevel::Type QualityType)
 {
 	if (Item.IsValid())
 	{
-		PlatformData[QualityType].ComboBoxSelectedName = *Item.Get();
+		PlatformData[QualityType].ComboBoxSelectedEntry = *Item.Get();
 		PlatformData[QualityType].bUpdateShaderCode = true;
 	}
 }
@@ -101,16 +101,19 @@ FText FShaderPlatformSettings::GetShaderCode(const EMaterialQualityLevel::Type Q
 	// check if shader compilation is done and extract shader code
 	if (bCompilationFinished)
 	{
-		TMap<FHashedName, TShaderRef<FShader>> ShaderMap;
+		TMap<FShaderId, TShaderRef<FShader>> ShaderMap;
 		MaterialShaderMap->GetShaderList(ShaderMap);
 
-		const auto Entry = ShaderMap.Find(PlatformData[QualityType].ComboBoxSelectedName);
+		const FShaderId& ShaderId = PlatformData[QualityType].ComboBoxSelectedEntry.ShaderId;
+
+		const auto Entry = ShaderMap.Find(ShaderId);
 		if (Entry != nullptr)
 		{
 			const TShaderRef<FShader>& Shader = *Entry;
 
-			const FName ShaderFName = Shader.GetType()->GetFName();
-			const FMemoryImageString* ShaderSource = MaterialShaderMap->GetShaderSource(ShaderFName);
+			const FName VertexFactoryName = Shader.GetVertexFactoryType()->GetFName();
+			const FName ShaderTypeName = Shader.GetType()->GetFName();
+			const FMemoryImageString* ShaderSource = MaterialShaderMap->GetShaderSource(VertexFactoryName, ShaderTypeName);
 			if (ShaderSource != nullptr)
 			{
 				PlatformData[QualityType].bUpdateShaderCode = false;
@@ -221,15 +224,23 @@ bool FShaderPlatformSettings::Update()
 					TMap<FShaderId, TShaderRef<FShader>> ShaderMap;
 					MaterialShaderMap->GetShaderList(ShaderMap);
 
-					QualityItem.ArrShaderNames.Empty();
+					QualityItem.ArrShaderEntries.Empty();
 					for (const auto& Entry : ShaderMap)
 					{
-						QualityItem.ArrShaderNames.Add(MakeShareable(new FName(Entry.Value.GetType()->GetFName())));
+						FShaderType* EntryShader = Entry.Value.GetType();
+						FVertexFactoryType* VertexFactory = Entry.Value.GetVertexFactoryType();
+						FString ShaderName = FString::Printf(TEXT("%s/%s"), VertexFactory->GetName(), EntryShader->GetName());
+						QualityItem.ArrShaderEntries.Add(MakeShareable(new FMaterialShaderEntry{ Entry.Key, ShaderName }));
 					}
 
-					if (QualityItem.ArrShaderNames.Num() > 0)
+					if (QualityItem.ArrShaderEntries.Num() > 0)
 					{
-						QualityItem.ComboBoxSelectedName = *QualityItem.ArrShaderNames[0];
+						QualityItem.ArrShaderEntries.Sort([](const TSharedPtr<FMaterialShaderEntry>& First, const TSharedPtr<FMaterialShaderEntry>& Second)
+							{
+								return First->Text < Second->Text;
+							});
+
+						QualityItem.ComboBoxSelectedEntry = *QualityItem.ArrShaderEntries[0];
 					}
 				}
 
@@ -701,10 +712,10 @@ TSharedRef<SDockTab> FMaterialStats::SpawnTab_ShaderCode(const FSpawnTabArgs& Ar
 	TSharedPtr<FShaderPlatformSettings> PlatformPtr = GetPlatformSettings(PlatformID);
 	check(PlatformPtr.IsValid());
 
-	TSharedRef<SComboBox<TSharedPtr<FName>>> ShaderBox = SNew(SComboBox<TSharedPtr<FName>>)
-		.OptionsSource(PlatformPtr->GetShaderNames(QualityLevel))
-		.OnGenerateWidget_Lambda([](TSharedPtr<FName> Value) { return SNew(STextBlock).Text(FText::FromName(*Value.Get())); })
-		.OnSelectionChanged_Lambda([PlatformPtr, QualityLevel](TSharedPtr<FName> Item, ESelectInfo::Type SelectInfo) { PlatformPtr->OnShaderViewComboSelectionChanged(Item, QualityLevel); })
+	TSharedRef<SComboBox<TSharedPtr<FMaterialShaderEntry>>> ShaderBox = SNew(SComboBox<TSharedPtr<FMaterialShaderEntry>>)
+		.OptionsSource(PlatformPtr->GetShaderEntries(QualityLevel))
+		.OnGenerateWidget_Lambda([](TSharedPtr<FMaterialShaderEntry> Value) { return SNew(STextBlock).Text(FText::FromString(Value->Text)); })
+		.OnSelectionChanged_Lambda([PlatformPtr, QualityLevel](TSharedPtr<FMaterialShaderEntry> Item, ESelectInfo::Type SelectInfo) { PlatformPtr->OnShaderViewComboSelectionChanged(Item, QualityLevel); })
 		[
 			SNew(STextBlock)
 			.Text_Lambda([PlatformPtr, QualityLevel]() { return PlatformPtr->GetSelectedShaderViewComboText(QualityLevel); })
diff --git a/Engine/Source/Editor/MaterialEditor/Private/MaterialStats.h b/Engine/Source/Editor/MaterialEditor/Private/MaterialStats.h
--- a/Engine/Source/Editor/MaterialEditor/Private/MaterialStats.h
+++ b/Engine/Source/Editor/MaterialEditor/Private/MaterialStats.h
@@ -66,6 +66,12 @@ struct FShaderStatsInfo
 	}
 };
 
+struct FMaterialShaderEntry
+{
+	FShaderId ShaderId;
+	FString Text;
+};
+
 /** structure used to manage shader compilation and source code extraction for a specified shader platform
 *   used for building the material stats */
 struct FShaderPlatformSettings
@@ -81,10 +87,10 @@ public:
 
 		/** array of shader ids for this platform; needed to fill ComboBox in MaterialEditor's shader viewer
 		* generated from ShaderID.ShaderType->GetFName() */
-		TArray<TSharedPtr<FName>> ArrShaderNames;
+		TArray<TSharedPtr<FMaterialShaderEntry>> ArrShaderEntries;
 
 		/** ComboBox current entry */
-		FName ComboBoxSelectedName;
+		FMaterialShaderEntry ComboBoxSelectedEntry;
 
 		/** flag that marks the usage of this data structure */
 		bool bExtractStats = false;
@@ -181,7 +187,7 @@ public:
 	FORCEINLINE TWeakPtr<class SDockTab> GetCodeViewerTab(const EMaterialQualityLevel::Type QualityLevel);
 
 	/** returns an array with the names of all the compiled shaders for this material with specified quality level */
-	FORCEINLINE const TArray<TSharedPtr<FName>> *GetShaderNames(const EMaterialQualityLevel::Type QualityLevel);
+	FORCEINLINE const TArray<TSharedPtr<FMaterialShaderEntry>> *GetShaderEntries(const EMaterialQualityLevel::Type QualityLevel);
 
 	/** when set this flag will indicate the presence of this material with a particular quality level inside the stats widget */
 	void SetExtractStatsFlag(const EMaterialQualityLevel::Type QualityType, const bool bValue);
@@ -219,7 +225,7 @@ public:
 	FText GetSelectedShaderViewComboText(EMaterialQualityLevel::Type QualityLevel) const;
 
 	/** callback function called when we change the content of the shader viewer combo-box, used to select a different shader to be displayed */
-	void OnShaderViewComboSelectionChanged(TSharedPtr<FName> Item, EMaterialQualityLevel::Type QualityType);
+	void OnShaderViewComboSelectionChanged(TSharedPtr<FMaterialShaderEntry> Item, EMaterialQualityLevel::Type QualityType);
 
 	/** returns the actual shader source selected by the shaders viewer's combo-box */
 	FText GetShaderCode(const EMaterialQualityLevel::Type QualityType);
@@ -568,10 +574,10 @@ FORCEINLINE TWeakPtr<class SDockTab> FShaderPlatformSettings::GetCodeViewerTab(c
 	return SomePlatformData.CodeViewerTab;
 }
 
-FORCEINLINE const TArray<TSharedPtr<FName>> *FShaderPlatformSettings::GetShaderNames(const EMaterialQualityLevel::Type QualityLevel)
+FORCEINLINE const TArray<TSharedPtr<FMaterialShaderEntry>> *FShaderPlatformSettings::GetShaderEntries(const EMaterialQualityLevel::Type QualityLevel)
 {
 	FPlatformData& SomePlatformData = GetPlatformData(QualityLevel);
-	return &SomePlatformData.ArrShaderNames;
+	return &SomePlatformData.ArrShaderEntries;
 }
 
 FORCEINLINE void FShaderPlatformSettings::SetExtractStatsFlag(const EMaterialQualityLevel::Type QualityType, const bool bValue)
diff --git a/Engine/Source/Editor/UnrealEd/Private/EditorBuildUtils.cpp b/Engine/Source/Editor/UnrealEd/Private/EditorBuildUtils.cpp
--- a/Engine/Source/Editor/UnrealEd/Private/EditorBuildUtils.cpp
+++ b/Engine/Source/Editor/UnrealEd/Private/EditorBuildUtils.cpp
@@ -1533,10 +1533,10 @@ void FMaterialOfflineCompilation::CopyPlatformSpecificStats()
 		return;
 	}
 
-	TMap<FHashedName, TShaderRef<FShader>> SrcShaders;
+	TMap<FShaderId, TShaderRef<FShader>> SrcShaders;
 	SrcShaderMap->GetShaderList(SrcShaders);
 
-	TMap<FHashedName, TShaderRef<FShader>> DstShaders;
+	TMap<FShaderId, TShaderRef<FShader>> DstShaders;
 	DstShaderMap->GetShaderList(DstShaders);
 
 	for (auto Pair : SrcShaders)
diff --git a/Engine/Source/Runtime/Engine/Private/Materials/MaterialShader.cpp b/Engine/Source/Runtime/Engine/Private/Materials/MaterialShader.cpp
--- a/Engine/Source/Runtime/Engine/Private/Materials/MaterialShader.cpp
+++ b/Engine/Source/Runtime/Engine/Private/Materials/MaterialShader.cpp
@@ -1577,6 +1577,11 @@ void FMaterialShaderMap::Compile(
 	}
 }
 
+static FHashedName PreprocessedSourceKeyFromName(FName VertexFactoryName, FName VertexTypeName)
+{
+	return FHashedName(FString::Printf(TEXT("%s/%s"), *VertexFactoryName.ToString(), *VertexTypeName.ToString()));
+}
+
 FShader* FMaterialShaderMap::ProcessCompilationResultsForSingleJob(FShaderCompileJob* SingleJob, const FShaderPipelineType* ShaderPipeline, const FSHAHash& MaterialShaderMapHash)
 {
 	check(SingleJob);
@@ -1624,8 +1629,10 @@ FShader* FMaterialShaderMap::ProcessCompilationResultsForSingleJob(FShaderCompil
 	// add shader source
 	{
 		// Keep the preprocessed source list sorted by type name
-		const int32 Index = Algo::LowerBoundBy(GetMutableContent()->ShaderProcessedSource, CurrentJob.ShaderType->GetHashedName(), [](const FMaterialProcessedSource& Value) { return Value.Name; });
-		GetMutableContent()->ShaderProcessedSource.EmplaceAt(Index, CurrentJob.ShaderType->GetHashedName(), *CurrentJob.Output.OptionalFinalShaderSource);
+		const FHashedName Key = CurrentJob.VFType ? PreprocessedSourceKeyFromName(CurrentJob.VFType->GetFName(), CurrentJob.ShaderType->GetFName()) : FHashedName(CurrentJob.ShaderType->GetFName());
+
+		const int32 Index = Algo::LowerBoundBy(GetMutableContent()->ShaderProcessedSource, Key, [](const FMaterialProcessedSource& Value) { return Value.Name; });
+		GetMutableContent()->ShaderProcessedSource.EmplaceAt(Index, Key, *CurrentJob.Output.OptionalFinalShaderSource);
 	}
 #endif
 
@@ -2169,11 +2176,13 @@ void FMaterialShaderMap::LoadMissingShadersFromMemory(const FMaterial* Material)
 #endif // WITH_EDITOR
 
 #if WITH_EDITOR
-const FMemoryImageString* FMaterialShaderMap::GetShaderSource(const FName ShaderTypeName) const
+const FMemoryImageString* FMaterialShaderMap::GetShaderSource(const FName VertexFactoryName, const FName ShaderTypeName) const
 {
+	const FHashedName Key = PreprocessedSourceKeyFromName(VertexFactoryName, ShaderTypeName);
+
 	for (const FMaterialProcessedSource& Source : GetContent()->ShaderProcessedSource)
 	{
-		if (Source.Name == ShaderTypeName)
+		if (Source.Name == Key)
 		{
 			return &Source.Source;
 		}
diff --git a/Engine/Source/Runtime/Engine/Private/ShaderDerivedDataVersion.h b/Engine/Source/Runtime/Engine/Private/ShaderDerivedDataVersion.h
--- a/Engine/Source/Runtime/Engine/Private/ShaderDerivedDataVersion.h
+++ b/Engine/Source/Runtime/Engine/Private/ShaderDerivedDataVersion.h
@@ -11,4 +11,4 @@
 // In case of merge conflicts with DDC versions, you MUST generate a new GUID and set this new
 // guid as version
 #define GLOBALSHADERMAP_DERIVEDDATA_VER			TEXT("75B89920109D418CAC7D305E50CB100C")
-#define MATERIALSHADERMAP_DERIVEDDATA_VER		TEXT("BB69768054CD4A42AF7DC8CF8FAEED77")
+#define MATERIALSHADERMAP_DERIVEDDATA_VER		TEXT("F586D90911784814A450869743D895F5")
diff --git a/Engine/Source/Runtime/Engine/Public/MaterialShared.h b/Engine/Source/Runtime/Engine/Public/MaterialShared.h
--- a/Engine/Source/Runtime/Engine/Public/MaterialShared.h
+++ b/Engine/Source/Runtime/Engine/Public/MaterialShared.h
@@ -1113,7 +1113,7 @@ public:
 	bool TryToAddToExistingCompilationTask(FMaterial* Material);
 
 #if WITH_EDITOR
-	ENGINE_API const FMemoryImageString *GetShaderSource(const FName ShaderTypeName) const;
+	ENGINE_API const FMemoryImageString *GetShaderSource(FName VertexFactoryName, FName ShaderTypeName) const;
 #endif
 
 	/** Builds a list of the shaders in a shader map. */
