Add grouping to material graph.
Add composite & pinbase expressions / nodes, which use reroutes under the hood to ensure zero material overhead.
Convert MaterialEditor to a WorkflowCentricApplication.
Generally add subgraph existence support to MaterialGraph / MaterialEditor.

diff --git a/Engine/Source/Editor/DetailCustomizations/Private/CompositeRerouteCustomization.cpp b/Engine/Source/Editor/DetailCustomizations/Private/CompositeRerouteCustomization.cpp
new file mode 100644
--- /dev/null
+++ b/Engine/Source/Editor/DetailCustomizations/Private/CompositeRerouteCustomization.cpp
@@ -0,0 +1,35 @@
+// Copyright Epic Games, Inc. All Rights Reserved.
+
+#include "CompositeRerouteCustomization.h"
+#include "PropertyCustomizationHelpers.h"
+#include "DetailLayoutBuilder.h"
+#include "Materials/MaterialExpressionPinBase.h"
+
+#define LOCTEXT_NAMESPACE "CompositeRerouteDetails"
+
+TSharedRef<IPropertyTypeCustomization> FCompositeRerouteCustomization::MakeInstance()
+{
+	return MakeShareable(new FCompositeRerouteCustomization());
+}
+
+void FCompositeRerouteCustomization::CustomizeChildren(TSharedRef<IPropertyHandle> StructPropertyHandle, class IDetailChildrenBuilder& StructBuilder, IPropertyTypeCustomizationUtils& StructCustomizationUtils)
+{
+	/* do nothing */
+}
+
+void FCompositeRerouteCustomization::CustomizeHeader(TSharedRef<IPropertyHandle> StructPropertyHandle, class FDetailWidgetRow& HeaderRow, IPropertyTypeCustomizationUtils& StructCustomizationUtils)
+{
+	HeaderRow
+		.NameContent()
+		[
+			StructPropertyHandle->CreatePropertyNameWidget()
+		]
+		.ValueContent()
+		.MaxDesiredWidth(0.0f)
+		.MinDesiredWidth(125.0f)
+		[
+			StructPropertyHandle->GetChildHandle(GET_MEMBER_NAME_CHECKED(FCompositeReroute, Name))->CreatePropertyValueWidget()
+		];
+}
+
+#undef LOCTEXT_NAMESPACE
diff --git a/Engine/Source/Editor/DetailCustomizations/Private/CompositeRerouteCustomization.h b/Engine/Source/Editor/DetailCustomizations/Private/CompositeRerouteCustomization.h
new file mode 100644
--- /dev/null
+++ b/Engine/Source/Editor/DetailCustomizations/Private/CompositeRerouteCustomization.h
@@ -0,0 +1,20 @@
+// Copyright Epic Games, Inc. All Rights Reserved.
+
+#pragma once
+
+#include "CoreMinimal.h"
+#include "IPropertyTypeCustomization.h"
+
+class IPropertyHandle;
+
+class FCompositeRerouteCustomization : public IPropertyTypeCustomization
+{
+public:
+	/** Makes a new instance of this property customization */
+	static TSharedRef<IPropertyTypeCustomization> MakeInstance();
+
+	// IPropertyTypeCustomization interface
+	virtual void CustomizeChildren(TSharedRef<class IPropertyHandle> StructPropertyHandle, class IDetailChildrenBuilder& StructBuilder, IPropertyTypeCustomizationUtils& StructCustomizationUtils) override;
+	virtual void CustomizeHeader(TSharedRef<class IPropertyHandle> StructPropertyHandle, class FDetailWidgetRow& HeaderRow, IPropertyTypeCustomizationUtils& StructCustomizationUtils) override;
+};
+
diff --git a/Engine/Source/Editor/DetailCustomizations/Private/DetailCustomizations.cpp b/Engine/Source/Editor/DetailCustomizations/Private/DetailCustomizations.cpp
--- a/Engine/Source/Editor/DetailCustomizations/Private/DetailCustomizations.cpp
+++ b/Engine/Source/Editor/DetailCustomizations/Private/DetailCustomizations.cpp
@@ -135,6 +135,7 @@
 #include "SkeletalMeshLODSettingsDetails.h"
 #include "MaterialExpressionLandscapeGrassCustomization.h"
 #include "MaterialExpressionTextureBaseDetails.h"
+#include "CompositeRerouteCustomization.h"
 #include "TimecodeDetailsCustomization.h"
 #include "SkeletonDetails.h"
 #include "MaterialShadingModelCustomization.h"
@@ -260,6 +261,7 @@ void FDetailCustomizationsModule::RegisterPropertyTypeCustomizations()
 	RegisterCustomPropertyTypeLayout("CompositionGraphCapturePasses", FOnGetPropertyTypeCustomizationInstance::CreateStatic(&FRenderPassesCustomization::MakeInstance));
 	RegisterCustomPropertyTypeLayout("WeightedBlendable", FOnGetPropertyTypeCustomizationInstance::CreateStatic(&FWeightedBlendableCustomization::MakeInstance));
 	RegisterCustomPropertyTypeLayout("MaterialProxySettings", FOnGetPropertyTypeCustomizationInstance::CreateStatic(&FMaterialProxySettingsCustomizations::MakeInstance));
+	RegisterCustomPropertyTypeLayout("CompositeReroute", FOnGetPropertyTypeCustomizationInstance::CreateStatic(&FCompositeRerouteCustomization::MakeInstance));
 	RegisterCustomPropertyTypeLayout("CameraFilmbackSettings", FOnGetPropertyTypeCustomizationInstance::CreateStatic(&FCameraFilmbackSettingsCustomization::MakeInstance));
 	RegisterCustomPropertyTypeLayout("CameraLensSettings", FOnGetPropertyTypeCustomizationInstance::CreateStatic(&FCameraLensSettingsCustomization::MakeInstance));
 	RegisterCustomPropertyTypeLayout("CameraFocusSettings", FOnGetPropertyTypeCustomizationInstance::CreateStatic(&FCameraFocusSettingsCustomization::MakeInstance));
diff --git a/Engine/Source/Editor/GraphEditor/Private/MaterialNodes/SGraphNodeMaterialBase.h b/Engine/Source/Editor/GraphEditor/Private/MaterialNodes/SGraphNodeMaterialBase.h
--- a/Engine/Source/Editor/GraphEditor/Private/MaterialNodes/SGraphNodeMaterialBase.h
+++ b/Engine/Source/Editor/GraphEditor/Private/MaterialNodes/SGraphNodeMaterialBase.h
@@ -126,7 +126,7 @@ protected:
 	/** Up when shown, down when hidden */
 	const FSlateBrush* GetExpressionPreviewArrow() const;
 
-private:
+protected:
 	/** Slate viewport for rendering preview via custom slate element */
 	TSharedPtr<FPreviewViewport> PreviewViewport;
 
diff --git a/Engine/Source/Editor/GraphEditor/Private/MaterialNodes/SGraphNodeMaterialComposite.cpp b/Engine/Source/Editor/GraphEditor/Private/MaterialNodes/SGraphNodeMaterialComposite.cpp
new file mode 100644
--- /dev/null
+++ b/Engine/Source/Editor/GraphEditor/Private/MaterialNodes/SGraphNodeMaterialComposite.cpp
@@ -0,0 +1,292 @@
+// Copyright Epic Games, Inc. All Rights Reserved.
+
+#include "MaterialNodes/SGraphNodeMaterialComposite.h"
+#include "Materials/MaterialExpressionComposite.h"
+#include "MaterialGraph/MaterialGraphNode_Composite.h"
+
+#include "EdGraph/EdGraph.h"
+#include "Widgets/SBoxPanel.h"
+#include "Framework/Application/SlateApplication.h"
+#include "Widgets/Layout/SSpacer.h"
+#include "Widgets/Images/SImage.h"
+#include "Widgets/SToolTip.h"
+#include "Styling/CoreStyle.h"
+#include "GraphEditorSettings.h"
+#include "SCommentBubble.h"
+#include "SGraphPreviewer.h"
+#include "IDocumentationPage.h"
+#include "IDocumentation.h"
+#include "Widgets/Text/SInlineEditableTextBlock.h"
+
+void SGraphNodeMaterialComposite::Construct(const FArguments& InArgs, class UMaterialGraphNode_Composite* InNode)
+{
+	this->GraphNode = InNode;
+	this->CompositeNode = InNode;
+	
+	//@TODO: Add previews to composites, figure out multiple output UX (scroll?).
+	//this->MaterialNode = CompositeNode;
+
+	this->SetCursor(EMouseCursor::CardinalCross);
+
+	this->UpdateGraphNode();
+}
+
+void SGraphNodeMaterialComposite::MoveTo(const FVector2D& NewPosition, FNodeSet& NodeFilter)
+{
+	SGraphNode::MoveTo(NewPosition, NodeFilter);
+
+	CompositeNode->MaterialExpression->MaterialExpressionEditorX = CompositeNode->NodePosX;
+	CompositeNode->MaterialExpression->MaterialExpressionEditorY = CompositeNode->NodePosY;
+	CompositeNode->MaterialExpression->MarkPackageDirty();
+	CompositeNode->MaterialDirtyDelegate.ExecuteIfBound();
+}
+
+void SGraphNodeMaterialComposite::UpdateGraphNode()
+{
+	InputPins.Empty();
+	OutputPins.Empty();
+	
+	// Reset variables that are going to be exposed, in case we are refreshing an already setup node.
+	RightNodeBox.Reset();
+	LeftNodeBox.Reset();
+
+	SetupErrorReporting();
+	TSharedPtr<SNodeTitle> NodeTitle = SNew(SNodeTitle, GraphNode);
+
+	//
+	//             ______________________
+	//            |      TITLE AREA      |
+	//            +-------+------+-------+
+	//            | (>) L |      | R (>) |
+	//            | (>) E |      | I (>) |
+	//            | (>) F |      | G (>) |
+	//            | (>) T |      | H (>) |
+	//            |       |      | T (>) |
+	//            |_______|______|_______|
+	//
+	this->ContentScale.Bind( this, &SGraphNode::GetContentScale );
+	this->GetOrAddSlot( ENodeZone::Center )
+		.HAlign(HAlign_Center)
+		.VAlign(VAlign_Center)
+		[
+			SNew(SBorder)
+			.BorderImage( FEditorStyle::GetBrush( "Graph.CollapsedNode.Body" ) )
+			.Padding(0)
+			[
+				SNew(SOverlay)
+				+SOverlay::Slot()
+				[
+					SNew(SImage)
+					.Image( FEditorStyle::GetBrush("Graph.CollapsedNode.BodyColorSpill") )
+					.ColorAndOpacity( this, &SGraphNode::GetNodeTitleColor )
+				]
+				+SOverlay::Slot()
+				[
+					SNew(SVerticalBox)
+					+SVerticalBox::Slot()
+					.AutoHeight()
+					.HAlign(HAlign_Fill)
+					.VAlign(VAlign_Top)
+					[
+						SNew(SOverlay)
+						+SOverlay::Slot()
+						.HAlign(HAlign_Left)
+						.VAlign(VAlign_Center)
+						[
+							SNew(SBorder)
+							.BorderImage( FEditorStyle::GetBrush("NoBorder") )
+							.Padding( FMargin(10,5,30,3) )
+							[
+								SNew(SVerticalBox)
+								+SVerticalBox::Slot()
+								.AutoHeight()
+								.HAlign(HAlign_Fill)
+								.VAlign(VAlign_Top)
+								[
+									SNew(SVerticalBox)
+									+SVerticalBox::Slot()
+										.AutoHeight()
+									[
+										SAssignNew(InlineEditableText, SInlineEditableTextBlock)
+										.Style( FEditorStyle::Get(), "Graph.Node.NodeTitleInlineEditableText" )
+										.Text( NodeTitle.Get(), &SNodeTitle::GetHeadTitle )
+										.OnVerifyTextChanged(this, &SGraphNodeMaterialComposite::OnVerifyNameTextChanged)
+										.OnTextCommitted(this, &SGraphNodeMaterialComposite::OnNameTextCommited)
+										.IsReadOnly( this, &SGraphNodeMaterialComposite::IsNameReadOnly )
+										.IsSelected(this, &SGraphNodeMaterialComposite::IsSelectedExclusively)
+									]
+									+SVerticalBox::Slot()
+										.AutoHeight()
+									[
+										NodeTitle.ToSharedRef()
+									]
+								]
+								+SVerticalBox::Slot()
+								.AutoHeight()
+								.Padding(1.0f)
+								[
+									ErrorReporting->AsWidget()
+								]
+							]
+						]
+					]
+					+SVerticalBox::Slot()
+					.AutoHeight()
+					.HAlign(HAlign_Fill)
+					.VAlign(VAlign_Top)
+					[
+						CreateNodeBody()
+					]
+				]
+			]
+		];
+
+	// Create comment bubble
+	TSharedPtr<SCommentBubble> CommentBubble;
+	const FSlateColor CommentColor = GetDefault<UGraphEditorSettings>()->DefaultCommentNodeTitleColor;
+
+	SAssignNew( CommentBubble, SCommentBubble )
+	.GraphNode( GraphNode )
+	.Text( this, &SGraphNode::GetNodeComment )
+	.OnTextCommitted( this, &SGraphNode::OnCommentTextCommitted )
+	.ColorAndOpacity( CommentColor )
+	.AllowPinning( true )
+	.EnableTitleBarBubble( true )
+	.EnableBubbleCtrls( true )
+	.GraphLOD( this, &SGraphNode::GetCurrentLOD )
+	.IsGraphNodeHovered( this, &SGraphNode::IsHovered );
+
+	GetOrAddSlot( ENodeZone::TopCenter )
+	.SlotOffset( TAttribute<FVector2D>( CommentBubble.Get(), &SCommentBubble::GetOffset ))
+	.SlotSize( TAttribute<FVector2D>( CommentBubble.Get(), &SCommentBubble::GetSize ))
+	.AllowScaling( TAttribute<bool>( CommentBubble.Get(), &SCommentBubble::IsScalingAllowed ))
+	.VAlign( VAlign_Top )
+	[
+		CommentBubble.ToSharedRef()
+	];
+
+	CreatePinWidgets();
+}
+
+UEdGraph* SGraphNodeMaterialComposite::GetInnerGraph() const
+{
+	return CompositeNode->BoundGraph;
+}
+
+TSharedPtr<SToolTip> SGraphNodeMaterialComposite::GetComplexTooltip()
+{
+	if (UEdGraph* BoundGraph = GetInnerGraph())
+	{
+		struct LocalUtils
+		{
+			static bool IsInteractive()
+			{
+				const FModifierKeysState ModifierKeys = FSlateApplication::Get().GetModifierKeys();
+				return ( ModifierKeys.IsAltDown() && ModifierKeys.IsControlDown() );
+			}
+		};
+
+		TSharedPtr<SToolTip> FinalToolTip = nullptr;
+		TSharedPtr<SVerticalBox> Container = nullptr;
+		SAssignNew(FinalToolTip, SToolTip)
+		.IsInteractive_Static(&LocalUtils::IsInteractive)
+		[
+			SAssignNew(Container, SVerticalBox)
+			+SVerticalBox::Slot()
+			.AutoHeight()
+			[
+				SNew( STextBlock )
+				.Text(this, &SGraphNodeMaterialComposite::GetTooltipTextForNode)
+				.Font(FCoreStyle::GetDefaultFontStyle("Regular", 8))
+				.WrapTextAt(160.0f)
+			]
+			+SVerticalBox::Slot()
+			.AutoHeight()
+			[
+				// Create preview for the tooltip, make sure to disable state overlays to prevent
+				// PIE and read-only borders obscuring the graph
+				SNew(SGraphPreviewer, BoundGraph)
+				.CornerOverlayText(this, &SGraphNodeMaterialComposite::GetPreviewCornerText)
+				.ShowGraphStateOverlay(false)
+			]
+		];
+
+		// Check to see whether this node has a documentation excerpt. If it does, create a doc box for the tooltip
+		TSharedRef<IDocumentationPage> DocPage = IDocumentation::Get()->GetPage(GraphNode->GetDocumentationLink(), nullptr);
+		if(DocPage->HasExcerpt(GraphNode->GetDocumentationExcerptName()))
+		{
+			Container->AddSlot()
+			.AutoHeight()
+			.Padding(FMargin( 0.0f, 5.0f ))
+			[
+				IDocumentation::Get()->CreateToolTip(FText::FromString("Documentation"), nullptr, GraphNode->GetDocumentationLink(), GraphNode->GetDocumentationExcerptName())
+			];
+		}
+
+		return FinalToolTip;
+	}
+	else
+	{
+		return SNew(SToolTip)
+			[
+				SNew(SVerticalBox)
+				+SVerticalBox::Slot()
+				.AutoHeight()
+				[
+					SNew( STextBlock )
+					.Text(NSLOCTEXT("CompositeNode", "CompositeNodeInvalidGraphMessage", "ERROR: Invalid Graph"))
+					.Font(FCoreStyle::GetDefaultFontStyle("Regular", 8))
+					.WrapTextAt(160.0f)
+				]
+			];
+	}
+
+}
+
+FText SGraphNodeMaterialComposite::GetPreviewCornerText() const
+{
+	UEdGraph* BoundGraph = GetInnerGraph();
+	return FText::FromString(BoundGraph->GetName());
+}
+
+FText SGraphNodeMaterialComposite::GetTooltipTextForNode() const
+{
+	return GraphNode->GetTooltipText();
+}
+
+TSharedRef<SWidget> SGraphNodeMaterialComposite::CreateNodeBody()
+{
+	if( GraphNode && GraphNode->Pins.Num() > 0 )
+	{
+		// Create the input and output pin areas if there are pins
+		return SNew(SBorder)
+			.BorderImage( FEditorStyle::GetBrush("NoBorder") )
+			.HAlign(HAlign_Fill)
+			.VAlign(VAlign_Fill)
+			.Padding( FMargin(0,3) )
+			[
+				SNew(SHorizontalBox)
+				+SHorizontalBox::Slot()
+				.HAlign(HAlign_Left)
+				.FillWidth(1.0f)
+				[
+					// LEFT
+					SAssignNew(LeftNodeBox, SVerticalBox)
+				]
+				+SHorizontalBox::Slot()
+				.AutoWidth()
+				.HAlign(HAlign_Right)
+				[
+					// RIGHT
+					SAssignNew(RightNodeBox, SVerticalBox)
+				]
+			];
+	}
+	else
+	{
+		// Create a spacer so the node has some body to it
+		return SNew(SSpacer)
+			.Size(FVector2D(100.f, 50.f));
+	}
+}
+
diff --git a/Engine/Source/Editor/GraphEditor/Private/MaterialNodes/SGraphNodeMaterialComposite.h b/Engine/Source/Editor/GraphEditor/Private/MaterialNodes/SGraphNodeMaterialComposite.h
new file mode 100644
--- /dev/null
+++ b/Engine/Source/Editor/GraphEditor/Private/MaterialNodes/SGraphNodeMaterialComposite.h
@@ -0,0 +1,42 @@
+// Copyright Epic Games, Inc. All Rights Reserved.
+
+#pragma once
+
+#include "CoreMinimal.h"
+#include "Widgets/DeclarativeSyntaxSupport.h"
+#include "Widgets/SWidget.h"
+#include "MaterialNodes/SGraphNodeMaterialBase.h"
+
+class SToolTip;
+class UEdGraph;
+class UMaterialGraphNode_Composite;
+
+//@TODO: This class is mostly C&P from UK2Node_Composite, consolidate composites to not be BP specific.
+class GRAPHEDITOR_API SGraphNodeMaterialComposite : public SGraphNodeMaterialBase
+{
+public:
+	SLATE_BEGIN_ARGS(SGraphNodeMaterialComposite){}
+	SLATE_END_ARGS()
+
+	void Construct(const FArguments& InArgs, UMaterialGraphNode_Composite* InNode);
+
+	// SGraphNode interface
+	virtual void UpdateGraphNode() override;
+	virtual TSharedPtr<SToolTip> GetComplexTooltip() override;
+	// End of SGraphNode interface
+
+	// SNodePanel::SNode interface
+	virtual void MoveTo(const FVector2D& NewPosition, FNodeSet& NodeFilter) override;
+	// End of SNodePanel::SNode interface
+protected:
+	virtual UEdGraph* GetInnerGraph() const;
+
+private:
+	FText GetPreviewCornerText() const;
+	FText GetTooltipTextForNode() const;
+
+	TSharedRef<SWidget> CreateNodeBody();
+
+	/** Cached material graph node pointer to avoid casting */
+	UMaterialGraphNode_Composite* CompositeNode;
+};
diff --git a/Engine/Source/Editor/GraphEditor/Private/NodeFactory.cpp b/Engine/Source/Editor/GraphEditor/Private/NodeFactory.cpp
--- a/Engine/Source/Editor/GraphEditor/Private/NodeFactory.cpp
+++ b/Engine/Source/Editor/GraphEditor/Private/NodeFactory.cpp
@@ -6,6 +6,7 @@
 #include "InputCoreTypes.h"
 #include "Widgets/DeclarativeSyntaxSupport.h"
 #include "MaterialGraph/MaterialGraphNode_Comment.h"
+#include "MaterialGraph/MaterialGraphNode_Composite.h"
 #include "MaterialGraph/MaterialGraphNode_Base.h"
 #include "MaterialGraph/MaterialGraphNode.h"
 #include "MaterialGraph/MaterialGraphNode_Root.h"
@@ -75,6 +76,7 @@
 
 #include "MaterialNodes/SGraphNodeMaterialBase.h"
 #include "MaterialNodes/SGraphNodeMaterialComment.h"
+#include "MaterialNodes/SGraphNodeMaterialComposite.h"
 #include "MaterialNodes/SGraphNodeMaterialResult.h"
 #include "MaterialGraphNode_Knot.h"
 
@@ -124,7 +126,14 @@ TSharedPtr<SGraphNode> FNodeFactory::CreateNodeWidget(UEdGraphNode* InNode)
 		}
 		else if (UMaterialGraphNode* MaterialNode = Cast<UMaterialGraphNode>(InNode))
 		{
-			return SNew(SGraphNodeMaterialBase, MaterialNode);
+			if (UMaterialGraphNode_Composite* MaterialComposite = Cast<UMaterialGraphNode_Composite>(InNode))
+			{
+				return SNew(SGraphNodeMaterialComposite, MaterialComposite);
+			}
+			else
+			{
+				return SNew(SGraphNodeMaterialBase, MaterialNode);
+			}
 		}
 	}
 
diff --git a/Engine/Source/Editor/MaterialEditor/Private/FindInMaterial.cpp b/Engine/Source/Editor/MaterialEditor/Private/FindInMaterial.cpp
--- a/Engine/Source/Editor/MaterialEditor/Private/FindInMaterial.cpp
+++ b/Engine/Source/Editor/MaterialEditor/Private/FindInMaterial.cpp
@@ -261,7 +261,11 @@ void SFindInMaterial::MatchTokens(const TArray<FString> &Tokens)
 	RootSearchResult.Reset();
 
 	UEdGraph* Graph = MaterialEditorPtr.Pin()->Material->MaterialGraph;
+	MatchTokensInGraph(Graph, Tokens);
+}
 
+void SFindInMaterial::MatchTokensInGraph(const UEdGraph* Graph, const TArray<FString>& Tokens)
+{
 	if (Graph == NULL)
 	{
 		return;
@@ -339,6 +343,11 @@ void SFindInMaterial::MatchTokens(const TArray<FString> &Tokens)
 			ItemsFound.Add(NodeResult);
 		}
 	}
+
+	for (const UEdGraph* Subgraph : Graph->SubGraphs)
+	{
+		MatchTokensInGraph(Subgraph, Tokens);
+	}
 }
 
 TSharedRef<ITableRow> SFindInMaterial::OnGenerateRow(FSearchResult InItem, const TSharedRef<STableViewBase>& OwnerTable)
diff --git a/Engine/Source/Editor/MaterialEditor/Private/MaterialEditingLibrary.cpp b/Engine/Source/Editor/MaterialEditor/Private/MaterialEditingLibrary.cpp
--- a/Engine/Source/Editor/MaterialEditor/Private/MaterialEditingLibrary.cpp
+++ b/Engine/Source/Editor/MaterialEditor/Private/MaterialEditingLibrary.cpp
@@ -459,6 +459,43 @@ UMaterialExpression* UMaterialEditingLibrary::CreateMaterialExpression(UMaterial
 	return CreateMaterialExpressionEx(Material, nullptr, ExpressionClass, nullptr, NodePosX, NodePosY);
 }
 
+UMaterialExpression* UMaterialEditingLibrary::DuplicateMaterialExpression(UMaterial* Material, UMaterialFunction* MaterialFunction, UMaterialExpression* Expression)
+{
+	UMaterialExpression* NewExpression = nullptr;
+	if (Material || MaterialFunction)
+	{
+		UObject* ExpressionOuter = Material;
+		if (MaterialFunction)
+		{
+			ExpressionOuter = MaterialFunction;
+		}
+
+		NewExpression = DuplicateObject(Expression, ExpressionOuter);
+
+		if (Material)
+		{
+			Material->Expressions.Add(NewExpression);
+			NewExpression->Material = Material;
+		}
+
+		if (MaterialFunction && !Material)
+		{
+			MaterialFunction->FunctionExpressions.Add(NewExpression);
+		}
+
+		// Create a GUID for the node
+		NewExpression->UpdateMaterialExpressionGuid(true, true);
+
+		if (Material)
+		{
+			Material->AddExpressionParameter(NewExpression, Material->EditorParameters);
+		}
+
+		NewExpression->MarkPackageDirty();
+	}
+	return NewExpression;
+}
+
 UMaterialExpression* UMaterialEditingLibrary::CreateMaterialExpressionInFunction(UMaterialFunction* MaterialFunction, TSubclassOf<UMaterialExpression> ExpressionClass, int32 NodePosX, int32 NodePosY)
 {
 	return CreateMaterialExpressionEx(nullptr, MaterialFunction, ExpressionClass, nullptr, NodePosX, NodePosY);
diff --git a/Engine/Source/Editor/MaterialEditor/Private/MaterialEditor.cpp b/Engine/Source/Editor/MaterialEditor/Private/MaterialEditor.cpp
--- a/Engine/Source/Editor/MaterialEditor/Private/MaterialEditor.cpp
+++ b/Engine/Source/Editor/MaterialEditor/Private/MaterialEditor.cpp
@@ -16,6 +16,7 @@
 #include "Widgets/Input/SCheckBox.h"
 #include "EditorStyleSet.h"
 #include "EdGraph/EdGraph.h"
+#include "WorkflowOrientedApp/WorkflowUObjectDocuments.h"
 #include "MaterialGraph/MaterialGraph.h"
 #include "MaterialGraph/MaterialGraphNode_Comment.h"
 #include "Editor/UnrealEdEngine.h"
@@ -23,6 +24,8 @@
 #include "Preferences/MaterialEditorOptions.h"
 #include "MaterialGraph/MaterialGraphNode.h"
 #include "MaterialGraph/MaterialGraphNode_Root.h"
+#include "MaterialGraph/MaterialGraphNode_Composite.h"
+#include "MaterialGraph/MaterialGraphNode_PinBase.h"
 #include "MaterialGraph/MaterialGraphSchema.h"
 #include "MaterialEditor/PreviewMaterial.h"
 #include "ThumbnailRendering/SceneThumbnailInfoWithPrimitive.h"
@@ -45,6 +48,7 @@
 #include "Materials/MaterialExpressionBreakMaterialAttributes.h"
 #include "Materials/MaterialExpressionCollectionParameter.h"
 #include "Materials/MaterialExpressionComment.h"
+#include "Materials/MaterialExpressionComposite.h"
 #include "Materials/MaterialExpressionComponentMask.h"
 #include "Materials/MaterialExpressionConstant.h"
 #include "Materials/MaterialExpressionConstant2Vector.h"
@@ -56,9 +60,11 @@
 #include "Materials/MaterialExpressionFunctionOutput.h"
 #include "Materials/MaterialExpressionMaterialAttributeLayers.h"
 #include "Materials/MaterialExpressionParameter.h"
+#include "Materials/MaterialExpressionPinBase.h"
 #include "Materials/MaterialExpressionTextureBase.h"
 #include "Materials/MaterialExpressionTextureSample.h"
 #include "Materials/MaterialExpressionParticleSubUV.h"
+#include "Materials/MaterialExpressionReroute.h"
 #include "Materials/MaterialExpressionRuntimeVirtualTextureSample.h"
 #include "Materials/MaterialExpressionRuntimeVirtualTextureSampleParameter.h"
 #include "Materials/MaterialExpressionScalarParameter.h"
@@ -83,6 +89,7 @@
 #include "MaterialCompiler.h"
 #include "EditorSupportDelegates.h"
 #include "AssetRegistryModule.h"
+#include "Tabs/MaterialEditorTabFactories.h"
 #include "IAssetTools.h"
 #include "IAssetTypeActions.h"
 #include "AssetToolsModule.h"
@@ -130,6 +137,8 @@
 #include "Materials/MaterialExpressionMaterialLayerOutput.h"
 
 #include "MaterialStats.h"
+#include "MaterialEditorTabs.h"
+#include "MaterialEditorModes.h"
 #include "Materials/MaterialExpression.h"
 
 #include "SMaterialParametersOverviewWidget.h"
@@ -153,15 +162,6 @@ static TAutoConsoleVariable<int32> CVarMaterialEdUseDevShaders(
 	TEXT("Toggles whether the material editor will use shaders that include extra overhead incurred by the editor. Material editor must be re-opened if changed at runtime."),
 	ECVF_RenderThreadSafe);
 
-const FName FMaterialEditor::PreviewTabId( TEXT( "MaterialEditor_Preview" ) );
-const FName FMaterialEditor::GraphCanvasTabId( TEXT( "MaterialEditor_GraphCanvas" ) );
-const FName FMaterialEditor::PropertiesTabId( TEXT( "MaterialEditor_MaterialProperties" ) );
-const FName FMaterialEditor::PaletteTabId( TEXT( "MaterialEditor_Palette" ) );
-const FName FMaterialEditor::FindTabId( TEXT( "MaterialEditor_Find" ) );
-const FName FMaterialEditor::PreviewSettingsTabId( TEXT ("MaterialEditor_PreviewSettings" ) );
-const FName FMaterialEditor::ParameterDefaultsTabId(TEXT("MaterialEditor_ParameterDefaults"));
-const FName FMaterialEditor::CustomPrimitiveTabId(TEXT("MaterialEditor_CustomPrimitiveData"));
-const FName FMaterialEditor::LayerPropertiesTabId(TEXT("MaterialInstanceEditor_MaterialLayerProperties"));
 ///////////////////////////
 // FMatExpressionPreview //
 ///////////////////////////
@@ -253,55 +253,49 @@ void FMatExpressionPreview::NotifyCompilationFinished()
 // FMaterialEditor //
 /////////////////////
 
-void FMaterialEditor::RegisterTabSpawners(const TSharedRef<class FTabManager>& InTabManager)
+void FMaterialEditor::RegisterToolbarTab(const TSharedRef<class FTabManager>& InTabManager)
 {
+	FAssetEditorToolkit::RegisterTabSpawners(InTabManager);
+
 	WorkspaceMenuCategory = InTabManager->AddLocalWorkspaceMenuCategory(LOCTEXT("WorkspaceMenu_MaterialEditor", "Material Editor"));
 	auto WorkspaceMenuCategoryRef = WorkspaceMenuCategory.ToSharedRef();
 
-	FAssetEditorToolkit::RegisterTabSpawners(InTabManager);
-	
-	InTabManager->RegisterTabSpawner( PreviewTabId, FOnSpawnTab::CreateSP(this, &FMaterialEditor::SpawnTab_Preview) )
+	InTabManager->RegisterTabSpawner(FMaterialEditorTabs::PreviewTabId, FOnSpawnTab::CreateSP(this, &FMaterialEditor::SpawnTab_Preview) )
 		.SetDisplayName( LOCTEXT("ViewportTab", "Viewport") )
 		.SetGroup( WorkspaceMenuCategoryRef )
 		.SetIcon(FSlateIcon(FEditorStyle::GetStyleSetName(), "LevelEditor.Tabs.Viewports"));
-	
-	InTabManager->RegisterTabSpawner( GraphCanvasTabId, FOnSpawnTab::CreateSP(this, &FMaterialEditor::SpawnTab_GraphCanvas) )
-		.SetDisplayName( LOCTEXT("GraphCanvasTab", "Graph") )
-		.SetGroup( WorkspaceMenuCategoryRef )
-		.SetIcon(FSlateIcon(FEditorStyle::GetStyleSetName(), "GraphEditor.EventGraph_16x"));
-	
-	InTabManager->RegisterTabSpawner( PropertiesTabId, FOnSpawnTab::CreateSP(this, &FMaterialEditor::SpawnTab_MaterialProperties) )
+
+	InTabManager->RegisterTabSpawner(FMaterialEditorTabs::PropertiesTabId, FOnSpawnTab::CreateSP(this, &FMaterialEditor::SpawnTab_MaterialProperties) )
 		.SetDisplayName( LOCTEXT("DetailsTab", "Details") )
 		.SetGroup( WorkspaceMenuCategoryRef )
 		.SetIcon(FSlateIcon(FEditorStyle::GetStyleSetName(), "LevelEditor.Tabs.Details"));
 
-	InTabManager->RegisterTabSpawner( PaletteTabId, FOnSpawnTab::CreateSP(this, &FMaterialEditor::SpawnTab_Palette) )
+	InTabManager->RegisterTabSpawner(FMaterialEditorTabs::PaletteTabId, FOnSpawnTab::CreateSP(this, &FMaterialEditor::SpawnTab_Palette) )
 		.SetDisplayName( LOCTEXT("PaletteTab", "Palette") )
 		.SetGroup( WorkspaceMenuCategoryRef )
 		.SetIcon(FSlateIcon(FEditorStyle::GetStyleSetName(), "Kismet.Tabs.Palette"));
 
-	InTabManager->RegisterTabSpawner(FindTabId, FOnSpawnTab::CreateSP(this, &FMaterialEditor::SpawnTab_Find))
+	InTabManager->RegisterTabSpawner(FMaterialEditorTabs::FindTabId, FOnSpawnTab::CreateSP(this, &FMaterialEditor::SpawnTab_Find))
 		.SetDisplayName(LOCTEXT("FindTab", "Find Results"))
 		.SetGroup( WorkspaceMenuCategoryRef )
 		.SetIcon(FSlateIcon(FEditorStyle::GetStyleSetName(), "Kismet.Tabs.FindResults"));
 
-	InTabManager->RegisterTabSpawner(PreviewSettingsTabId, FOnSpawnTab::CreateSP(this, &FMaterialEditor::SpawnTab_PreviewSettings))
+	InTabManager->RegisterTabSpawner(FMaterialEditorTabs::PreviewSettingsTabId, FOnSpawnTab::CreateSP(this, &FMaterialEditor::SpawnTab_PreviewSettings))
 		.SetDisplayName( LOCTEXT("PreviewSceneSettingsTab", "Preview Scene Settings") )
 		.SetGroup( WorkspaceMenuCategoryRef )
 		.SetIcon(FSlateIcon(FEditorStyle::GetStyleSetName(), "LevelEditor.Tabs.Details"));
 
-
-	InTabManager->RegisterTabSpawner(ParameterDefaultsTabId, FOnSpawnTab::CreateSP(this, &FMaterialEditor::SpawnTab_ParameterDefaults))
+	InTabManager->RegisterTabSpawner(FMaterialEditorTabs::ParameterDefaultsTabId, FOnSpawnTab::CreateSP(this, &FMaterialEditor::SpawnTab_ParameterDefaults))
 		.SetDisplayName(LOCTEXT("ParameterDefaultsTab", "Parameter Defaults"))
 		.SetGroup(WorkspaceMenuCategoryRef)
 		.SetIcon(FSlateIcon(FEditorStyle::GetStyleSetName(), "LevelEditor.Tabs.Details"));
 
-	InTabManager->RegisterTabSpawner(CustomPrimitiveTabId, FOnSpawnTab::CreateSP(this, &FMaterialEditor::SpawnTab_CustomPrimitiveData))
+	InTabManager->RegisterTabSpawner(FMaterialEditorTabs::CustomPrimitiveTabId, FOnSpawnTab::CreateSP(this, &FMaterialEditor::SpawnTab_CustomPrimitiveData))
 		.SetDisplayName(LOCTEXT("CustomPrimitiveTab", "Custom Primitive Data"))
 		.SetGroup(WorkspaceMenuCategoryRef)
 		.SetIcon(FSlateIcon(FEditorStyle::GetStyleSetName(), "LevelEditor.Tabs.Details"));
 
-	InTabManager->RegisterTabSpawner(LayerPropertiesTabId, FOnSpawnTab::CreateSP(this, &FMaterialEditor::SpawnTab_LayerProperties))
+	InTabManager->RegisterTabSpawner(FMaterialEditorTabs::LayerPropertiesTabId, FOnSpawnTab::CreateSP(this, &FMaterialEditor::SpawnTab_LayerProperties))
 		.SetDisplayName(LOCTEXT("LayerPropertiesTab", "Layer Parameters"))
 		.SetGroup(WorkspaceMenuCategoryRef)
 		.SetIcon(FSlateIcon(FEditorStyle::GetStyleSetName(), "LevelEditor.Tabs.Layers"));
@@ -311,20 +305,24 @@ void FMaterialEditor::RegisterTabSpawners(const TSharedRef<class FTabManager>& I
 	OnRegisterTabSpawners().Broadcast(InTabManager);
 }
 
+void FMaterialEditor::RegisterTabSpawners(const TSharedRef<class FTabManager>& InTabManager)
+{
+	DocumentManager->SetTabManager(InTabManager);
+	FWorkflowCentricApplication::RegisterTabSpawners(InTabManager);
+}
 
 void FMaterialEditor::UnregisterTabSpawners(const TSharedRef<class FTabManager>& InTabManager)
 {
 	FAssetEditorToolkit::UnregisterTabSpawners(InTabManager);
 
-	InTabManager->UnregisterTabSpawner( PreviewTabId );
-	InTabManager->UnregisterTabSpawner( GraphCanvasTabId );
-	InTabManager->UnregisterTabSpawner( PropertiesTabId );
-	InTabManager->UnregisterTabSpawner( PaletteTabId );
-	InTabManager->UnregisterTabSpawner( FindTabId );
-	InTabManager->UnregisterTabSpawner( PreviewSettingsTabId );
-	InTabManager->UnregisterTabSpawner(ParameterDefaultsTabId);
-	InTabManager->UnregisterTabSpawner(CustomPrimitiveTabId);
-	InTabManager->UnregisterTabSpawner( LayerPropertiesTabId );
+	InTabManager->UnregisterTabSpawner(FMaterialEditorTabs::PreviewTabId);
+	InTabManager->UnregisterTabSpawner(FMaterialEditorTabs::PropertiesTabId);
+	InTabManager->UnregisterTabSpawner(FMaterialEditorTabs::PaletteTabId);
+	InTabManager->UnregisterTabSpawner(FMaterialEditorTabs::FindTabId);
+	InTabManager->UnregisterTabSpawner(FMaterialEditorTabs::PreviewSettingsTabId);
+	InTabManager->UnregisterTabSpawner(FMaterialEditorTabs::ParameterDefaultsTabId);
+	InTabManager->UnregisterTabSpawner(FMaterialEditorTabs::CustomPrimitiveTabId);
+	InTabManager->UnregisterTabSpawner(FMaterialEditorTabs::LayerPropertiesTabId);
 
 	MaterialStatsManager->UnregisterTabs();
 
@@ -463,82 +461,30 @@ void FMaterialEditor::InitMaterialEditor( const EToolkitMode::Type Mode, const T
 	BindCommands();
 	RegisterToolBar();
 
-	const TSharedRef<FTabManager::FLayout> StandaloneDefaultLayout = FTabManager::NewLayout("Standalone_MaterialEditor_Layout_v10")
-	->AddArea
-	(
-		FTabManager::NewPrimaryArea() ->SetOrientation(Orient_Vertical)
-		->Split
-		(
-			FTabManager::NewStack()
-			->SetSizeCoefficient(0.1f)
-			->SetHideTabWell( true )
-			->AddTab(GetToolbarTabId(), ETabState::OpenedTab)
-		)
-		->Split
-		(
-			FTabManager::NewSplitter() ->SetOrientation(Orient_Horizontal) ->SetSizeCoefficient(0.9f)
-			->Split
-			(
-				FTabManager::NewSplitter() ->SetOrientation(Orient_Vertical) ->SetSizeCoefficient(0.2f)
-				->Split
-				(
-					FTabManager::NewStack()
-					->SetHideTabWell( true )
-					->AddTab( PreviewTabId, ETabState::OpenedTab )
-				)
-				->Split
-				(
-					FTabManager::NewStack()
-					->AddTab( PropertiesTabId, ETabState::OpenedTab )
-					->AddTab( PreviewSettingsTabId, ETabState::ClosedTab )
-					->AddTab(ParameterDefaultsTabId, ETabState::OpenedTab)
-					->AddTab(CustomPrimitiveTabId, ETabState::ClosedTab)
-					->AddTab( LayerPropertiesTabId, ETabState::ClosedTab )
-					->SetForegroundTab( PropertiesTabId )
-				)
-			)
-			->Split
-			(
-				FTabManager::NewSplitter() ->SetOrientation( Orient_Vertical )
-				->SetSizeCoefficient(0.80f)
-				->Split
-				(
-					FTabManager::NewStack() 
-					->SetSizeCoefficient(0.8f)
-					->SetHideTabWell( true )
-					->AddTab( GraphCanvasTabId, ETabState::OpenedTab )
-				)
-				->Split
-				(
-					FTabManager::NewStack()
-					->SetSizeCoefficient( 0.20f )
-						->AddTab(FMaterialStats::GetGridStatsTabName(), ETabState::ClosedTab)
-						->AddTab(FMaterialStats::GetGridOldStatsTabName(), ETabState::ClosedTab)
-					->AddTab( FindTabId, ETabState::ClosedTab )
-				)
-			)
-			->Split
-			(
-				FTabManager::NewSplitter() ->SetOrientation(Orient_Horizontal) ->SetSizeCoefficient(0.2f)
-				->Split
-				(
-					FTabManager::NewStack()
-					->AddTab( PaletteTabId, ETabState::OpenedTab )
-				)
-			)
-		)
-	);
+	TSharedPtr<FMaterialEditor> ThisPtr(SharedThis(this));
+	DocumentManager->Initialize(ThisPtr);
+
+	// Register the document factories
+	{
+		TSharedRef<FDocumentTabFactory> GraphEditorFactory = MakeShareable(new FMaterialGraphEditorSummoner(ThisPtr,
+			FMaterialGraphEditorSummoner::FOnCreateGraphEditorWidget::CreateSP(this, &FMaterialEditor::CreateGraphEditorWidget)
+		));
+
+		// Also store off a reference to the grapheditor factory so we can find all the tabs spawned by it later.
+		GraphEditorTabFactoryPtr = GraphEditorFactory;
+		DocumentManager->RegisterDocumentFactory(GraphEditorFactory);
+	}
 
 	const bool bCreateDefaultStandaloneMenu = true;
 	const bool bCreateDefaultToolbar = true;
 
 	// Add the preview material to the objects being edited, so that we can find this editor from the temporary material graph
+	const TSharedRef<FTabManager::FLayout> DummyLayout = FTabManager::NewLayout("NullLayout")->AddArea(FTabManager::NewPrimaryArea());
 	TArray< UObject* > ObjectsToEdit;
 	ObjectsToEdit.Add(ObjectToEdit);
 	ObjectsToEdit.Add(Material);
 	ObjectsToEdit.Add(MaterialEditorInstance);
-	FAssetEditorToolkit::InitAssetEditor( Mode, InitToolkitHost, MaterialEditorAppIdentifier, StandaloneDefaultLayout, bCreateDefaultStandaloneMenu, bCreateDefaultToolbar, ObjectsToEdit, false );
-
+	InitAssetEditor( Mode, InitToolkitHost, MaterialEditorAppIdentifier, DummyLayout, bCreateDefaultStandaloneMenu, bCreateDefaultToolbar, ObjectsToEdit, false );
 	AddMenuExtender(GetMenuExtensibilityManager()->GetAllExtenders(GetToolkitCommands(), GetEditingObjects()));
 
 	IMaterialEditorModule* MaterialEditorModule = &FModuleManager::LoadModuleChecked<IMaterialEditorModule>( "MaterialEditor" );
@@ -547,6 +493,13 @@ void FMaterialEditor::InitMaterialEditor( const EToolkitMode::Type Mode, const T
 	ExtendToolbar();
 	RegenerateMenusAndToolbars();
 
+	{
+		AddApplicationMode(
+			FMaterialEditorApplicationModes::StandardMaterialEditorMode,
+			MakeShareable(new FMaterialEditorApplicationMode(SharedThis(this))));
+		SetCurrentMode(FMaterialEditorApplicationModes::StandardMaterialEditorMode);
+	}
+
 	// @todo toolkit world centric editing
 	/*if( IsWorldCentricAssetEditor() )
 	{
@@ -589,7 +542,7 @@ void FMaterialEditor::InitMaterialEditor( const EToolkitMode::Type Mode, const T
 		if (Material->Expressions.Num() == 0)
 		{
 			// If this is an empty function, create an output by default and start previewing it
-			if (GraphEditor.IsValid())
+			if (FocusedGraphEdPtr.IsValid())
 			{
 				check(!bMaterialDirty);
 				FVector2D OutputPlacement = FVector2D(200, 300);
@@ -763,6 +716,11 @@ void FMaterialEditor::UpdateGenerator()
 	}
 }
 
+void FMaterialEditor::NavigateTab(FDocumentTracker::EOpenDocumentCause InCause)
+{
+	OpenDocument(nullptr, InCause);
+}
+
 FMaterialEditor::FMaterialEditor()
 	: bMaterialDirty(false)
 	, bStatsFromPreviewMaterial(false)
@@ -788,6 +746,7 @@ FMaterialEditor::FMaterialEditor()
 	, ToolBarExtensibilityManager(new FExtensibilityManager)
 	, MaterialEditorInstance(nullptr)
 {
+	DocumentManager = MakeShareable(new FDocumentTracker);
 }
 
 FMaterialEditor::~FMaterialEditor()
@@ -927,6 +886,449 @@ void FMaterialEditor::UpdatePreviewViewportsVisibility()
 	}
 }
 
+void FMaterialEditor::CollapseNodesIntoGraph(UEdGraphNode* InGatewayNode, UMaterialGraphNode* InEntryNode, UMaterialGraphNode* InResultNode, UEdGraph* InSourceGraph, UEdGraph* InDestinationGraph, TSet<UEdGraphNode*>& InCollapsableNodes)
+{
+	const UMaterialGraphSchema* MaterialSchema = GetDefault<UMaterialGraphSchema>(); // ?
+
+	// Keep track of the statistics of the node positions so the new nodes can be located reasonably well
+	float SumNodeX = 0.0f;
+	float SumNodeY = 0.0f;
+	float MinNodeX = 1e9f;
+	float MinNodeY = 1e9f;
+	float MaxNodeX = -1e9f;
+	float MaxNodeY = -1e9f;
+	
+	// Move the nodes over, which may create cross-graph references that we need fix up ASAP
+	for (TSet<UEdGraphNode*>::TConstIterator NodeIt(InCollapsableNodes); NodeIt; ++NodeIt)
+	{
+		UEdGraphNode* Node = *NodeIt;
+		Node->Modify();
+
+		// Update stats
+		SumNodeX += Node->NodePosX;
+		SumNodeY += Node->NodePosY;
+		MinNodeX = FMath::Min<float>(MinNodeX, Node->NodePosX);
+		MinNodeY = FMath::Min<float>(MinNodeY, Node->NodePosY);
+		MaxNodeX = FMath::Max<float>(MaxNodeX, Node->NodePosX);
+		MaxNodeY = FMath::Max<float>(MaxNodeY, Node->NodePosY);
+
+		// Move the node over
+		InSourceGraph->Nodes.Remove(Node);
+		InDestinationGraph->Nodes.Add(Node);
+		Node->Rename(/*NewName=*/ NULL, /*NewOuter=*/ InDestinationGraph);
+
+		// Move the sub-graph to the new graph
+		if(UMaterialGraphNode_Composite* Composite = Cast<UMaterialGraphNode_Composite>(Node))
+		{
+			InSourceGraph->SubGraphs.Remove(Composite->BoundGraph);
+			InDestinationGraph->SubGraphs.Add(Composite->BoundGraph);
+			Composite->BoundGraph->SubgraphExpression = CastChecked<UMaterialGraphNode_Composite>(InGatewayNode)->MaterialExpression;
+		}
+
+		// Mark the node's expression as owned by the gateway node's expression
+		UMaterialGraphNode* GatewayMaterialNode = CastChecked<UMaterialGraphNode>(InGatewayNode);
+		if (UMaterialGraphNode* MaterialNode = Cast<UMaterialGraphNode>(Node))
+		{
+			MaterialNode->MaterialExpression->SubgraphExpression = GatewayMaterialNode->MaterialExpression;
+		}
+		else if (UMaterialGraphNode_Comment* CommentNode = Cast<UMaterialGraphNode_Comment>(Node))
+		{
+			CommentNode->MaterialExpressionComment->SubgraphExpression = GatewayMaterialNode->MaterialExpression;
+		}
+
+		TArray<UEdGraphPin*> OutputGatewayExecPins;
+
+		// Find cross-graph links
+		for (int32 PinIndex = 0; PinIndex < Node->Pins.Num(); ++PinIndex)
+		{
+			UEdGraphPin* LocalPin = Node->Pins[PinIndex];
+
+			bool bIsGatewayPin = false;
+			if(LocalPin->LinkedTo.Num())
+			{
+				for (int32 LinkIndex = 0; LinkIndex < LocalPin->LinkedTo.Num(); ++LinkIndex)
+				{
+					UEdGraphPin* TrialPin = LocalPin->LinkedTo[LinkIndex];
+					if (!InCollapsableNodes.Contains(TrialPin->GetOwningNode()))
+					{
+						bIsGatewayPin = true;
+						break;
+					}
+				}
+			}
+
+			// Thunk cross-graph links thru the gateway
+			if (bIsGatewayPin)
+			{
+				// Local port is either the entry or the result node in the collapsed graph
+				// Remote port is the node placed in the source graph
+				UMaterialGraphNode* LocalPort = (LocalPin->Direction == EGPD_Input) ? InEntryNode : InResultNode;
+
+				// Add a new pin to the entry/exit node and to the composite node
+				UEdGraphPin* LocalPortPin = NULL;
+				UEdGraphPin* RemotePortPin = NULL;
+
+				if(LocalPin->LinkedTo[0]->GetOwningNode() != InEntryNode)
+				{
+					const FName UniquePortName = InGatewayNode->CreateUniquePinName(LocalPin->PinName);
+
+					if(!RemotePortPin && !LocalPortPin)
+					{
+						FEdGraphPinType PinType = LocalPin->PinType;
+						RemotePortPin = InGatewayNode->CreatePin(LocalPin->Direction, PinType, UniquePortName);
+						LocalPortPin = LocalPort->CreatePin((LocalPin->Direction == EGPD_Input) ? EGPD_Output : EGPD_Input, PinType, UniquePortName);
+
+						// Create reroute expressions / pins for the pinbase.
+						UMaterialGraphNode_Composite* CompositeNode = CastChecked<UMaterialGraphNode_Composite>(InGatewayNode);
+						UMaterialExpression* GatewayRerouteExpression = UMaterialEditingLibrary::CreateMaterialExpressionEx(Material, MaterialFunction, UMaterialExpressionReroute::StaticClass());
+						UMaterialExpressionReroute* GatewayReroutePin = CastChecked<UMaterialExpressionReroute>(GatewayRerouteExpression);
+						UMaterialExpressionPinBase* PinBase = CastChecked<UMaterialExpressionPinBase>(LocalPort->MaterialExpression);
+
+						GatewayRerouteExpression->SubgraphExpression = CompositeNode->MaterialExpression;
+						PinBase->ReroutePins.Add(FCompositeReroute{UniquePortName, GatewayReroutePin});
+						PinBase->Modify();
+					}
+				}
+
+				check(LocalPortPin);
+				check(RemotePortPin);
+
+				LocalPin->Modify();
+
+				// Route the links
+				for (int32 LinkIndex = 0; LinkIndex < LocalPin->LinkedTo.Num(); ++LinkIndex)
+				{
+					UEdGraphPin* RemotePin = LocalPin->LinkedTo[LinkIndex];
+					RemotePin->Modify();
+
+					if (!InCollapsableNodes.Contains(RemotePin->GetOwningNode()) && RemotePin->GetOwningNode() != InEntryNode && RemotePin->GetOwningNode() != InResultNode)
+					{
+						// Fix up the remote pin
+						RemotePin->LinkedTo.Remove(LocalPin);
+						RemotePin->MakeLinkTo(RemotePortPin);
+
+						// The Entry Node only supports a single link, so if we made links above
+						// we need to break them now, to make room for the new link.
+						if (LocalPort == InEntryNode)
+						{
+							LocalPortPin->BreakAllPinLinks();
+						}
+
+						// Fix up the local pin
+						LocalPin->LinkedTo.Remove(RemotePin);
+						--LinkIndex;
+						LocalPin->MakeLinkTo(LocalPortPin);
+					}
+				}
+			}
+		}
+	}
+
+	// Reposition the newly created nodes
+	const int32 NumNodes = InCollapsableNodes.Num();
+	const float CenterX = NumNodes == 0 ? SumNodeX : SumNodeX / NumNodes;
+	const float CenterY = NumNodes == 0 ? SumNodeY : SumNodeY / NumNodes;
+	const float MinusOffsetX = 160.0f; //@TODO: Random magic numbers
+	const float PlusOffsetX = 300.0f;
+
+	// Put the gateway node at the center of the empty space in the old graph
+	InGatewayNode->NodePosX = CenterX;
+	InGatewayNode->NodePosY = CenterY;
+	InGatewayNode->SnapToGrid(SNodePanel::GetSnapGridSize());
+
+	if (UMaterialGraphNode_Composite* CompositeNode = Cast<UMaterialGraphNode_Composite>(InGatewayNode))
+	{
+		CompositeNode->MaterialExpression->MaterialExpressionEditorX = InGatewayNode->NodePosX;
+		CompositeNode->MaterialExpression->MaterialExpressionEditorY = InGatewayNode->NodePosY;
+	}
+
+	// Put the entry and exit nodes on either side of the nodes in the new graph
+	if (NumNodes != 0)
+	{
+		InEntryNode->NodePosX = MinNodeX - MinusOffsetX;
+		InEntryNode->NodePosY = CenterY;
+		InEntryNode->SnapToGrid(SNodePanel::GetSnapGridSize());
+		InEntryNode->MaterialExpression->MaterialExpressionEditorX = InEntryNode->NodePosX;
+		InEntryNode->MaterialExpression->MaterialExpressionEditorY = InEntryNode->NodePosY;
+
+		InResultNode->NodePosX = MaxNodeX + PlusOffsetX;
+		InResultNode->NodePosY = CenterY;
+		InResultNode->SnapToGrid(SNodePanel::GetSnapGridSize());
+		InResultNode->MaterialExpression->MaterialExpressionEditorX = InResultNode->NodePosX;
+		InResultNode->MaterialExpression->MaterialExpressionEditorY = InResultNode->NodePosY;
+	}
+}
+
+void FMaterialEditor::CollapseNodes(TSet<UEdGraphNode*>& InCollapsableNodes)
+{
+	TSharedPtr<SGraphEditor> FocusedGraphEd = FocusedGraphEdPtr.Pin();
+	if (!FocusedGraphEd.IsValid())
+	{
+		return;
+	}
+
+	UEdGraph* SourceGraph = FocusedGraphEd->GetCurrentGraph();
+	SourceGraph->Modify();
+
+	// Create the composite node that will serve as the gateway into the subgraph
+	UMaterialGraphNode_Composite* GatewayNode = NULL;
+	{
+		GatewayNode = Cast<UMaterialGraphNode_Composite>(FMaterialGraphSchemaAction_NewComposite::SpawnNode(SourceGraph, FVector2D(0, 0)));
+		GatewayNode->bCanRenameNode = true;
+		check(GatewayNode);
+	}
+
+	UEdGraph* DestinationGraph = GatewayNode->BoundGraph;
+	UMaterialExpressionComposite* CompositeExpression = CastChecked<UMaterialExpressionComposite>(GatewayNode->MaterialExpression);
+
+	CollapseNodesIntoGraph(GatewayNode, 
+		Cast<UMaterialGraphNode>(CompositeExpression->InputExpressions->GraphNode),
+		Cast<UMaterialGraphNode>(CompositeExpression->OutputExpressions->GraphNode),
+		SourceGraph, 
+		DestinationGraph, 
+		InCollapsableNodes);
+
+	UpdateMaterialAfterGraphChange();
+
+	// Now that the expressions are updated, reconstruct the nodes
+	// Need to do this to prevent user from copying a freshly collapsed node not built from its expression
+	GatewayNode->ReconstructNode();
+	Cast<UMaterialGraphNode>(CompositeExpression->InputExpressions->GraphNode)->ReconstructNode();
+	Cast<UMaterialGraphNode>(CompositeExpression->OutputExpressions->GraphNode)->ReconstructNode();
+}
+
+void FMaterialEditor::ExpandNode(UEdGraphNode* InNodeToExpand, UEdGraph* InSourceGraph, TSet<UEdGraphNode*>& OutExpandedNodes)
+{
+	UEdGraph* DestinationGraph = InNodeToExpand->GetGraph();
+	UEdGraph* SourceGraph = InSourceGraph;
+	check(SourceGraph);
+
+	// Mark all edited objects so they will appear in the transaction record if needed.
+	DestinationGraph->Modify();
+	SourceGraph->Modify();
+	InNodeToExpand->Modify();
+
+	UEdGraphNode* Entry = nullptr;
+	UEdGraphNode* Result = nullptr;
+
+	const bool bIsCollapsedGraph = InNodeToExpand->IsA<UMaterialGraphNode_Composite>();
+
+	MoveNodesToGraph(SourceGraph->Nodes, DestinationGraph, OutExpandedNodes, &Entry, &Result, bIsCollapsedGraph);
+	CollapseGatewayNode(InNodeToExpand, Entry, Result, &OutExpandedNodes);
+
+	bool bPreviewExpressionDeleted = false;
+
+	if (Entry)
+	{
+		UMaterialExpressionPinBase* PinBase = Cast<UMaterialExpressionPinBase>(Cast<UMaterialGraphNode>(Entry)->MaterialExpression);
+		PinBase->DeleteReroutePins();
+		Material->Expressions.Remove(PinBase);
+		PinBase->MarkPendingKill();
+		Entry->DestroyNode();
+		bPreviewExpressionDeleted |= PinBase == PreviewExpression;
+	}
+
+	if (Result)
+	{
+		UMaterialExpressionPinBase* PinBase = Cast<UMaterialExpressionPinBase>(Cast<UMaterialGraphNode>(Result)->MaterialExpression);
+		PinBase->DeleteReroutePins();
+		Material->Expressions.Remove(PinBase);
+		PinBase->MarkPendingKill();
+		Result->DestroyNode();
+		bPreviewExpressionDeleted |= PinBase == PreviewExpression;
+	}
+
+	// Make sure any subgraphs get propagated appropriately
+	if (SourceGraph->SubGraphs.Num() > 0)
+	{
+		DestinationGraph->SubGraphs.Append(SourceGraph->SubGraphs);
+		SourceGraph->SubGraphs.Empty();
+	}
+
+	// Remove the gateway node and source graph
+	UMaterialExpression* CompositeExpression = Cast<UMaterialGraphNode>(InNodeToExpand)->MaterialExpression;
+	CompositeExpression->Modify();
+	Material->Expressions.Remove(CompositeExpression);
+	CompositeExpression->MarkPendingKill();
+	InNodeToExpand->DestroyNode();
+	bPreviewExpressionDeleted |= CompositeExpression == PreviewExpression;
+
+	if (bPreviewExpressionDeleted)
+	{
+		// The preview expression was deleted.  Null out our reference to it and reset to the normal preview material
+		SetPreviewExpression(nullptr);
+		RegenerateCodeView();
+		UpdatePreviewMaterial();
+	}
+}
+
+void FMaterialEditor::MoveNodesToAveragePos(TSet<UEdGraphNode*>& AverageNodes, FVector2D SourcePos, bool bExpandedNodesNeedUniqueGuid) const
+{
+	if (AverageNodes.Num() > 0)
+	{
+		FVector2D AvgNodePosition(0.0f, 0.0f);
+
+		for (TSet<UEdGraphNode*>::TIterator It(AverageNodes); It; ++It)
+		{
+			UEdGraphNode* Node = *It;
+			AvgNodePosition.X += Node->NodePosX;
+			AvgNodePosition.Y += Node->NodePosY;
+		}
+
+		float InvNumNodes = 1.0f / float(AverageNodes.Num());
+		AvgNodePosition.X *= InvNumNodes;
+		AvgNodePosition.Y *= InvNumNodes;
+
+		TSharedPtr<SGraphEditor> FocusedGraphEd = FocusedGraphEdPtr.Pin();
+
+		for (UEdGraphNode* ExpandedNode : AverageNodes)
+		{
+			ExpandedNode->NodePosX = (ExpandedNode->NodePosX - AvgNodePosition.X) + SourcePos.X;
+			ExpandedNode->NodePosY = (ExpandedNode->NodePosY - AvgNodePosition.Y) + SourcePos.Y;
+
+			ExpandedNode->SnapToGrid(SNodePanel::GetSnapGridSize());
+
+			if (bExpandedNodesNeedUniqueGuid)
+			{
+				ExpandedNode->CreateNewGuid();
+			}
+
+			//Add expanded node to selection
+			FocusedGraphEd->SetNodeSelection(ExpandedNode, true);
+		}
+	}
+}
+
+void FMaterialEditor::MoveNodesToGraph(TArray<UEdGraphNode*>& SourceNodes, UEdGraph* DestinationGraph, TSet<UEdGraphNode*>& OutExpandedNodes, UEdGraphNode** OutEntry, UEdGraphNode** OutResult, const bool bIsCollapsedGraph)
+{
+	// Move the nodes over, remembering any that are boundary nodes
+	while (SourceNodes.Num())
+	{
+		UEdGraphNode* Node = SourceNodes.Pop();
+		UEdGraph* OriginalGraph = Node->GetGraph();
+
+		Node->Modify();
+		OriginalGraph->Modify();
+		Node->Rename(/*NewName=*/ nullptr, /*NewOuter=*/ DestinationGraph, REN_DontCreateRedirectors);
+
+		// Remove the node from the original graph
+		OriginalGraph->RemoveNode(Node, false);
+
+		// We do not check CanPasteHere when determining CanCollapseNodes, unlike CanCollapseSelectionToFunction/Macro,
+		// so when expanding a collapsed graph we don't want to check the CanPasteHere function:
+		if (!bIsCollapsedGraph && !Node->CanPasteHere(DestinationGraph))
+		{
+			Node->BreakAllNodeLinks();
+			continue;
+		}
+
+		// Successfully added the node to the graph, we may need to remove flags
+		if (Node->HasAllFlags(RF_Transient) && !DestinationGraph->HasAllFlags(RF_Transient))
+		{
+			Node->SetFlags(RF_Transactional);
+			Node->ClearFlags(RF_Transient);
+			TArray<UObject*> Subobjects;
+			GetObjectsWithOuter(Node, Subobjects);
+			for (UObject* Subobject : Subobjects)
+			{
+				Subobject->ClearFlags(RF_Transient);
+				Subobject->SetFlags(RF_Transactional);
+			}
+		}
+
+		DestinationGraph->AddNode(Node, /* bFromUI */ false, /* bSelectNewNode */ false);
+
+		if (UMaterialGraphNode_Composite* Composite = Cast<UMaterialGraphNode_Composite>(Node))
+		{
+			OriginalGraph->SubGraphs.Remove(Composite->BoundGraph);
+			DestinationGraph->SubGraphs.Add(Composite->BoundGraph);
+		}
+
+		UMaterialGraphNode* MaterialNode = Cast<UMaterialGraphNode>(Node);
+		if (MaterialNode && MaterialNode->MaterialExpression->IsA(UMaterialExpressionPinBase::StaticClass()))
+		{
+			UMaterialExpressionPinBase* PinBase = Cast<UMaterialExpressionPinBase>(MaterialNode->MaterialExpression);
+
+			if (PinBase && PinBase->PinDirection == EGPD_Output)
+			{
+				*OutEntry = Node;
+			}
+			else if (PinBase && PinBase->PinDirection == EGPD_Input)
+			{
+				*OutResult = Node;
+			}
+		}
+		else
+		{
+			OutExpandedNodes.Add(Node);
+		}
+	}
+}
+
+bool FMaterialEditor::CollapseGatewayNode(UEdGraphNode* InNode, UEdGraphNode* InEntryNode, UEdGraphNode* InResultNode, TSet<UEdGraphNode*>* OutExpandedNodes)
+{
+	bool bSuccessful = true;
+
+	// We iterate the array in reverse so we can both remove the subpins safely after we've read them and
+	// so we have split nested structs we combine them back together in the right order
+	for (int32 BoundaryPinIndex = InNode->Pins.Num() - 1; BoundaryPinIndex >= 0; --BoundaryPinIndex)
+	{
+		UEdGraphPin* const BoundaryPin = InNode->Pins[BoundaryPinIndex];
+
+		// For each pin in the gateway node, find the associated pin in the entry or result node.
+		UEdGraphNode* const GatewayNode = (BoundaryPin->Direction == EGPD_Input) ? InEntryNode : InResultNode;
+		UEdGraphPin* GatewayPin = nullptr;
+		if (GatewayNode)
+		{
+			for (int32 PinIdx = GatewayNode->Pins.Num() - 1; PinIdx >= 0; --PinIdx)
+			{
+				UEdGraphPin* const Pin = GatewayNode->Pins[PinIdx];
+
+				// Function graphs have a single exec path through them, so only one exec pin for input and another for output. In this fashion, they must not be handled by name.
+				if ((Pin->PinName == BoundaryPin->PinName) && (Pin->Direction != BoundaryPin->Direction))
+				{
+					GatewayPin = Pin;
+					break;
+				}
+			}
+		}
+
+		if (GatewayPin)
+		{
+			//@TODO: This is same as UEdGraphSchema_K2::CombineTwoPinNetsAndRemoveOldPins, except we don't care about default values
+			// since material graphs pins can't do that. More consolidation
+			auto CombineTwoPinNetsAndRemoveOldPins = [](UEdGraphPin* InPinA, UEdGraphPin* InPinB)
+			{
+				// Make direct connections between the things that connect to A or B, removing A and B from the picture
+				for (int32 IndexA = 0; IndexA < InPinA->LinkedTo.Num(); ++IndexA)
+				{
+					UEdGraphPin* FarA = InPinA->LinkedTo[IndexA];
+					// TODO: Michael N. says this if check should be unnecessary once the underlying issue is fixed.
+					// (Probably should use a check() instead once it's removed though.  See additional cases above.
+					if (FarA != nullptr)
+					{
+						for (int32 IndexB = 0; IndexB < InPinB->LinkedTo.Num(); ++IndexB)
+						{
+							UEdGraphPin* FarB = InPinB->LinkedTo[IndexB];
+
+							if (FarB != nullptr)
+							{
+								FarA->Modify();
+								FarB->Modify();
+								FarA->MakeLinkTo(FarB);
+							}
+
+						}
+					}
+				}
+			};
+			CombineTwoPinNetsAndRemoveOldPins(BoundaryPin, GatewayPin);
+		}
+	}
+
+	return bSuccessful;
+}
+
 void FMaterialEditor::SetQualityPreview(EMaterialQualityLevel::Type NewQuality)
 {
 	NodeQualityLevel = NewQuality;
@@ -959,10 +1361,6 @@ void FMaterialEditor::CreateInternalWidgets()
 
 	FPropertyEditorModule& PropertyEditorModule = FModuleManager::GetModuleChecked<FPropertyEditorModule>( "PropertyEditor" );
 
-	GraphEditor = CreateGraphEditorWidget();
-	// Manually set zoom level to avoid deferred zooming
-	GraphEditor->SetViewLocation(FVector2D::ZeroVector, 1);
-
 	const FDetailsViewArgs DetailsViewArgs( false, false, true, FDetailsViewArgs::HideNameArea, true, this );
 	MaterialDetailsView = PropertyEditorModule.CreateDetailView( DetailsViewArgs );
 
@@ -1004,6 +1402,13 @@ void FMaterialEditor::CreateInternalWidgets()
 		LayoutCollectionParameterDetails
 		);
 
+	FOnGetDetailCustomizationInstance LayoutCompositeDetails = FOnGetDetailCustomizationInstance::CreateStatic(&FMaterialExpressionCompositeDetails::MakeInstance);
+
+	MaterialDetailsView->RegisterInstancedCustomPropertyLayout( 
+		UMaterialExpressionComposite::StaticClass(), 
+		LayoutCompositeDetails
+		);
+
 	MaterialDetailsView->OnFinishedChangingProperties().AddSP(this, &FMaterialEditor::OnFinishedChangingProperties);
 
 	PropertyEditorModule.RegisterCustomClassLayout( UMaterial::StaticClass()->GetFName(), FOnGetDetailCustomizationInstance::CreateStatic(&FMaterialDetailCustomization::MakeInstance ) );
@@ -1124,6 +1529,14 @@ void FMaterialEditor::OnFinishedChangingParametersFromOverview(const FPropertyCh
 	}
 }
 
+void FMaterialEditor::OnChangeBreadCrumbGraph(UEdGraph* InGraph)
+{
+	if (InGraph && FocusedGraphEdPtr.IsValid())
+	{
+		OpenDocument(InGraph, FDocumentTracker::NavigatingCurrentDocument);
+	}
+}
+
 void FMaterialEditor::GeneratorRowsRefreshed()
 {
 	MaterialParametersOverviewWidget->Refresh();
@@ -1218,7 +1631,7 @@ void FMaterialEditor::InitToolMenuContext(FToolMenuContext& MenuContext)
 
 void FMaterialEditor::RegisterToolBar()
 {
-	const FName MenuName = GetToolMenuToolbarName();
+	const FName MenuName = FAssetEditorToolkit::GetToolMenuToolbarName();
 	if (!UToolMenus::Get()->IsMenuRegistered(MenuName))
 	{
 		UToolMenu* ToolBar = UToolMenus::Get()->RegisterMenu(MenuName, "AssetEditor.DefaultToolBar", EMultiBoxType::ToolBar);
@@ -1750,6 +2163,12 @@ void FMaterialEditor::DrawMessages( FViewport* InViewport, FCanvas* Canvas )
 
 void FMaterialEditor::RecenterEditor()
 {
+	TSharedPtr<SGraphEditor> FocusedGraphEd = FocusedGraphEdPtr.Pin();
+	if (!FocusedGraphEd)
+	{
+		return;
+	}
+
 	UEdGraphNode* FocusNode = NULL;
 
 	if (MaterialFunction)
@@ -1791,8 +2210,8 @@ void FMaterialEditor::RecenterEditor()
 		// Get current view location so that we don't change the zoom amount
 		FVector2D CurrLocation;
 		float CurrZoomLevel;
-		GraphEditor->GetViewLocation(CurrLocation, CurrZoomLevel);
-		GraphEditor->SetViewLocation(FVector2D::ZeroVector, CurrZoomLevel);
+		FocusedGraphEd->GetViewLocation(CurrLocation, CurrZoomLevel);
+		FocusedGraphEd->SetViewLocation(FVector2D::ZeroVector, CurrZoomLevel);
 	}
 }
 
@@ -2482,34 +2901,38 @@ void FMaterialEditor::UpdateMaterialInfoList()
 		NewErrorHash += FMD5::HashAnsiString(*CompileErrors[ErrorIndex]);
 	}
 
-	TSharedPtr<SWidget> TitleBar = GraphEditor->GetTitleBar();
-	if (NewErrorHash != MaterialErrorHash)
+	if (TSharedPtr<SGraphEditor> FocusedGraphEd = FocusedGraphEdPtr.Pin())
 	{
-		MaterialErrorHash = NewErrorHash;
-		MaterialInfoList.Reset();
-
-		TArray< TSharedRef<class FTokenizedMessage> > Messages;
-		FString FeatureLevelName;
-		GetFeatureLevelName(GMaxRHIFeatureLevel, FeatureLevelName);
-		for (int32 ErrorIndex = 0; ErrorIndex < CompileErrors.Num(); ErrorIndex++)
+		TSharedPtr<SWidget> TitleBar = FocusedGraphEd->GetTitleBar();
+		TSharedPtr<SMaterialEditorTitleBar> MaterialTitleBar = StaticCastSharedPtr<SMaterialEditorTitleBar>(TitleBar);
+		if (NewErrorHash != MaterialErrorHash)
 		{
-			FString ErrorString = FString::Printf(TEXT("[%s] %s"), *FeatureLevelName, *CompileErrors[ErrorIndex]);
-			TSharedRef<FTokenizedMessage> Line = FTokenizedMessage::Create(EMessageSeverity::Error);
-			Line->AddToken(FTextToken::Create(FText::FromString(ErrorString)));
-			Messages.Add(Line);
+			MaterialErrorHash = NewErrorHash;
+			MaterialInfoList.Reset();
 
-			MaterialInfoList.Add(MakeShareable(new FMaterialInfo(ErrorString, FLinearColor::Red)));
-		}
+			TArray< TSharedRef<class FTokenizedMessage> > Messages;
+			FString FeatureLevelName;
+			GetFeatureLevelName(GMaxRHIFeatureLevel, FeatureLevelName);
+			for (int32 ErrorIndex = 0; ErrorIndex < CompileErrors.Num(); ErrorIndex++)
+			{
+				FString ErrorString = FString::Printf(TEXT("[%s] %s"), *FeatureLevelName, *CompileErrors[ErrorIndex]);
+				TSharedRef<FTokenizedMessage> Line = FTokenizedMessage::Create(EMessageSeverity::Error);
+				Line->AddToken(FTextToken::Create(FText::FromString(ErrorString)));
+				Messages.Add(Line);
+
+				MaterialInfoList.Add(MakeShareable(new FMaterialInfo(ErrorString, FLinearColor::Red)));
+			}
 
-		StatsListing->ClearMessages();
-		StatsListing->AddMessages(Messages);
+			StatsListing->ClearMessages();
+			StatsListing->AddMessages(Messages);
 
-		StaticCastSharedPtr<SMaterialEditorTitleBar>(TitleBar)->RequestRefresh();
-		TitleBar->SetVisibility(EVisibility::Visible);
-	}
-	else if (!CompileErrors.Num())
+			MaterialTitleBar->RequestRefresh();
+		}
+
+		if (MaterialTitleBar->MaterialInfoList)
 		{
-		TitleBar->SetVisibility(EVisibility::Collapsed);
+			MaterialTitleBar->MaterialInfoList->SetVisibility(CompileErrors.IsEmpty() ? EVisibility::Collapsed : EVisibility::Visible);
+		}
 	}
 
 	// extract material stats
@@ -2551,10 +2974,31 @@ void FMaterialEditor::UpdateGraphNodeStates()
 		}
 	}
 	
+	// Update main material graph and all subgraphs
+	bUpdatedErrorState |= UpdateGraphNodeState(Material->MaterialGraph, ErrorMaterialResource, VisibleExpressions, bShowAllNodes);
+
+	bPreviewFeaturesChanged = false;
+
+	if (bUpdatedErrorState || bToggledVisibleState)
+	{
+		// Rebuild the SGraphNodes to display/hide error block
+		if (TSharedPtr<SGraphEditor> FocusedGraphEd = FocusedGraphEdPtr.Pin())
+		{
+			FocusedGraphEd->NotifyGraphChanged();
+		}
+	}
+}
+
+bool FMaterialEditor::UpdateGraphNodeState(UEdGraph* Graph, const FMaterialResource* ErrorMaterialResource, TArray<UMaterialExpression*>& VisibleExpressions, bool bShowAllNodes)
+{
+	bool bUpdatedErrorState = false;
+
+	UMaterialGraph* MaterialGraph = CastChecked<UMaterialGraph>(Graph);
+
 	// Have to loop through everything here as there's no way to be notified when the material resource updates
-	for (int32 Index = 0; Index < Material->MaterialGraph->Nodes.Num(); ++Index)
+	for (int32 Index = 0; Index < MaterialGraph->Nodes.Num(); ++Index)
 	{
-		UMaterialGraphNode* MaterialNode = Cast<UMaterialGraphNode>(Material->MaterialGraph->Nodes[Index]);
+		UMaterialGraphNode* MaterialNode = Cast<UMaterialGraphNode>(MaterialGraph->Nodes[Index]);
 		if (MaterialNode)
 		{
 			MaterialNode->bIsPreviewExpression = (PreviewExpression == MaterialNode->MaterialExpression);
@@ -2588,13 +3032,12 @@ void FMaterialEditor::UpdateGraphNodeStates()
 		}
 	}
 
-	bPreviewFeaturesChanged = false;
-
-	if (bUpdatedErrorState || bToggledVisibleState)
+	for (UEdGraph* SubGraph : Graph->SubGraphs)
 	{
-		// Rebuild the SGraphNodes to display/hide error block
-		GraphEditor->NotifyGraphChanged();
+		bUpdatedErrorState |= UpdateGraphNodeState(SubGraph, ErrorMaterialResource, VisibleExpressions, bShowAllNodes);
 	}
+
+	return bUpdatedErrorState;
 }
 
 void FMaterialEditor::AddReferencedObjects( FReferenceCollector& Collector )
@@ -2801,7 +3244,11 @@ void FMaterialEditor::OnCameraHome()
 void FMaterialEditor::OnShowConnectors()
 {
 	bHideUnusedConnectors = !bHideUnusedConnectors;
-	GraphEditor->SetPinVisibility(bHideUnusedConnectors ? SGraphEditor::Pin_HideNoConnection : SGraphEditor::Pin_Show);
+
+	if (TSharedPtr<SGraphEditor> FocusedGraphEd = FocusedGraphEdPtr.Pin())
+	{
+		FocusedGraphEd->SetPinVisibility(bHideUnusedConnectors ? SGraphEditor::Pin_HideNoConnection : SGraphEditor::Pin_Show);
+	}
 }
 
 bool FMaterialEditor::IsOnShowConnectorsChecked() const
@@ -2850,9 +3297,15 @@ bool FMaterialEditor::IsOnAlwaysRefreshAllPreviews() const
 
 void FMaterialEditor::ToggleHideUnrelatedNodes()
 {
-	bHideUnrelatedNodes = !bHideUnrelatedNodes;
-
-	GraphEditor->ResetAllNodesUnrelatedStates();
+	TSharedPtr<SGraphEditor> FocusedGraphEd = FocusedGraphEdPtr.Pin();
+	if (!FocusedGraphEd)
+	{
+		return;
+	}
+
+	bHideUnrelatedNodes = !bHideUnrelatedNodes;
+
+	FocusedGraphEd->ResetAllNodesUnrelatedStates();
 
 	if (bHideUnrelatedNodes && bSelectRegularNode)
 	{
@@ -2915,6 +3368,12 @@ void FMaterialEditor::CollectUpstreamNodes(UMaterialGraphNode* CurrentNode, TArr
 
 void FMaterialEditor::HideUnrelatedNodes()
 {
+	TSharedPtr<SGraphEditor> FocusedGraphEd = FocusedGraphEdPtr.Pin();
+	if (!FocusedGraphEd)
+	{
+		return;
+	}
+
 	TArray<UMaterialGraphNode*> NodesToShow;
 
 	const FGraphPanelSelectionSet SelectedNodes = GetSelectedNodes();
@@ -2931,7 +3390,7 @@ void FMaterialEditor::HideUnrelatedNodes()
 		}
 	}
 
-	TArray<class UEdGraphNode*> AllNodes = GraphEditor->GetCurrentGraph()->Nodes;
+	TArray<class UEdGraphNode*> AllNodes = FocusedGraphEd->GetCurrentGraph()->Nodes;
 
 	TArray<UEdGraphNode*> CommentNodes;
 	TArray<UEdGraphNode*> RelatedNodes;
@@ -2957,7 +3416,7 @@ void FMaterialEditor::HideUnrelatedNodes()
 		}
 	}
 
-	GraphEditor->FocusCommentNodes(CommentNodes, RelatedNodes);
+	FocusedGraphEd->FocusCommentNodes(CommentNodes, RelatedNodes);
 }
 
 void FMaterialEditor::MakeHideUnrelatedNodesOptionsMenu(UToolMenu* Menu)
@@ -3016,11 +3475,17 @@ void FMaterialEditor::OnLockNodeStateCheckStateChanged(ECheckBoxState NewChecked
 
 void FMaterialEditor::OnFocusWholeChainCheckStateChanged(ECheckBoxState NewCheckedState)
 {
+	TSharedPtr<SGraphEditor> FocusedGraphEd = FocusedGraphEdPtr.Pin();
+	if (!FocusedGraphEd)
+	{
+		return;
+	}
+
 	bFocusWholeChain = (NewCheckedState == ECheckBoxState::Checked) ? true : false;
 
 	if (bHideUnrelatedNodes && !bLockNodeFadeState && bSelectRegularNode)
 	{
-		GraphEditor->ResetAllNodesUnrelatedStates();
+		FocusedGraphEd->ResetAllNodesUnrelatedStates();
 
 		HideUnrelatedNodes();
 	}
@@ -3251,9 +3716,12 @@ void FMaterialEditor::OnConvertObjects()
 		DeleteNodes(NodesToDelete);
 
 		// Select each of the newly converted expressions
-		for ( TArray<UEdGraphNode*>::TConstIterator NodeIter(NodesToSelect); NodeIter; ++NodeIter )
+		if (TSharedPtr<SGraphEditor> FocusedGraphEd = FocusedGraphEdPtr.Pin())
 		{
-			GraphEditor->SetNodeSelection(*NodeIter, true);
+			for (TArray<UEdGraphNode*>::TConstIterator NodeIter(NodesToSelect); NodeIter; ++NodeIter)
+			{
+				FocusedGraphEd->SetNodeSelection(*NodeIter, true);
+			}
 		}
 	}
 }
@@ -3336,9 +3804,12 @@ void FMaterialEditor::OnConvertTextures()
 		DeleteNodes(NodesToDelete);
 
 		// Select each of the newly converted expressions
-		for ( TArray<UEdGraphNode*>::TConstIterator NodeIter(NodesToSelect); NodeIter; ++NodeIter )
+		if (TSharedPtr<SGraphEditor> FocusedGraphEd = FocusedGraphEdPtr.Pin())
 		{
-			GraphEditor->SetNodeSelection(*NodeIter, true);
+			for ( TArray<UEdGraphNode*>::TConstIterator NodeIter(NodesToSelect); NodeIter; ++NodeIter )
+			{
+				FocusedGraphEd->SetNodeSelection(*NodeIter, true);
+			}
 		}
 	}
 }
@@ -3353,7 +3824,11 @@ void FMaterialEditor::OnPreviewNode()
 			UMaterialGraphNode* GraphNode = Cast<UMaterialGraphNode>(*NodeIt);
 			if (GraphNode)
 			{
-				GraphEditor->NotifyGraphChanged();
+				if (TSharedPtr<SGraphEditor> FocusedGraphEd = FocusedGraphEdPtr.Pin())
+				{
+					FocusedGraphEd->NotifyGraphChanged();
+				}
+				
 				SetPreviewExpression(GraphNode->MaterialExpression);
 			}
 		}
@@ -3432,9 +3907,12 @@ void FMaterialEditor::OnSelectDownstreamNodes()
 		NodesToCheck.Remove(CurrentNode);
 	}
 
-	for (int32 Index = 0; Index < NodesToSelect.Num(); ++Index)
+	if (TSharedPtr<SGraphEditor> FocusedGraphEd = FocusedGraphEdPtr.Pin())
 	{
-		GraphEditor->SetNodeSelection(NodesToSelect[Index], true);
+		for (int32 Index = 0; Index < NodesToSelect.Num(); ++Index)
+		{
+			FocusedGraphEd->SetNodeSelection(NodesToSelect[Index], true);
+		}
 	}
 }
 
@@ -3485,9 +3963,12 @@ void FMaterialEditor::OnSelectUpstreamNodes()
 		NodesToCheck.Remove(CurrentNode);
 	}
 
-	for (int32 Index = 0; Index < NodesToSelect.Num(); ++Index)
+	if (TSharedPtr<SGraphEditor> FocusedGraphEd = FocusedGraphEdPtr.Pin())
 	{
-		GraphEditor->SetNodeSelection(NodesToSelect[Index], true);
+		for (int32 Index = 0; Index < NodesToSelect.Num(); ++Index)
+		{
+			FocusedGraphEd->SetNodeSelection(NodesToSelect[Index], true);
+		}
 	}
 }
 
@@ -3499,20 +3980,54 @@ void FMaterialEditor::OnForceRefreshPreviews()
 
 void FMaterialEditor::OnCreateComment()
 {
-	CreateNewMaterialExpressionComment(GraphEditor->GetPasteLocation());
+	if (TSharedPtr<SGraphEditor> FocusedGraphEd = FocusedGraphEdPtr.Pin())
+	{
+		CreateNewMaterialExpressionComment(FocusedGraphEd->GetPasteLocation());
+	}
 }
 
 void FMaterialEditor::OnCreateComponentMaskNode()
 {
-	CreateNewMaterialExpression(UMaterialExpressionComponentMask::StaticClass(), GraphEditor->GetPasteLocation(), true, false);
+	if (TSharedPtr<SGraphEditor> FocusedGraphEd = FocusedGraphEdPtr.Pin())
+	{
+		CreateNewMaterialExpression(UMaterialExpressionComponentMask::StaticClass(), FocusedGraphEd->GetPasteLocation(), true, false);
+	}
 }
 
 void FMaterialEditor::OnFindInMaterial()
 {
-	TabManager->TryInvokeTab(FindTabId);
+	TabManager->TryInvokeTab(FMaterialEditorTabs::FindTabId);
 	FindResults->FocusForUse();
 }
 
+void FMaterialEditor::OnGraphEditorFocused(const TSharedRef<class SGraphEditor>& InGraphEditor)
+{
+	if (FocusedGraphEdPtr == InGraphEditor)
+	{
+		return;
+	}
+
+	// Update the graph editor that is currently focused
+	FocusedGraphEdPtr = InGraphEditor;
+
+	// Refresh navigation history
+	TSharedPtr<SWidget> TitleBar = FocusedGraphEdPtr.Pin()->GetTitleBar();
+	StaticCastSharedPtr<SMaterialEditorTitleBar>(TitleBar)->RequestRefresh();
+
+	// Update the inspector as well, to show selection from the focused graph editor
+	FGraphPanelSelectionSet SelectedNodes = GetSelectedNodes();
+
+	if (bHideUnrelatedNodes && SelectedNodes.Num() <= 0)
+	{
+		FocusedGraphEdPtr.Pin()->ResetAllNodesUnrelatedStates();
+	}
+}
+
+void FMaterialEditor::OnGraphEditorBackgrounded(const TSharedRef<SGraphEditor>& InGraphEditor)
+{
+	FocusedGraphEdPtr = nullptr;
+}
+
 UClass* FMaterialEditor::GetOnPromoteToParameterClass(const UEdGraphPin* TargetPin) const
 {
 	UMaterialGraphNode_Root* RootPinNode = Cast<UMaterialGraphNode_Root>(TargetPin->GetOwningNode());
@@ -3889,19 +4404,6 @@ TSharedRef<SDockTab> FMaterialEditor::SpawnTab_Preview(const FSpawnTabArgs& Args
 	return SpawnedTab;
 }
 
-TSharedRef<SDockTab> FMaterialEditor::SpawnTab_GraphCanvas(const FSpawnTabArgs& Args)
-{
-	TSharedRef<SDockTab> SpawnedTab = SNew(SDockTab)
-		.Label(LOCTEXT("GraphCanvasTitle", "Graph"));
-
-	if (GraphEditor.IsValid())
-	{
-		SpawnedTab->SetContent(GraphEditor.ToSharedRef());
-	}
-
-	return SpawnedTab;
-}
-
 TSharedRef<SDockTab> FMaterialEditor::SpawnTab_MaterialProperties(const FSpawnTabArgs& Args)
 {
 	TSharedPtr<SDockTab> DetailsTab = SNew(SDockTab)
@@ -3911,10 +4413,10 @@ TSharedRef<SDockTab> FMaterialEditor::SpawnTab_MaterialProperties(const FSpawnTa
 			MaterialDetailsView.ToSharedRef()
 		];
 
-	if (GraphEditor.IsValid())
+	if (FocusedGraphEdPtr.IsValid())
 	{
 		// Since we're initialising, make sure nothing is selected
-		GraphEditor->ClearSelectionSet();
+		FocusedGraphEdPtr.Pin()->ClearSelectionSet();
 	}
 	SpawnedDetailsTab = DetailsTab;
 	return DetailsTab.ToSharedRef();
@@ -3922,7 +4424,7 @@ TSharedRef<SDockTab> FMaterialEditor::SpawnTab_MaterialProperties(const FSpawnTa
 
 TSharedRef<SDockTab> FMaterialEditor::SpawnTab_Palette(const FSpawnTabArgs& Args)
 {
-	check( Args.GetTabId() == PaletteTabId );
+	check( Args.GetTabId() == FMaterialEditorTabs::PaletteTabId );
 
 	TSharedRef<SDockTab> SpawnedTab = SNew(SDockTab)
 		.Icon(FEditorStyle::GetBrush("Kismet.Tabs.Palette"))
@@ -3940,7 +4442,7 @@ TSharedRef<SDockTab> FMaterialEditor::SpawnTab_Palette(const FSpawnTabArgs& Args
 
 TSharedRef<SDockTab> FMaterialEditor::SpawnTab_Find(const FSpawnTabArgs& Args)
 {
-	check(Args.GetTabId() == FindTabId);
+	check(Args.GetTabId() == FMaterialEditorTabs::FindTabId);
 
 	TSharedRef<SDockTab> SpawnedTab = SNew(SDockTab)
 		.Icon(FEditorStyle::GetBrush("Kismet.Tabs.FindResults"))
@@ -3958,7 +4460,7 @@ TSharedRef<SDockTab> FMaterialEditor::SpawnTab_Find(const FSpawnTabArgs& Args)
 
 TSharedRef<SDockTab> FMaterialEditor::SpawnTab_PreviewSettings(const FSpawnTabArgs& Args)
 {
-	check(Args.GetTabId() == PreviewSettingsTabId);
+	check(Args.GetTabId() == FMaterialEditorTabs::PreviewSettingsTabId);
 
 	TSharedRef<SWidget> InWidget = SNullWidget::NullWidget;
 	if (PreviewViewport.IsValid())
@@ -4091,12 +4593,54 @@ void FMaterialEditor::SetPreviewExpression(UMaterialExpression* NewPreviewExpres
 
 void FMaterialEditor::JumpToNode(const UEdGraphNode* Node)
 {
-	GraphEditor->JumpToNode(Node, false);
+	if (TSharedPtr<SGraphEditor> FocusedGraphEd = FocusedGraphEdPtr.Pin())
+	{
+		if (Node->GetGraph() != FocusedGraphEd->GetCurrentGraph())
+		{
+			JumpToHyperlink(Node->GetGraph());
+
+			// Graph changed so editor changed, use new one to jump
+			FocusedGraphEdPtr.Pin()->JumpToNode(Node, false);
+		}
+		else
+		{
+			FocusedGraphEd->JumpToNode(Node, false);
+		}
+	}
 }
 
-UMaterialExpression* FMaterialEditor::CreateNewMaterialExpression(UClass* NewExpressionClass, const FVector2D& NodePos, bool bAutoSelect, bool bAutoAssignResource)
+bool FMaterialEditor::FindOpenTabsContainingDocument(const UObject* DocumentID, TArray<TSharedPtr<SDockTab>>& Results)
+{
+	int32 StartingCount = Results.Num();
+
+	TSharedRef<FTabPayload_UObject> Payload = FTabPayload_UObject::Make(DocumentID);
+
+	DocumentManager->FindMatchingTabs(Payload, /*inout*/ Results);
+
+	// Did we add anything new?
+	return (StartingCount != Results.Num());
+}
+
+TSharedPtr<SDockTab> FMaterialEditor::OpenDocument(const UObject* DocumentID, FDocumentTracker::EOpenDocumentCause Cause)
+{
+	TSharedRef<FTabPayload_UObject> Payload = FTabPayload_UObject::Make(DocumentID);
+	TSharedPtr<SDockTab> DocumentTab = DocumentManager->OpenDocument(Payload, Cause);
+
+	return DocumentTab;
+}
+
+void FMaterialEditor::CloseDocumentTab(const UObject* DocumentID)
+{
+	TSharedRef<FTabPayload_UObject> Payload = FTabPayload_UObject::Make(DocumentID);
+	DocumentManager->CloseTab(Payload);
+}
+
+UMaterialExpression* FMaterialEditor::CreateNewMaterialExpression(UClass* NewExpressionClass, const FVector2D& NodePos, bool bAutoSelect, bool bAutoAssignResource, const UEdGraph* Graph)
 {
 	check( NewExpressionClass->IsChildOf(UMaterialExpression::StaticClass()) );
+	TSharedPtr<SGraphEditor> FocusedGraphEd = FocusedGraphEdPtr.Pin();
+	UMaterialGraph* ExpressionGraph = Graph ? CastChecked<UMaterialGraph>(const_cast<UEdGraph*>(Graph)) : Material->MaterialGraph; 
+	ExpressionGraph->Modify();
 
 	if (!IsAllowedExpressionType(NewExpressionClass, MaterialFunction != NULL))
 	{
@@ -4105,9 +4649,9 @@ UMaterialExpression* FMaterialEditor::CreateNewMaterialExpression(UClass* NewExp
 	}
 
 	// Clear the selection
-	if ( bAutoSelect )
+	if ( bAutoSelect && FocusedGraphEd )
 	{
-		GraphEditor->ClearSelectionSet();
+		FocusedGraphEd->ClearSelectionSet();
 	}
 
 	// Create the new expression.
@@ -4128,12 +4672,12 @@ UMaterialExpression* FMaterialEditor::CreateNewMaterialExpression(UClass* NewExp
 
 		if (NewExpression)
 		{
-			Material->MaterialGraph->AddExpression(NewExpression, bAutoSelect);
+			ExpressionGraph->AddExpression(NewExpression, bAutoSelect);
 
 			// Select the new node.
-			if ( bAutoSelect )
+			if ( bAutoSelect && FocusedGraphEd )
 			{
-				GraphEditor->SetNodeSelection(NewExpression->GraphNode, true);
+				FocusedGraphEd->SetNodeSelection(NewExpression->GraphNode, true);
 			}
 		}
 	}
@@ -4145,13 +4689,81 @@ UMaterialExpression* FMaterialEditor::CreateNewMaterialExpression(UClass* NewExp
 	Material->MarkPackageDirty();
 
 	RefreshExpressionPreviews();
-	GraphEditor->NotifyGraphChanged();
+	if (FocusedGraphEd)
+	{
+		FocusedGraphEd->NotifyGraphChanged();
+	}
 	SetMaterialDirty();
 	return NewExpression;
 }
 
-UMaterialExpressionComment* FMaterialEditor::CreateNewMaterialExpressionComment(const FVector2D& NodePos)
+UMaterialExpressionComposite* FMaterialEditor::CreateNewMaterialExpressionComposite(const FVector2D& NodePos, const UEdGraph* Graph)
 {
+	TSharedPtr<SGraphEditor> FocusedGraphEd = FocusedGraphEdPtr.Pin();
+	UMaterialGraph* ExpressionGraph = Graph ? CastChecked<UMaterialGraph>(const_cast<UEdGraph*>(Graph)) : Material->MaterialGraph;
+	ExpressionGraph->Modify();
+
+	UMaterialExpressionComposite* NewComposite = nullptr;
+	{
+		const FScopedTransaction Transaction(NSLOCTEXT("UnrealEd", "MaterialEditorNewComposite", "Material Editor: New Composite"));
+		Material->Modify();
+
+		UObject* ExpressionOuter = Material;
+		if (MaterialFunction)
+		{
+			ExpressionOuter = MaterialFunction;
+		}
+
+		UMaterialExpression* NewExpression = UMaterialEditingLibrary::CreateMaterialExpressionEx(Material, MaterialFunction, UMaterialExpressionComposite::StaticClass(), nullptr, NodePos.X, NodePos.Y);
+		NewComposite = Cast<UMaterialExpressionComposite>(NewExpression);
+
+		if (NewComposite)
+		{
+			UMaterialExpression* InputPinBase = UMaterialEditingLibrary::CreateMaterialExpressionEx(Material, MaterialFunction, UMaterialExpressionPinBase::StaticClass());
+			NewComposite->InputExpressions = Cast<UMaterialExpressionPinBase>(InputPinBase);
+			NewComposite->InputExpressions->PinDirection = EGPD_Output;
+			NewComposite->InputExpressions->SubgraphExpression = NewComposite;
+
+			UMaterialExpression* OutputPinBase = UMaterialEditingLibrary::CreateMaterialExpressionEx(Material, MaterialFunction, UMaterialExpressionPinBase::StaticClass());
+			NewComposite->OutputExpressions = Cast<UMaterialExpressionPinBase>(OutputPinBase);
+			NewComposite->OutputExpressions->PinDirection = EGPD_Input;
+			NewComposite->OutputExpressions->SubgraphExpression = NewComposite;
+
+			// Create graph node, subgraph, and pinbase graph nodes
+			{
+				ExpressionGraph->AddExpression(NewComposite, true);
+
+				UMaterialGraphNode_Composite* CompositeNode = CastChecked<UMaterialGraphNode_Composite>(NewComposite->GraphNode);
+				CompositeNode->BoundGraph = ExpressionGraph->AddSubGraph(NewComposite);
+				CompositeNode->BoundGraph->Rename(*CastChecked<UMaterialExpressionComposite>(CompositeNode->MaterialExpression)->SubgraphName);
+
+				CompositeNode->BoundGraph->AddExpression(InputPinBase, false);
+				CompositeNode->BoundGraph->AddExpression(OutputPinBase, false);
+			}
+
+			if (FocusedGraphEd)
+			{
+				FocusedGraphEd->ClearSelectionSet();
+				FocusedGraphEd->SetNodeSelection(NewComposite->GraphNode, true);
+			}
+		}
+	}
+
+	RefreshExpressionPreviews();
+	if (FocusedGraphEd)
+	{
+		FocusedGraphEd->NotifyGraphChanged();
+	}
+	SetMaterialDirty();
+	return NewComposite;
+}
+
+UMaterialExpressionComment* FMaterialEditor::CreateNewMaterialExpressionComment(const FVector2D& NodePos, const UEdGraph* Graph)
+{
+	TSharedPtr<SGraphEditor> FocusedGraphEd = FocusedGraphEdPtr.Pin();
+	UMaterialGraph* ExpressionGraph = Graph ? CastChecked<UMaterialGraph>(const_cast<UEdGraph*>(Graph)) : Material->MaterialGraph;
+	ExpressionGraph->Modify();
+
 	UMaterialExpressionComment* NewComment = NULL;
 	{
 		const FScopedTransaction Transaction(NSLOCTEXT("UnrealEd", "MaterialEditorCreateComment", "Material Editor: Create comment"));
@@ -4169,7 +4781,7 @@ UMaterialExpressionComment* FMaterialEditor::CreateNewMaterialExpressionComment(
 		Material->EditorComments.Add( NewComment );
 
 		FSlateRect Bounds;
-		if (GraphEditor->GetBoundsForSelectedNodes(Bounds, 50.0f))
+		if (FocusedGraphEd && FocusedGraphEd->GetBoundsForSelectedNodes(Bounds, 50.0f))
 		{
 			NewComment->MaterialExpressionEditorX = Bounds.Left;
 			NewComment->MaterialExpressionEditorY = Bounds.Top;
@@ -4189,15 +4801,21 @@ UMaterialExpressionComment* FMaterialEditor::CreateNewMaterialExpressionComment(
 
 		NewComment->Text = NSLOCTEXT("K2Node", "CommentBlock_NewEmptyComment", "Comment").ToString();
 
-		Material->MaterialGraph->AddComment(NewComment, true);
+		ExpressionGraph->AddComment(NewComment, true);
 
 		// Select the new comment.
-		GraphEditor->ClearSelectionSet();
-		GraphEditor->SetNodeSelection(NewComment->GraphNode, true);
+		if (FocusedGraphEd)
+		{
+			FocusedGraphEd->ClearSelectionSet();
+			FocusedGraphEd->SetNodeSelection(NewComment->GraphNode, true);
+		}
 	}
 
 	Material->MarkPackageDirty();
-	GraphEditor->NotifyGraphChanged();
+	if (FocusedGraphEd)
+	{
+		FocusedGraphEd->NotifyGraphChanged();
+	}
 	SetMaterialDirty();
 	return NewComment;
 }
@@ -4213,7 +4831,10 @@ void FMaterialEditor::ForceRefreshExpressionPreviews()
 
 void FMaterialEditor::AddToSelection(UMaterialExpression* Expression)
 {
-	GraphEditor->SetNodeSelection(Expression->GraphNode, true);
+	if (TSharedPtr<SGraphEditor> FocusedGraphEd = FocusedGraphEdPtr.Pin())
+	{
+		FocusedGraphEd->SetNodeSelection(Expression->GraphNode, true);
+	}
 }
 
 void FMaterialEditor::JumpToExpression(UMaterialExpression* Expression)
@@ -4256,12 +4877,15 @@ void FMaterialEditor::JumpToExpression(UMaterialExpression* Expression)
 
 void FMaterialEditor::SelectAllNodes()
 {
-	GraphEditor->SelectAllNodes();
+	if (TSharedPtr<SGraphEditor> FocusedGraphEd = FocusedGraphEdPtr.Pin())
+	{
+		FocusedGraphEd->SelectAllNodes();
+	}
 }
 
 bool FMaterialEditor::CanSelectAllNodes() const
 {
-	return GraphEditor.IsValid();
+	return FocusedGraphEdPtr.IsValid();
 }
 
 void FMaterialEditor::DeleteSelectedNodes()
@@ -4292,51 +4916,16 @@ void FMaterialEditor::DeleteNodes(const TArray<UEdGraphNode*>& NodesToDelete)
 
 		{
 			const FScopedTransaction Transaction( NSLOCTEXT("UnrealEd", "MaterialEditorDelete", "Material Editor: Delete") );
-			Material->Modify();
-
-			for (int32 Index = 0; Index < NodesToDelete.Num(); ++Index)
-			{
-				if (NodesToDelete[Index]->CanUserDeleteNode())
-				{
-					// Break all node links first so that we don't update the material before deleting
-					NodesToDelete[Index]->BreakAllNodeLinks();
-
-					FBlueprintEditorUtils::RemoveNode(NULL, NodesToDelete[Index], true);
-
-					if (UMaterialGraphNode* GraphNode = Cast<UMaterialGraphNode>(NodesToDelete[Index]))
-					{
-						UMaterialExpression* MaterialExpression = GraphNode->MaterialExpression;
-
-						bHaveExpressionsToDelete = true;
-
-						DestroyColorPicker();
-
-						if( PreviewExpression == MaterialExpression )
-						{
-							// The expression being previewed is also being deleted
-							bPreviewExpressionDeleted = true;
-						}
-
-						MaterialExpression->Modify();
-						Material->Expressions.Remove( MaterialExpression );
-						Material->RemoveExpressionParameter(MaterialExpression);
-						// Make sure the deleted expression is caught by gc
-						MaterialExpression->MarkPendingKill();
-					}
-					else if (UMaterialGraphNode_Comment* CommentNode = Cast<UMaterialGraphNode_Comment>(NodesToDelete[Index]))
-					{
-						CommentNode->MaterialExpressionComment->Modify();
-						Material->EditorComments.Remove( CommentNode->MaterialExpressionComment );
-					}
-				}
-			}
-
+			DeleteNodesInternal(NodesToDelete, bHaveExpressionsToDelete, bPreviewExpressionDeleted);
 			Material->MaterialGraph->LinkMaterialExpressionsFromGraph();
 		} // ScopedTransaction
 
 		// Deselect all expressions and comments.
-		GraphEditor->ClearSelectionSet();
-		GraphEditor->NotifyGraphChanged();
+		if (TSharedPtr<SGraphEditor> FocusedGraphEd = FocusedGraphEdPtr.Pin())
+		{
+			FocusedGraphEd->ClearSelectionSet();
+			FocusedGraphEd->NotifyGraphChanged();
+		}
 
 		if ( bHaveExpressionsToDelete )
 		{
@@ -4377,19 +4966,25 @@ bool FMaterialEditor::CanDeleteNodes() const
 
 void FMaterialEditor::DeleteSelectedDuplicatableNodes()
 {
+	TSharedPtr<SGraphEditor> FocusedGraphEd = FocusedGraphEdPtr.Pin();
+	if (!FocusedGraphEd)
+	{
+		return;
+	}
+
 	// Cache off the old selection
 	const FGraphPanelSelectionSet OldSelectedNodes = GetSelectedNodes();
 
 	// Clear the selection and only select the nodes that can be duplicated
 	FGraphPanelSelectionSet RemainingNodes;
-	GraphEditor->ClearSelectionSet();
+	FocusedGraphEd->ClearSelectionSet();
 
 	for (FGraphPanelSelectionSet::TConstIterator SelectedIter(OldSelectedNodes); SelectedIter; ++SelectedIter)
 	{
 		UEdGraphNode* Node = Cast<UEdGraphNode>(*SelectedIter);
 		if ((Node != NULL) && Node->CanDuplicateNode())
 		{
-			GraphEditor->SetNodeSelection(Node, true);
+			FocusedGraphEd->SetNodeSelection(Node, true);
 		}
 		else
 		{
@@ -4401,13 +4996,83 @@ void FMaterialEditor::DeleteSelectedDuplicatableNodes()
 	DeleteSelectedNodes();
 
 	// Reselect whatever's left from the original selection after the deletion
-	GraphEditor->ClearSelectionSet();
+	FocusedGraphEd->ClearSelectionSet();
 
 	for (FGraphPanelSelectionSet::TConstIterator SelectedIter(RemainingNodes); SelectedIter; ++SelectedIter)
 	{
 		if (UEdGraphNode* Node = Cast<UEdGraphNode>(*SelectedIter))
 		{
-			GraphEditor->SetNodeSelection(Node, true);
+			FocusedGraphEd->SetNodeSelection(Node, true);
+		}
+	}
+}
+
+void FMaterialEditor::DeleteNodesInternal(const TArray<class UEdGraphNode*>& NodesToDelete, bool& bHaveExpressionsToDelete, bool& bPreviewExpressionDeleted)
+{
+	Material->Modify();
+
+	for (int32 Index = 0; Index < NodesToDelete.Num(); ++Index)
+	{
+		if (NodesToDelete[Index]->CanUserDeleteNode())
+		{
+			// If this is a user-selected pinbase, don't allow the delete to pass
+			if (Cast<UMaterialGraphNode_PinBase>(NodesToDelete[Index]) && GetSelectedNodes().Contains(NodesToDelete[Index]))
+			{
+				continue;
+			}
+
+			if (UMaterialGraphNode* GraphNode = Cast<UMaterialGraphNode>(NodesToDelete[Index]))
+			{
+				// Break all node links first so that we don't update the material before deleting
+				GraphNode->Modify();
+				GraphNode->BreakAllNodeLinks();
+
+				UMaterialExpression* MaterialExpression = GraphNode->MaterialExpression;
+
+				bHaveExpressionsToDelete = true;
+
+				DestroyColorPicker();
+
+				if( PreviewExpression == MaterialExpression )
+				{
+					// The expression being previewed is also being deleted
+					bPreviewExpressionDeleted = true;
+				}
+
+				if (UMaterialGraphNode_PinBase* PinBaseNode = Cast<UMaterialGraphNode_PinBase>(GraphNode))
+				{
+					UMaterialExpressionPinBase* PinBase = CastChecked<UMaterialExpressionPinBase>(MaterialExpression);
+					PinBase->DeleteReroutePins();
+				}
+						
+				if (UMaterialGraphNode_Composite* CompositeNode = Cast<UMaterialGraphNode_Composite>(GraphNode))
+				{
+					CloseDocumentTab(CompositeNode->BoundGraph);
+
+					UMaterialExpressionComposite* SubgraphComposite = CastChecked<UMaterialExpressionComposite>(MaterialExpression);
+					SubgraphComposite->Modify();
+
+					// Remove all subgraph nodes, note that composites remove their subgraph in DestroyNode
+					const TArray<UEdGraphNode*> SubgraphNodesToDelete = CompositeNode->BoundGraph->Nodes;
+					DeleteNodesInternal(SubgraphNodesToDelete, bHaveExpressionsToDelete, bPreviewExpressionDeleted);
+				}
+
+				MaterialExpression->Modify();
+				Material->Expressions.Remove( MaterialExpression );
+				Material->RemoveExpressionParameter(MaterialExpression);
+				// Make sure the deleted expression is caught by gc
+				MaterialExpression->MarkPendingKill();
+			}
+			else if (UMaterialGraphNode_Comment* CommentNode = Cast<UMaterialGraphNode_Comment>(NodesToDelete[Index]))
+			{
+				CommentNode->Modify();
+				CommentNode->BreakAllNodeLinks();
+				CommentNode->MaterialExpressionComment->Modify();
+				Material->EditorComments.Remove( CommentNode->MaterialExpressionComment );
+			}
+
+			// Now that we are done with the node, remove it
+			FBlueprintEditorUtils::RemoveNode(NULL, NodesToDelete[Index], true);
 		}
 	}
 }
@@ -4461,18 +5126,83 @@ bool FMaterialEditor::CanCopyNodes() const
 
 void FMaterialEditor::PasteNodes()
 {
-	PasteNodesHere(GraphEditor->GetPasteLocation());
+	if (TSharedPtr<SGraphEditor> FocusedGraphEd = FocusedGraphEdPtr.Pin())
+	{
+		PasteNodesHere(FocusedGraphEd->GetPasteLocation(), FocusedGraphEd->GetCurrentGraph());
+	}
+}
+
+void FMaterialEditor::PostPasteMaterialExpression(UMaterialExpression* NewExpression)
+{
+	// Deep copy subgraph expressions
+	if (UMaterialExpressionComposite* Composite = Cast<UMaterialExpressionComposite>(NewExpression))
+	{
+		UMaterialGraphNode_Composite* CompositeNode = Cast<UMaterialGraphNode_Composite>(Composite->GraphNode);
+		DeepCopyExpressions(CompositeNode->BoundGraph, Composite);
+
+		// We just updated all our child expressions, reconstruct node.
+		Composite->GraphNode->ReconstructNode();
+	}
+
+	// There can be only one default mesh paint texture.
+	UMaterialExpressionTextureBase* TextureSample = Cast<UMaterialExpressionTextureBase>(NewExpression);
+	if (TextureSample)
+	{
+		TextureSample->IsDefaultMeshpaintTexture = false;
+	}
+
+	NewExpression->UpdateParameterGuid(true, true);
+	Material->AddExpressionParameter(NewExpression, Material->EditorParameters);
+
+	UMaterialExpressionFunctionInput* FunctionInput = Cast<UMaterialExpressionFunctionInput>(NewExpression);
+	if (FunctionInput)
+	{
+		FunctionInput->ConditionallyGenerateId(true);
+		FunctionInput->ValidateName();
+	}
+
+	UMaterialExpressionFunctionOutput* FunctionOutput = Cast<UMaterialExpressionFunctionOutput>(NewExpression);
+	if (FunctionOutput)
+	{
+		FunctionOutput->ConditionallyGenerateId(true);
+		FunctionOutput->ValidateName();
+	}
+
+	UMaterialExpressionMaterialFunctionCall* FunctionCall = Cast<UMaterialExpressionMaterialFunctionCall>(NewExpression);
+	if (FunctionCall)
+	{
+		// When pasting new nodes, we don't want to break all node links as this information is used by UpdateMaterialAfterGraphChange() below,
+		// to recreate all the connections in the pasted group.
+		// Just update the function input/outputs here.
+		const bool bRecreateAndLinkNode = false;
+		FunctionCall->UpdateFromFunctionResource(bRecreateAndLinkNode);
+
+		// If an unknown material function has been pasted, remove the graph node pins (as the expression will also have had its inputs/outputs removed).
+		// This will be displayed as an orphaned "Unspecified Function" node.
+		if (FunctionCall->MaterialFunction == nullptr &&
+			FunctionCall->FunctionInputs.Num() == 0 &&
+			FunctionCall->FunctionOutputs.Num() == 0)
+		{
+			NewExpression->GraphNode->Pins.Empty();
+		}
+	}
 }
 
-void FMaterialEditor::PasteNodesHere(const FVector2D& Location)
+void FMaterialEditor::PasteNodesHere(const FVector2D& Location, const class UEdGraph* Graph)
 {
 	// Undo/Redo support
 	const FScopedTransaction Transaction( NSLOCTEXT("UnrealEd", "MaterialEditorPaste", "Material Editor: Paste") );
 	Material->MaterialGraph->Modify();
 	Material->Modify();
 
+	UMaterialGraph* ExpressionGraph = Graph ? CastChecked<UMaterialGraph>(const_cast<UEdGraph*>(Graph)) : Material->MaterialGraph;
+	ExpressionGraph->Modify();
+
 	// Clear the selection set (newly pasted stuff will be selected)
-	GraphEditor->ClearSelectionSet();
+	if (TSharedPtr<SGraphEditor> FocusedGraphEd = FocusedGraphEdPtr.Pin())
+	{
+		FocusedGraphEd->ClearSelectionSet();
+	}
 
 	// Grab the text to paste from the clipboard.
 	FString TextToImport;
@@ -4480,7 +5210,7 @@ void FMaterialEditor::PasteNodesHere(const FVector2D& Location)
 
 	// Import the nodes
 	TSet<UEdGraphNode*> PastedNodes;
-	FEdGraphUtilities::ImportNodesFromText(Material->MaterialGraph, TextToImport, /*out*/ PastedNodes);
+	FEdGraphUtilities::ImportNodesFromText(ExpressionGraph, TextToImport, /*out*/ PastedNodes);
 
 	//Average position of nodes so we can move them while still maintaining relative distances to each other
 	FVector2D AvgNodePosition(0.0f,0.0f);
@@ -4512,6 +5242,7 @@ void FMaterialEditor::PasteNodesHere(const FVector2D& Location)
 			UMaterialExpression* NewExpression = GraphNode->MaterialExpression;
 			NewExpression->Material = Material;
 			NewExpression->Function = MaterialFunction;
+			NewExpression->SubgraphExpression = ExpressionGraph->SubgraphExpression;
 
 			// Make sure the param name is valid after the paste
 			if (NewExpression->HasAParameterName())
@@ -4521,48 +5252,8 @@ void FMaterialEditor::PasteNodesHere(const FVector2D& Location)
 
 			Material->Expressions.Add(NewExpression);
 
-			// There can be only one default mesh paint texture.
-			UMaterialExpressionTextureBase* TextureSample = Cast<UMaterialExpressionTextureBase>( NewExpression );
-			if( TextureSample )
-			{
-				TextureSample->IsDefaultMeshpaintTexture = false;
-			}
-
-			NewExpression->UpdateParameterGuid(true, true);
-			Material->AddExpressionParameter(NewExpression, Material->EditorParameters);
-
-			UMaterialExpressionFunctionInput* FunctionInput = Cast<UMaterialExpressionFunctionInput>( NewExpression );
-			if( FunctionInput )
-			{
-				FunctionInput->ConditionallyGenerateId(true);
-				FunctionInput->ValidateName();
-			}
-
-			UMaterialExpressionFunctionOutput* FunctionOutput = Cast<UMaterialExpressionFunctionOutput>( NewExpression );
-			if( FunctionOutput )
-			{
-				FunctionOutput->ConditionallyGenerateId(true);
-				FunctionOutput->ValidateName();
-			}
-
-			UMaterialExpressionMaterialFunctionCall* FunctionCall = Cast<UMaterialExpressionMaterialFunctionCall>( NewExpression );
-			if( FunctionCall )
-			{
-				// When pasting new nodes, we don't want to break all node links as this information is used by UpdateMaterialAfterGraphChange() below,
-				// to recreate all the connections in the pasted group.
-				// Just update the function input/outputs here.
-				const bool bRecreateAndLinkNode = false;
-				FunctionCall->UpdateFromFunctionResource(bRecreateAndLinkNode);
-
-				// If an unknown material function has been pasted, remove the graph node pins (as the expression will also have had its inputs/outputs removed).
-				// This will be displayed as an orphaned "Unspecified Function" node.
-				if (FunctionCall->MaterialFunction == nullptr &&
-					FunctionCall->FunctionInputs.Num() == 0 &&
-					FunctionCall->FunctionOutputs.Num() == 0)
-				{
-					GraphNode->Pins.Empty();
-				}
-			}
+			ensure(NewExpression->GraphNode == GraphNode);
+			PostPasteMaterialExpression(NewExpression);
 		}
 		else if (UMaterialGraphNode_Comment* CommentNode = Cast<UMaterialGraphNode_Comment>(Node))
 		{
@@ -4571,12 +5262,16 @@ void FMaterialEditor::PasteNodesHere(const FVector2D& Location)
 				CommentNode->MaterialDirtyDelegate = Material->MaterialGraph->MaterialDirtyDelegate;
 				CommentNode->MaterialExpressionComment->Material = Material;
 				CommentNode->MaterialExpressionComment->Function = MaterialFunction;
+				CommentNode->MaterialExpressionComment->SubgraphExpression = ExpressionGraph->SubgraphExpression;
 				Material->EditorComments.Add(CommentNode->MaterialExpressionComment);
 			}
 		}
 
 		// Select the newly pasted stuff
-		GraphEditor->SetNodeSelection(Node, true);
+		if (TSharedPtr<SGraphEditor> FocusedGraphEd = FocusedGraphEdPtr.Pin())
+		{
+			FocusedGraphEd->SetNodeSelection(Node, true);
+		}
 
 		Node->NodePosX = (Node->NodePosX - AvgNodePosition.X) + Location.X ;
 		Node->NodePosY = (Node->NodePosY - AvgNodePosition.Y) + Location.Y ;
@@ -4590,7 +5285,10 @@ void FMaterialEditor::PasteNodesHere(const FVector2D& Location)
 	UpdateMaterialAfterGraphChange();
 
 	// Update UI
-	GraphEditor->NotifyGraphChanged();
+	if (TSharedPtr<SGraphEditor> FocusedGraphEd = FocusedGraphEdPtr.Pin())
+	{
+		FocusedGraphEd->NotifyGraphChanged();
+	}
 }
 
 bool FMaterialEditor::CanPasteNodes() const
@@ -4633,7 +5331,6 @@ FText FMaterialEditor::GetOriginalObjectName() const
 void FMaterialEditor::UpdateMaterialAfterGraphChange()
 {
 	FlushRenderingCommands();
-	
 	Material->MaterialGraph->LinkMaterialExpressionsFromGraph();
 
 	// Update the current preview material.
@@ -4652,32 +5349,95 @@ void FMaterialEditor::UpdateMaterialAfterGraphChange()
 
 	if (bHideUnrelatedNodes && !bLockNodeFadeState && bSelectRegularNode)
 	{
-		GraphEditor->ResetAllNodesUnrelatedStates();
+		if (TSharedPtr<SGraphEditor> FocusedGraphEd = FocusedGraphEdPtr.Pin())
+		{
+			FocusedGraphEd->ResetAllNodesUnrelatedStates();
+		}
 
 		HideUnrelatedNodes();
 	}
 }
 
+void FMaterialEditor::JumpToHyperlink(const UObject* ObjectReference)
+{
+	if (const UEdGraphNode* Node = Cast<const UEdGraphNode>(ObjectReference))
+	{
+		JumpToNode(Node);
+	}
+	else if (const UEdGraph* Graph = Cast<const UEdGraph>(ObjectReference))
+	{
+		// Navigating into things should re-use the current tab when it makes sense
+		FDocumentTracker::EOpenDocumentCause OpenMode = FDocumentTracker::OpenNewDocument;
+		if ((Graph == Material->MaterialGraph) || Cast<UMaterialGraphNode_Composite>(Graph->GetOuter()))
+		{
+			OpenMode = FDocumentTracker::NavigatingCurrentDocument;
+		}
+		else
+		{
+			// Walk up the outer chain to see if any tabs have a parent of this document open for edit, and if so
+			// we should reuse that one and drill in deeper instead
+			for (UObject* WalkPtr = const_cast<UEdGraph*>(Graph); WalkPtr != nullptr; WalkPtr = WalkPtr->GetOuter())
+			{
+				TArray< TSharedPtr<SDockTab> > TabResults;
+				if (FindOpenTabsContainingDocument(WalkPtr, /*out*/ TabResults))
+				{
+					// See if the parent was active
+					bool bIsActive = false;
+					for (TSharedPtr<SDockTab> Tab : TabResults)
+					{
+						if (Tab->IsActive())
+						{
+							bIsActive = true;
+							break;
+						}
+					}
+
+					if (bIsActive)
+					{
+						OpenMode = FDocumentTracker::NavigatingCurrentDocument;
+						break;
+					}
+				}
+			}
+		}
+
+		// Force it to open in a new document if shift is pressed
+		const bool bIsShiftPressed = FSlateApplication::Get().GetModifierKeys().IsShiftDown();
+		if (bIsShiftPressed)
+		{
+			OpenMode = FDocumentTracker::ForceOpenNewDocument;
+		}
+
+		// Open the document
+		OpenDocument(Graph, OpenMode);
+	}
+}
+
 int32 FMaterialEditor::GetNumberOfSelectedNodes() const
 {
-	return GraphEditor->GetSelectedNodes().Num();
+	if (TSharedPtr<SGraphEditor> FocusedGraphEd = FocusedGraphEdPtr.Pin())
+	{
+		return FocusedGraphEd->GetSelectedNodes().Num();
+	}
+
+	return 0;
 }
 
 FGraphPanelSelectionSet FMaterialEditor::GetSelectedNodes() const
 {
 	FGraphPanelSelectionSet CurrentSelection;
-	if (GraphEditor.IsValid())
+	if (FocusedGraphEdPtr.IsValid())
 	{
-		CurrentSelection = GraphEditor->GetSelectedNodes();
+		CurrentSelection = FocusedGraphEdPtr.Pin()->GetSelectedNodes();
 	}
 	return CurrentSelection;
 }
 
 void FMaterialEditor::GetBoundsForNode(const UEdGraphNode* InNode, class FSlateRect& OutRect, float InPadding) const
 {
-	if (GraphEditor.IsValid())
+	if (FocusedGraphEdPtr.IsValid())
 	{
-		GraphEditor->GetBoundsForNode(InNode, OutRect, InPadding);
+		FocusedGraphEdPtr.Pin()->GetBoundsForNode(InNode, OutRect, InPadding);
 	}
 }
 
@@ -4708,6 +5468,22 @@ void FMaterialEditor::RedoGraphAction()
 	int32 NumExpressions = Material->Expressions.Num();
 	GEditor->RedoTransaction();
 
+	if (TSharedPtr<SGraphEditor> FocusedGraphEd = FocusedGraphEdPtr.Pin())
+	{
+		// @TODO Find a more coherent check for this, rather than rely on sage knowledge that the schema won't exist.
+		// If our previous transaction created the current graph, it won't have a schema after undo.
+		if (FocusedGraphEd->GetCurrentGraph()->IsPendingKill())
+		{
+			CloseDocumentTab(FocusedGraphEd->GetCurrentGraph());
+			DocumentManager->CleanInvalidTabs();
+			DocumentManager->RefreshAllTabs();
+			GetTabManager()->TryInvokeTab(FMaterialEditorTabs::GraphEditor);
+		}
+
+		// Active graph can change above, re-acquire ptr
+		FocusedGraphEdPtr.Pin()->NotifyGraphChanged();
+	}
+
 	if(NumExpressions != Material->Expressions.Num())
 	{
 		Material->BuildEditorParameterList();
@@ -4716,75 +5492,209 @@ void FMaterialEditor::RedoGraphAction()
 	UpdateGenerator();
 }
 
+void FMaterialEditor::OnCollapseNodes()
+{
+	const UMaterialGraphSchema* Schema = GetDefault<UMaterialGraphSchema>();
+
+	// Does the selection set contain anything that is legal to collapse?
+	TSet<UEdGraphNode*> CollapsableNodes;
+	const FGraphPanelSelectionSet SelectedNodes = GetSelectedNodes();
+	for (FGraphPanelSelectionSet::TConstIterator NodeIt(SelectedNodes); NodeIt; ++NodeIt)
+	{
+		if (UEdGraphNode* SelectedNode = Cast<UEdGraphNode>(*NodeIt))
+		{
+			if (Schema->CanEncapuslateNode(*SelectedNode))
+			{
+				CollapsableNodes.Add(SelectedNode);
+			}
+		}
+	}
+	
+	// Sort for deterministic pin ordering, unaffected by selection order.
+	CollapsableNodes.StableSort([](const UEdGraphNode& A, const UEdGraphNode& B)
+	{
+		if (A.NodePosY == B.NodePosY)
+		{
+			return A.NodePosX > B.NodePosX;
+		}
+		return A.NodePosY < B.NodePosY;
+	});
+
+	// Collapse them
+	if (CollapsableNodes.Num())
+	{
+		const FScopedTransaction Transaction(FGraphEditorCommands::Get().CollapseNodes->GetDescription());
+		Material->Modify();
+		Material->MaterialGraph->Modify();
+
+		CollapseNodes(CollapsableNodes);
+	}
+}
+
+bool FMaterialEditor::CanCollapseNodes() const
+{
+	// Does the selection set contain anything that is legal to collapse?
+	const UMaterialGraphSchema* Schema = GetDefault<UMaterialGraphSchema>();
+	FGraphPanelSelectionSet SelectedNodes = GetSelectedNodes();
+	for (FGraphPanelSelectionSet::TConstIterator NodeIt(SelectedNodes); NodeIt; ++NodeIt)
+	{
+		if (UEdGraphNode* Node = Cast<UEdGraphNode>(*NodeIt))
+		{
+			if (Schema->CanEncapuslateNode(*Node))
+			{
+				return true;
+			}
+		}
+	}
+
+	return false;
+}
+
+void FMaterialEditor::OnExpandNodes()
+{
+	const FScopedTransaction Transaction(FGraphEditorCommands::Get().ExpandNodes->GetLabel());
+	Material->Modify();
+	Material->MaterialGraph->Modify();
+
+	TSet<UEdGraphNode*> ExpandedNodes;
+	TSharedPtr<SGraphEditor> FocusedGraphEd = FocusedGraphEdPtr.Pin();
+
+	// Expand selected nodes into the focused graph context.
+	const FGraphPanelSelectionSet SelectedNodes = GetSelectedNodes();
+
+	if (FocusedGraphEd)
+	{
+		FocusedGraphEd->ClearSelectionSet();
+	}
+
+	for (FGraphPanelSelectionSet::TConstIterator NodeIt(SelectedNodes); NodeIt; ++NodeIt)
+	{
+		ExpandedNodes.Empty();
+		bool bExpandedNodesNeedUniqueGuid = true;
+
+		DocumentManager->CleanInvalidTabs();
+
+		if (UMaterialGraphNode_Composite* SelectedCompositeNode = Cast<UMaterialGraphNode_Composite>(*NodeIt))
+		{
+			// No need to assign unique GUIDs since the source graph will be removed.
+			bExpandedNodesNeedUniqueGuid = false;
+
+			// Expand the composite node back into the world
+			UEdGraph* SourceGraph = SelectedCompositeNode->BoundGraph;
+			ExpandNode(SelectedCompositeNode, SourceGraph, /*inout*/ ExpandedNodes);
+
+			FBlueprintEditorUtils::RemoveGraph(nullptr, SourceGraph, EGraphRemoveFlags::None);
+			SourceGraph->MarkPendingKill();
+		}
+
+		UEdGraphNode* SourceNode = CastChecked<UEdGraphNode>(*NodeIt);
+		check(SourceNode);
+		MoveNodesToAveragePos(ExpandedNodes, FVector2D(SourceNode->NodePosX, SourceNode->NodePosY), bExpandedNodesNeedUniqueGuid);
+	}
+
+	UMaterialExpression* SubgraphExpression = FocusedGraphEd ? Cast<UMaterialGraph>(FocusedGraphEd->GetCurrentGraph())->SubgraphExpression : Material->MaterialGraph->SubgraphExpression;
+
+	for (UEdGraphNode* ExpandedNode : ExpandedNodes)
+	{
+		if (UMaterialGraphNode* MaterialNode = Cast<UMaterialGraphNode>(ExpandedNode))
+		{
+			MaterialNode->MaterialExpression->Modify();
+			MaterialNode->MaterialExpression->SubgraphExpression = SubgraphExpression;
+		}
+		else if (UMaterialGraphNode_Comment* CommentNode = Cast<UMaterialGraphNode_Comment>(ExpandedNode))
+		{
+			CommentNode->MaterialExpressionComment->Modify();
+			CommentNode->MaterialExpressionComment->SubgraphExpression = SubgraphExpression;
+		}
+	}
+
+	UpdateMaterialAfterGraphChange();
+}
+
+bool FMaterialEditor::CanExpandNodes() const
+{
+	// Does the selection set contain any composite nodes that are legal to expand?
+	const FGraphPanelSelectionSet SelectedNodes = GetSelectedNodes();
+	for (FGraphPanelSelectionSet::TConstIterator NodeIt(SelectedNodes); NodeIt; ++NodeIt)
+	{
+		if (Cast<UMaterialGraphNode_Composite>(*NodeIt))
+		{
+			return true;
+		}
+	}
+
+	return false;
+}
+
 void FMaterialEditor::OnAlignTop()
 {
-	if (GraphEditor.IsValid())
+	if (FocusedGraphEdPtr.IsValid())
 	{
-		GraphEditor->OnAlignTop();
+		FocusedGraphEdPtr.Pin()->OnAlignTop();
 	}
 }
 
 void FMaterialEditor::OnAlignMiddle()
 {
-	if (GraphEditor.IsValid())
+	if (FocusedGraphEdPtr.IsValid())
 	{
-		GraphEditor->OnAlignMiddle();
+		FocusedGraphEdPtr.Pin()->OnAlignMiddle();
 	}
 }
 
 void FMaterialEditor::OnAlignBottom()
 {
-	if (GraphEditor.IsValid())
+	if (FocusedGraphEdPtr.IsValid())
 	{
-		GraphEditor->OnAlignBottom();
+		FocusedGraphEdPtr.Pin()->OnAlignBottom();
 	}
 }
 
 void FMaterialEditor::OnAlignLeft()
 {
-	if (GraphEditor.IsValid())
+	if (FocusedGraphEdPtr.IsValid())
 	{
-		GraphEditor->OnAlignLeft();
+		FocusedGraphEdPtr.Pin()->OnAlignLeft();
 	}
 }
 
 void FMaterialEditor::OnAlignCenter()
 {
-	if (GraphEditor.IsValid())
+	if (FocusedGraphEdPtr.IsValid())
 	{
-		GraphEditor->OnAlignCenter();
+		FocusedGraphEdPtr.Pin()->OnAlignCenter();
 	}
 }
 
 void FMaterialEditor::OnAlignRight()
 {
-	if (GraphEditor.IsValid())
+	if (FocusedGraphEdPtr.IsValid())
 	{
-		GraphEditor->OnAlignRight();
+		FocusedGraphEdPtr.Pin()->OnAlignRight();
 	}
 }
 
 void FMaterialEditor::OnStraightenConnections()
 {
-	if (GraphEditor.IsValid())
+	if (FocusedGraphEdPtr.IsValid())
 	{
-		GraphEditor->OnStraightenConnections();
+		FocusedGraphEdPtr.Pin()->OnStraightenConnections();
 	}
 }
 
 void FMaterialEditor::OnDistributeNodesH()
 {
-	if (GraphEditor.IsValid())
+	if (FocusedGraphEdPtr.IsValid())
 	{
-		GraphEditor->OnDistributeNodesH();
+		FocusedGraphEdPtr.Pin()->OnDistributeNodesH();
 	}
 }
 
 void FMaterialEditor::OnDistributeNodesV()
 {
-	if (GraphEditor.IsValid())
+	if (FocusedGraphEdPtr.IsValid())
 	{
-		GraphEditor->OnDistributeNodesV();
+		FocusedGraphEdPtr.Pin()->OnDistributeNodesV();
 	}
 }
 
@@ -4793,7 +5703,10 @@ void FMaterialEditor::PostUndo(bool bSuccess)
 {
 	if (bSuccess)
 	{	
-		GraphEditor->ClearSelectionSet();
+		if (TSharedPtr<SGraphEditor> FocusedGraphEd = FocusedGraphEdPtr.Pin())
+		{
+			FocusedGraphEd->ClearSelectionSet();
+		}
 		
 		Material->BuildEditorParameterList();
 
@@ -4803,7 +5716,26 @@ void FMaterialEditor::PostUndo(bool bSuccess)
 		UpdatePreviewViewportsVisibility();
 
 		RefreshExpressionPreviews();
-		GraphEditor->NotifyGraphChanged();
+
+		// Remove any tabs are that are pending kill or otherwise invalid UObject pointers.
+		bool bNeedOpenGraphEditor = false;
+
+		if (TSharedPtr<SGraphEditor> FocusedGraphEd = FocusedGraphEdPtr.Pin())
+		{
+			// @TODO Find a more coherent check for this, rather than rely on sage knowledge that the schema won't exist.
+			// If our previous transaction created the current graph, it won't have a schema after undo.
+			if (!FocusedGraphEd->GetCurrentGraph()->GetSchema())
+			{
+				CloseDocumentTab(FocusedGraphEd->GetCurrentGraph());
+				DocumentManager->CleanInvalidTabs();
+				DocumentManager->RefreshAllTabs();
+				GetTabManager()->TryInvokeTab(FMaterialEditorTabs::GraphEditor);
+			}
+
+			// Active graph can change above, re-acquire ptr
+			FocusedGraphEdPtr.Pin()->NotifyGraphChanged();
+		}
+
 		SetMaterialDirty();
 
 		UpdateGenerator();
@@ -5097,10 +6029,15 @@ void FMaterialEditor::OnColorPickerCommitted(FLinearColor LinearColor)
 	MaterialCustomPrimitiveDataWidget->UpdateEditorInstance(MaterialEditorInstance);
 }
 
-TSharedRef<SGraphEditor> FMaterialEditor::CreateGraphEditorWidget()
+/** Create new tab for the supplied graph - don't call this directly, instead call OpenDocument to track history.*/
+TSharedRef<SGraphEditor> FMaterialEditor::CreateGraphEditorWidget(TSharedRef<class FTabInfo> InTabInfo, class UEdGraph* InGraph)
 {
-	GraphEditorCommands = MakeShareable( new FUICommandList );
+	check((InGraph != nullptr) && Cast<UMaterialGraph>(InGraph));
+
+	if (!GraphEditorCommands)
 	{
+		GraphEditorCommands = MakeShareable(new FUICommandList);
+
 		// Editing commands
 		GraphEditorCommands->MapAction( FGenericCommands::Get().SelectAll,
 			FExecuteAction::CreateSP( this, &FMaterialEditor::SelectAllNodes ),
@@ -5207,6 +6144,19 @@ TSharedRef<SGraphEditor> FMaterialEditor::CreateGraphEditorWidget()
 			FCanExecuteAction::CreateSP(this, &FMaterialEditor::CanGoToDocumentation)
 		);
 
+		// Collapse Node Commands
+		GraphEditorCommands->MapAction( FGraphEditorCommands::Get().CollapseNodes,
+			FExecuteAction::CreateSP( this, &FMaterialEditor::OnCollapseNodes ),
+			FCanExecuteAction::CreateSP( this, &FMaterialEditor::CanCollapseNodes )
+		);
+
+		GraphEditorCommands->MapAction( FGraphEditorCommands::Get().ExpandNodes,
+			FExecuteAction::CreateSP( this, &FMaterialEditor::OnExpandNodes ),
+			FCanExecuteAction::CreateSP( this, &FMaterialEditor::CanExpandNodes ),
+			FIsActionChecked(),
+			FIsActionButtonVisible::CreateSP( this, &FMaterialEditor::CanExpandNodes )
+		);
+
 		// Alignment Commands
 		GraphEditorCommands->MapAction(FGraphEditorCommands::Get().AlignNodesTop,
 			FExecuteAction::CreateSP(this, &FMaterialEditor::OnAlignTop)
@@ -5252,11 +6202,14 @@ TSharedRef<SGraphEditor> FMaterialEditor::CreateGraphEditorWidget()
 	InEvents.OnNodeDoubleClicked = FSingleNodeEvent::CreateSP(this, &FMaterialEditor::OnNodeDoubleClicked);
 	InEvents.OnTextCommitted = FOnNodeTextCommitted::CreateSP(this, &FMaterialEditor::OnNodeTitleCommitted);
 	InEvents.OnVerifyTextCommit = FOnNodeVerifyTextCommit::CreateSP(this, &FMaterialEditor::OnVerifyNodeTextCommit);
-	InEvents.OnSpawnNodeByShortcut = SGraphEditor::FOnSpawnNodeByShortcut::CreateSP(this, &FMaterialEditor::OnSpawnGraphNodeByShortcut, static_cast<UEdGraph*>(Material->MaterialGraph));
+	InEvents.OnSpawnNodeByShortcut = SGraphEditor::FOnSpawnNodeByShortcut::CreateSP(this, &FMaterialEditor::OnSpawnGraphNodeByShortcut, static_cast<UEdGraph*>(InGraph));
 
 	// Create the title bar widget
 	TSharedPtr<SWidget> TitleBarWidget = SNew(SMaterialEditorTitleBar)
+		.EdGraphObj(InGraph)
 		.TitleText(this, &FMaterialEditor::GetOriginalObjectName)
+		.OnDifferentGraphCrumbClicked(this, &FMaterialEditor::OnChangeBreadCrumbGraph)
+		.HistoryNavigationWidget(InTabInfo->CreateHistoryNavigationWidget())
 		.MaterialInfoList(&MaterialInfoList);
 
 	return SNew(SGraphEditor)
@@ -5264,9 +6217,11 @@ TSharedRef<SGraphEditor> FMaterialEditor::CreateGraphEditorWidget()
 		.IsEditable(true)
 		.TitleBar(TitleBarWidget)
 		.Appearance(this, &FMaterialEditor::GetGraphAppearance)
-		.GraphToEdit(Material->MaterialGraph)
+		.GraphToEdit(InGraph)
 		.GraphEvents(InEvents)
 		.ShowGraphStateOverlay(false)
+		.OnNavigateHistoryBack(FSimpleDelegate::CreateSP(this, &FMaterialEditor::NavigateTab, FDocumentTracker::NavigateBackwards))
+		.OnNavigateHistoryForward(FSimpleDelegate::CreateSP(this, &FMaterialEditor::NavigateTab, FDocumentTracker::NavigateForwards))
 		.AssetEditorToolkit(this->AsShared());
 }
 
@@ -5296,6 +6251,78 @@ FGraphAppearanceInfo FMaterialEditor::GetGraphAppearance() const
 	return AppearanceInfo;
 }
 
+void FMaterialEditor::DeepCopyExpressions(UMaterialGraph* CopyGraph, UMaterialExpression* NewSubgraphExpression)
+{
+	if (!CopyGraph || !NewSubgraphExpression)
+	{
+		return;
+	}
+	
+	CopyGraph->Modify();
+	CopyGraph->SubgraphExpression = NewSubgraphExpression;
+
+	// Duplicate subnodes
+	auto DuplicateExpression = [&](auto* Expression)
+	{
+		using ExpressionType = std::remove_pointer<decltype(Expression)>::type;
+		return Cast<ExpressionType>(UMaterialEditingLibrary::DuplicateMaterialExpression(Material, MaterialFunction, Expression));
+	};
+
+	for (UEdGraphNode* Node : CopyGraph->Nodes)
+	{
+		if (UMaterialGraphNode* MaterialNode = Cast<UMaterialGraphNode>(Node))
+		{
+			MaterialNode->Modify();
+			MaterialNode->Rename(/*NewName=*/ NULL, /*NewOuter=*/ CopyGraph);
+			if (UMaterialExpressionPinBase* PinBase = Cast<UMaterialExpressionPinBase>(MaterialNode->MaterialExpression))
+			{
+				UMaterialExpressionPinBase* OldPinBase = PinBase;
+				UMaterialExpressionPinBase* NewPinBase = DuplicateExpression(OldPinBase);
+				MaterialNode->MaterialExpression = NewPinBase;
+
+				NewPinBase->SubgraphExpression = NewSubgraphExpression;
+				NewPinBase->ReroutePins.Empty();
+				for (FCompositeReroute& Reroute : OldPinBase->ReroutePins)
+				{
+					UMaterialExpressionReroute* DupReroute = DuplicateExpression(Reroute.Expression);
+					DupReroute->SubgraphExpression = NewSubgraphExpression;
+					NewPinBase->ReroutePins.Add({ Reroute.Name, DupReroute });
+				}
+
+				UMaterialExpressionComposite* SubGraphComposite = CastChecked<UMaterialExpressionComposite>(NewSubgraphExpression);
+				if (NewPinBase->PinDirection == EGPD_Output)
+				{
+					SubGraphComposite->InputExpressions = NewPinBase;
+				}
+				else
+				{
+					SubGraphComposite->OutputExpressions = NewPinBase;
+				}
+			}
+			else
+			{
+				MaterialNode->MaterialExpression = DuplicateExpression(MaterialNode->MaterialExpression);
+				MaterialNode->MaterialExpression->SubgraphExpression = NewSubgraphExpression;
+			}
+			
+			// GraphNode is transient so it won't be duplicated. 
+			MaterialNode->MaterialExpression->GraphNode = MaterialNode;
+			PostPasteMaterialExpression(MaterialNode->MaterialExpression);
+
+		}
+		else if (UMaterialGraphNode_Comment* Comment = Cast<UMaterialGraphNode_Comment>(MaterialNode))
+		{
+			Comment->Modify();
+			Comment->Rename(/*NewName=*/ NULL, /*NewOuter=*/ CopyGraph);
+			Comment->MaterialExpressionComment = DuplicateExpression(Comment->MaterialExpressionComment);
+			Comment->MaterialExpressionComment->SubgraphExpression = NewSubgraphExpression;
+
+			// GraphNode is transient so it won't be duplicated. 
+			Comment->MaterialExpressionComment->GraphNode = Comment;
+		}
+	}
+}
+
 void FMaterialEditor::CleanUnusedExpressions()
 {
 	TArray<UEdGraphNode*> UnusedNodes;
@@ -5333,8 +6360,11 @@ void FMaterialEditor::CleanUnusedExpressions()
 			Material->MaterialGraph->LinkMaterialExpressionsFromGraph();
 		} // ScopedTransaction
 
-		GraphEditor->ClearSelectionSet();
-		GraphEditor->NotifyGraphChanged();
+		if (TSharedPtr<SGraphEditor> FocusedGraphEd = FocusedGraphEdPtr.Pin())
+		{
+			FocusedGraphEd->ClearSelectionSet();
+			FocusedGraphEd->NotifyGraphChanged();
+		}
 
 		SetMaterialDirty();
 	}
@@ -5464,7 +6494,10 @@ void FMaterialEditor::OnSelectedNodesChanged(const TSet<class UObject*>& NewSele
 	
 	if (bHideUnrelatedNodes && !bLockNodeFadeState)
 	{
-		GraphEditor->ResetAllNodesUnrelatedStates();
+		if (TSharedPtr<SGraphEditor> FocusedGraphEd = FocusedGraphEdPtr.Pin())
+		{
+			FocusedGraphEd->ResetAllNodesUnrelatedStates();
+		}
 
 		if (bSelectRegularNode)
 		{
@@ -5477,7 +6510,11 @@ void FMaterialEditor::OnNodeDoubleClicked(class UEdGraphNode* Node)
 {
 	UMaterialGraphNode* GraphNode = Cast<UMaterialGraphNode>(Node);
 
-	if (GraphNode)
+	if (Node && Node->CanJumpToDefinition())
+	{
+		Node->JumpToDefinition();
+	}
+	else if (GraphNode)
 	{
 		UMaterialExpressionConstant3Vector* Constant3Expression = Cast<UMaterialExpressionConstant3Vector>(GraphNode->MaterialExpression);
 		UMaterialExpressionConstant4Vector* Constant4Expression = Cast<UMaterialExpressionConstant4Vector>(GraphNode->MaterialExpression);
@@ -5529,7 +6566,7 @@ void FMaterialEditor::OnNodeDoubleClicked(class UEdGraphNode* Node)
 
 			// Open a color picker 
 			FColorPickerArgs PickerArgs;
-			PickerArgs.ParentWidget = GraphEditor;//AsShared();
+			PickerArgs.ParentWidget = FocusedGraphEdPtr.Pin();
 			PickerArgs.bUseAlpha = ChannelEditStruct.Alpha != nullptr;
 			PickerArgs.bOnlyRefreshOnOk = false;
 			PickerArgs.bOnlyRefreshOnMouseUp = true;
@@ -5580,6 +6617,12 @@ void FMaterialEditor::OnNodeDoubleClicked(class UEdGraphNode* Node)
 
 void FMaterialEditor::OnRenameNode()
 {
+	TSharedPtr<SGraphEditor> FocusedGraphEd = FocusedGraphEdPtr.Pin();
+	if (!FocusedGraphEd)
+	{
+		return;
+	}
+
 	const FGraphPanelSelectionSet SelectedNodes = GetSelectedNodes();
 	for (FGraphPanelSelectionSet::TConstIterator NodeIt(SelectedNodes); NodeIt; ++NodeIt)
 	{
@@ -5587,7 +6630,7 @@ void FMaterialEditor::OnRenameNode()
 		if (SelectedNode != nullptr && SelectedNode->GetCanRenameNode())
 		{
 			bool ToRename = true;
-			GraphEditor->IsNodeTitleVisible(SelectedNode, ToRename);
+			FocusedGraphEd->IsNodeTitleVisible(SelectedNode, ToRename);
 			break;
 		}
 	}
diff --git a/Engine/Source/Editor/MaterialEditor/Private/MaterialEditor.h b/Engine/Source/Editor/MaterialEditor/Private/MaterialEditor.h
--- a/Engine/Source/Editor/MaterialEditor/Private/MaterialEditor.h
+++ b/Engine/Source/Editor/MaterialEditor/Private/MaterialEditor.h
@@ -14,6 +14,8 @@
 #include "EditorUndoClient.h"
 #include "MaterialShared.h"
 #include "Toolkits/IToolkitHost.h"
+#include "WorkflowOrientedApp/WorkflowTabFactory.h"
+#include "WorkflowOrientedApp/WorkflowTabManager.h"
 #include "IMaterialEditor.h"
 #include "Editor/PropertyEditor/Public/IDetailsView.h"
 #include "SMaterialEditorViewport.h"
@@ -34,6 +36,7 @@ class UEdGraph;
 class UFactory;
 class UMaterialEditorOptions;
 class UMaterialExpressionComment;
+class UMaterialExpressionComposite;
 class UMaterialInstance;
 class UMaterialGraphNode;
 struct FGraphAppearanceInfo;
@@ -242,6 +245,8 @@ struct FMaterialInfo
 class FMaterialEditor : public IMaterialEditor, public FGCObject, public FTickableGameObject, public FEditorUndoClient, public FNotifyHook
 {
 public:
+	// @todo This is a hack for now until we reconcile the default toolbar with application modes
+	void RegisterToolbarTab(const TSharedRef<class FTabManager>& TabManager);
 	virtual void RegisterTabSpawners(const TSharedRef<class FTabManager>& TabManager) override;
 	virtual void UnregisterTabSpawners(const TSharedRef<class FTabManager>& TabManager) override;
 
@@ -346,6 +351,17 @@ public:
 	 * Updates flags on the Material Nodes to avoid expensive look up calls when rendering
 	 */
 	void UpdateGraphNodeStates();
+
+	/**
+	 * Updates flags on the Material Nodes for a single graph and it's subgraphs.
+	 * 
+	 * @param	Graph					Material graph to update
+	 * @param	ErrorMaterialResource	Material Resource containing known errors to draw for graph
+	 * @param	VisibleExpressions		List of expressions that should be visible (As in not grayed out / disabled).
+	 * @param	bShowAllNodes           True implies all nodes should be visible.
+	 * @return	True if we updated the error state of this material during graph update.
+	 */
+	bool UpdateGraphNodeState(UEdGraph* Graph, const FMaterialResource* ErrorMaterialResource, TArray<UMaterialExpression*>& VisibleExpressions, bool bShowAllNodes);
 	
 	// Widget Accessors
 	TSharedRef<class IDetailsView> GetDetailView() const {return MaterialDetailsView.ToSharedRef();}
@@ -379,17 +395,34 @@ public:
 	/** Pan the view to center on a particular node */
 	void JumpToNode(const UEdGraphNode* Node);
 
+	/** Called when graph editor focus is changed */
+	virtual void OnGraphEditorFocused(const TSharedRef<class SGraphEditor>& InGraphEditor);
+
+	/** Called when the graph editor tab is backgrounded */
+	virtual void OnGraphEditorBackgrounded(const TSharedRef<SGraphEditor>& InGraphEditor);
+
+	// Finds any open tabs containing the specified document and adds them to the specified array; returns true if at least one is found
+	bool FindOpenTabsContainingDocument(const UObject* DocumentID, /*inout*/ TArray< TSharedPtr<SDockTab> >& Results);
+
+	/** Open workflow document tab */
+	TSharedPtr<SDockTab> OpenDocument(const UObject* DocumentID, FDocumentTracker::EOpenDocumentCause Cause);
+
+	/** Close workflow document tab */
+	void CloseDocumentTab(const UObject* DocumentID);
+
 	// IMaterial Editor Interface
-	virtual UMaterialExpression* CreateNewMaterialExpression(UClass* NewExpressionClass, const FVector2D& NodePos, bool bAutoSelect, bool bAutoAssignResource) override;
-	virtual UMaterialExpressionComment* CreateNewMaterialExpressionComment(const FVector2D& NodePos) override;
+	virtual UMaterialExpression* CreateNewMaterialExpression(UClass* NewExpressionClass, const FVector2D& NodePos, bool bAutoSelect, bool bAutoAssignResource, const class UEdGraph* Graph = nullptr) override;
+	virtual UMaterialExpressionComposite* CreateNewMaterialExpressionComposite(const FVector2D& NodePos, const class UEdGraph* Graph = nullptr) override;
+	virtual UMaterialExpressionComment* CreateNewMaterialExpressionComment(const FVector2D& NodePos, const class UEdGraph* Graph = nullptr) override;
 	virtual void ForceRefreshExpressionPreviews() override;
 	virtual void AddToSelection(UMaterialExpression* Expression) override;
 	virtual void JumpToExpression(UMaterialExpression* Expression) override;
 	virtual void DeleteSelectedNodes() override;
 	virtual FText GetOriginalObjectName() const override;
 	virtual void UpdateMaterialAfterGraphChange() override;
+	virtual void JumpToHyperlink(const UObject* ObjectReference) override; 
 	virtual bool CanPasteNodes() const override;
-	virtual void PasteNodesHere(const FVector2D& Location) override;
+	virtual void PasteNodesHere(const FVector2D& Location, const class UEdGraph* Graph = nullptr) override;
 	virtual int32 GetNumberOfSelectedNodes() const override;
 	virtual TSet<UObject*> GetSelectedNodes() const override;
 	virtual void GetBoundsForNode(const UEdGraphNode* InNode, class FSlateRect& OutRect, float InPadding) const override;
@@ -456,6 +489,12 @@ public:
 	/** Configuration class used to store editor settings across sessions. */
 	UMaterialEditorOptions* EditorOptions;
 	
+	/** Document manager for workflow tabs */
+	TSharedPtr<FDocumentTracker> DocumentManager;
+
+	/** Factory that spawns graph editors; used to look up all tabs spawned by it. */
+	TWeakPtr<FDocumentTabFactory> GraphEditorTabFactoryPtr;
+
 protected:
 	//~ FAssetEditorToolkit interface
 	virtual void GetSaveableObjects(TArray<UObject*>& OutObjects) const override;
@@ -513,6 +552,8 @@ protected:
 	bool CanDeleteNodes() const;
 	/** Delete only the currently selected nodes that can be duplicated */
 	void DeleteSelectedDuplicatableNodes();
+	/** Recursively deletes nodes if needed */
+	void DeleteNodesInternal(const TArray<class UEdGraphNode*>& NodesToDelete, bool& bHaveExpressionsToDelete, bool& bPreviewExpressionDeleted);
 
 	/** Copy the currently selected nodes */
 	void CopySelectedNodes();
@@ -522,6 +563,9 @@ protected:
 	/** Paste the contents of the clipboard */
 	void PasteNodes();
 
+	/** Handle transient properties, and other things that can't be done in PostPasteNode */
+	void PostPasteMaterialExpression(UMaterialExpression* NewExpression);
+
 	/** Cut the currently selected nodes */
 	void CutSelectedNodes();
 	/** Whether we are able to cut the currently selected nodes */
@@ -544,6 +588,16 @@ protected:
 	/** Check if node can be renamed */
 	bool CanRenameNodes() const;
 
+	/** Collapse node group */
+	void OnCollapseNodes();
+	/** Check if nodes can be collapsed */
+	bool CanCollapseNodes() const;
+
+	/** Expand node group */
+	void OnExpandNodes();
+	/** Check if node can be expanded */
+	bool CanExpandNodes() const;
+
 	void OnAlignTop();
 	void OnAlignMiddle();
 	void OnAlignBottom();
@@ -578,6 +632,73 @@ private:
 	/** Updates the 3D and UI preview viewport visibility based on material domain */
 	void UpdatePreviewViewportsVisibility();
 
+	//@TODO: these methods are mostly C&P from BlueprintEditor, consider consolidating logic to graph editor. Note: We don't support macros / functions / tunnels / split pins, and also have material expression specific considerations. */
+	// void CollapseNodesIntoGraph(UEdGraphNode* InGatewayNode, UMaterialGraphNode* InEntryNode, UMaterialGraphNode* InResultNode, UEdGraph* InSourceGraph, UEdGraph* InDestinationGraph, TSet<UEdGraphNode*>& InCollapsableNodes);
+	// void CollapseNodes(TSet<class UEdGraphNode*>& InCollapsableNodes);
+	// static void ExpandNode(UEdGraphNode* InNodeToExpand, UEdGraph* InSourceGraph, TSet<UEdGraphNode*>& OutExpandedNodes);
+	// void MoveNodesToAveragePos(TSet<UEdGraphNode*>& AverageNodes, FVector2D SourcePos, bool bExpandedNodesNeedUniqueGuid = false) const;
+	// static void MoveNodesToGraph(TArray<UEdGraphNode*>& SourceNodes, UEdGraph* DestinationGraph, TSet<UEdGraphNode*>& OutExpandedNodes, UEdGraphNode** OutEntry, UEdGraphNode** OutResult, const bool bIsCollapsedGraph = false);
+	// bool CollapseGatewayNode(UK2Node* InNode, UEdGraphNode* InEntryNode, UEdGraphNode* InResultNode, TSet<UEdGraphNode*>* OutExpandedNodes = nullptr) const;
+
+	/**
+	 * Collapses a selection of nodes into a graph for composite.
+	 *
+	 * @param InGatewayNode				The node replacing the selection of nodes
+	 * @param InEntryNode				The entry node in the graph
+	 * @param InResultNode				The result node in the graph
+	 * @param InSourceGraph				The graph the selection is from
+	 * @param InDestinationGraph		The destination graph to move the selected nodes to
+	 * @param InCollapsableNodes		The selection of nodes being collapsed
+	 */
+	void CollapseNodesIntoGraph(UEdGraphNode* InGatewayNode, UMaterialGraphNode* InEntryNode, UMaterialGraphNode* InResultNode, UEdGraph* InSourceGraph, UEdGraph* InDestinationGraph, TSet<UEdGraphNode*>& InCollapsableNodes);
+
+	/** Called when a selection of nodes are being collapsed into a sub-graph */
+	void CollapseNodes(TSet<class UEdGraphNode*>& InCollapsableNodes);
+
+	/**
+	 * Expands passed in node 
+	 * 
+	 * @param InNodeToExpand			The node containing the selection of nodes that ill be removed
+	 * @param InSourceGraph				The graph containing the original node, 
+	 * @param OutExpandedNodes			The nodes expanded into the source graph
+	 */
+	void ExpandNode(UEdGraphNode* InNodeToExpand, UEdGraph* InSourceGraph, TSet<UEdGraphNode*>& OutExpandedNodes);
+
+	/**
+	* Move the given set of nodes to an average spot near the Source position
+	*
+	* @param AverageNodes					The nodes to move
+	* @param SourcePos						The source position used to average the nodes around
+	* @param bExpandedNodesNeedUniqueGuid	If true then a new Guid will be generated for each node in the set
+	*/
+	void MoveNodesToAveragePos(TSet<UEdGraphNode*>& AverageNodes, FVector2D SourcePos, bool bExpandedNodesNeedUniqueGuid = false) const;
+
+	/**
+	* Move every node from the source graph to the destination graph. Add Each node that is moved to the OutExpandedNodes set.
+	* If the source graph is a function graph, keep track of the entry and result nodes in the given Out Parameters.
+	*
+	* @param SourceNodes		Nodes to move
+	* @param DestinationGraph	Graph to move nodes to
+	* @param OutExpandedNodes	Set of each node that was moved from the source to destination graph
+	* @param OutEntry			Pointer to the function entry node
+	* @param OutResult			Pointer to the function result node
+	* @param bIsCollapsedGraph	Whether or not the source graph is collapsed
+	**/
+	static void MoveNodesToGraph(TArray<UEdGraphNode*>& SourceNodes, UEdGraph* DestinationGraph, TSet<UEdGraphNode*>& OutExpandedNodes, UEdGraphNode** OutEntry, UEdGraphNode** OutResult, const bool bIsCollapsedGraph = false);
+
+	/**
+	 * Makes connections into/or out of the gateway node, connect directly to the associated networks on the opposite side of the tunnel
+	 * When done, none of the pins on the gateway node will be connected to anything.
+	 * Requires both this gateway node and it's associated node to be in the same graph already (post-merging)
+	 *
+	 * @param InGatewayNode			The function or tunnel node
+	 * @param InEntryNode			The entry node in the inner graph
+	 * @param InResultNode			The result node in the inner graph
+	 *
+	 * @return						Returns TRUE if successful
+	 */
+	static bool CollapseGatewayNode(UEdGraphNode* InNode, UEdGraphNode* InEntryNode, UEdGraphNode* InResultNode, TSet<UEdGraphNode*>* OutExpandedNodes = nullptr);
+
 	/** Helper functions for the quality level node display toggling */
 	void SetQualityPreview(EMaterialQualityLevel::Type NewQuality);
 	bool IsQualityPreviewChecked(EMaterialQualityLevel::Type TestQuality);
@@ -749,7 +870,7 @@ private:
 	void OnColorPickerCommitted(FLinearColor LinearColor);
 
 	/** Create new graph editor widget */
-	TSharedRef<class SGraphEditor> CreateGraphEditorWidget();
+	TSharedRef<class SGraphEditor> CreateGraphEditorWidget(TSharedRef<class FTabInfo> InTabInfo, class UEdGraph* InGraph);
 
 	/** Gets the current Material Graph's appearance */
 	FGraphAppearanceInfo GetGraphAppearance() const;
@@ -759,6 +880,16 @@ private:
 	 */
 	void CleanUnusedExpressions();
 
+	/**
+	 * Perform a deep copy of all expressions within the given graph.
+	 *
+	 * @param	CopyGraph, graph whose expression's need a deep copy
+	 * @param	NewSubgraphExpression, material expression that will be used as the subgraph expression for all
+	 *			deep copied expressions, and will become the graph's new subgraph expression
+	 *
+	 */
+	void DeepCopyExpressions(UMaterialGraph* CopyGraph, UMaterialExpression* NewSubgraphExpression);
+
 	/**
 	 * Displays a warning message to the user if the expressions to remove would cause any issues
 	 *
@@ -775,7 +906,6 @@ private:
 
 private:
 	TSharedRef<SDockTab> SpawnTab_Preview(const FSpawnTabArgs& Args);
-	TSharedRef<SDockTab> SpawnTab_GraphCanvas(const FSpawnTabArgs& Args);
 	TSharedRef<SDockTab> SpawnTab_MaterialProperties(const FSpawnTabArgs& Args);
 	TSharedRef<SDockTab> SpawnTab_Palette(const FSpawnTabArgs& Args);
 	TSharedRef<SDockTab> SpawnTab_Find(const FSpawnTabArgs& Args);
@@ -786,14 +916,16 @@ private:
 
 	void OnFinishedChangingProperties(const FPropertyChangedEvent& PropertyChangedEvent);
 	void OnFinishedChangingParametersFromOverview(const FPropertyChangedEvent& PropertyChangedEvent);
+	void OnChangeBreadCrumbGraph(class UEdGraph* InGraph);
 	void GeneratorRowsRefreshed();
 	void UpdateGenerator();
+	void NavigateTab(FDocumentTracker::EOpenDocumentCause InCause);
 private:
 	/** Property View */
 	TSharedPtr<class IDetailsView> MaterialDetailsView;
 
-	/** New Graph Editor */
-	TSharedPtr<class SGraphEditor> GraphEditor;
+	/** Currently focused graph editor */
+	TWeakPtr<class SGraphEditor> FocusedGraphEdPtr;
 
 	/** Preview Viewport widget */
 	TSharedPtr<class SMaterialEditor3DPreviewViewport> PreviewViewport;
@@ -867,17 +999,6 @@ private:
 	TSharedPtr<FExtensibilityManager> MenuExtensibilityManager;
 	TSharedPtr<FExtensibilityManager> ToolBarExtensibilityManager;
 
-	/**	The tab ids for the material editor */
-	static const FName PreviewTabId;
-	static const FName GraphCanvasTabId;
-	static const FName PropertiesTabId;
-	static const FName PaletteTabId;
-	static const FName FindTabId;
-	static const FName PreviewSettingsTabId;
-	static const FName ParameterDefaultsTabId;
-	static const FName CustomPrimitiveTabId;
-	static const FName LayerPropertiesTabId;
-
 	/** Object that stores all of the possible parameters we can edit. */
 	class UMaterialEditorPreviewParameters* MaterialEditorInstance;
 
diff --git a/Engine/Source/Editor/MaterialEditor/Private/MaterialEditorDetailCustomization.cpp b/Engine/Source/Editor/MaterialEditor/Private/MaterialEditorDetailCustomization.cpp
--- a/Engine/Source/Editor/MaterialEditor/Private/MaterialEditorDetailCustomization.cpp
+++ b/Engine/Source/Editor/MaterialEditor/Private/MaterialEditorDetailCustomization.cpp
@@ -12,6 +12,8 @@
 #include "DetailCategoryBuilder.h"
 #include "Materials/MaterialExpressionScalarParameter.h"
 #include "Materials/MaterialExpressionVectorParameter.h"
+#include "Materials/MaterialExpressionComposite.h"
+#include "Materials/MaterialExpressionPinBase.h"
 #include "MaterialLayersFunctionsCustomization.h"
 #include "PropertyEditorModule.h"
 #include "MaterialEditor/MaterialEditorPreviewParameters.h"
@@ -846,4 +848,44 @@ void FMaterialExpressionLayersParameterDetails::CustomizeDetails(IDetailLayoutBu
 	
 }
 
+TSharedRef<class IDetailCustomization> FMaterialExpressionCompositeDetails::MakeInstance()
+{
+	return MakeShareable(new FMaterialExpressionCompositeDetails());
+}
+
+void FMaterialExpressionCompositeDetails::CustomizeDetails(IDetailLayoutBuilder& DetailLayout)
+{
+	FName DefaultCategory = NAME_None;
+	IDetailCategoryBuilder& Category = DetailLayout.EditCategory(DefaultCategory);
+	Category.AddProperty("SubgraphName");
+
+	const FName MaterialExpressionCategory = TEXT("MaterialExpression");
+	IDetailCategoryBuilder& ExpressionCategory = DetailLayout.EditCategory(MaterialExpressionCategory);
+
+	TArray< TWeakObjectPtr<UObject> > Objects;
+	DetailLayout.GetObjectsBeingCustomized(Objects);
+	
+	for (const auto& WeakObjectPtr : Objects)
+	{
+		UObject* Object = WeakObjectPtr.Get();
+
+		UMaterialExpressionComposite* Composite = Cast<UMaterialExpressionComposite>(Object);
+
+		if (Composite)
+		{
+			if (IDetailPropertyRow* PinBaseRow = ExpressionCategory.AddExternalObjectProperty({Composite->InputExpressions}, GET_MEMBER_NAME_CHECKED(UMaterialExpressionPinBase, ReroutePins)))
+			{
+				PinBaseRow->DisplayName(LOCTEXT("InputPinsComposite", "Input Pins"));
+				PinBaseRow->Visibility(EVisibility::Visible);
+			}
+
+			if (IDetailPropertyRow* PinBaseRow = ExpressionCategory.AddExternalObjectProperty({Composite->OutputExpressions }, GET_MEMBER_NAME_CHECKED(UMaterialExpressionPinBase, ReroutePins)))
+			{
+				PinBaseRow->DisplayName(LOCTEXT("OutPinsComposite", "Output Pins"));
+				PinBaseRow->Visibility(EVisibility::Visible);
+			}
+		}
+	}
+}
+
 #undef LOCTEXT_NAMESPACE
diff --git a/Engine/Source/Editor/MaterialEditor/Private/MaterialEditorDetailCustomization.h b/Engine/Source/Editor/MaterialEditor/Private/MaterialEditorDetailCustomization.h
--- a/Engine/Source/Editor/MaterialEditor/Private/MaterialEditorDetailCustomization.h
+++ b/Engine/Source/Editor/MaterialEditor/Private/MaterialEditorDetailCustomization.h
@@ -121,6 +121,22 @@ private:
 	TArray<TSharedPtr<FString>> ParametersSource;
 };
 
+
+/**
+ * Customizes the details of a Composite node,
+ * specifically exposing pinbases.
+ */
+class FMaterialExpressionCompositeDetails : public IDetailCustomization
+{
+public:
+
+	/** Makes a new instance of this detail layout class for a specific detail view requesting it */
+	static TSharedRef<class IDetailCustomization> MakeInstance();
+
+	/** IDetailCustomization interface */
+	virtual void CustomizeDetails(IDetailLayoutBuilder& DetailLayout) override;
+};
+
 class FMaterialDetailCustomization : public IDetailCustomization
 {
 public:
diff --git a/Engine/Source/Editor/MaterialEditor/Private/MaterialEditorModes.cpp b/Engine/Source/Editor/MaterialEditor/Private/MaterialEditorModes.cpp
new file mode 100644
--- /dev/null
+++ b/Engine/Source/Editor/MaterialEditor/Private/MaterialEditorModes.cpp
@@ -0,0 +1,111 @@
+// Copyright Epic Games, Inc. All Rights Reserved.
+
+#include "MaterialEditorModes.h"
+#include "Settings/EditorExperimentalSettings.h"
+#include "ToolMenus.h"
+
+#include "EdGraph/EdGraph.h"
+#include "EdGraph/EdGraphSchema.h"
+
+#include "Framework/Docking/LayoutExtender.h"
+#include "MaterialStats.h"
+#include "MaterialEditor.h"
+#include "MaterialEditorTabs.h"
+#include "Tabs/MaterialEditorTabFactories.h"
+
+#define LOCTEXT_NAMESPACE "MaterialEditor"
+
+const FName FMaterialEditorApplicationModes::StandardMaterialEditorMode(TEXT("GraphName"));
+
+TSharedPtr<FTabManager::FLayout> FMaterialEditorApplicationModes::GetDefaultEditorLayout(TSharedPtr<FMaterialEditor> InMaterialEditor)
+{
+	const TSharedRef<FTabManager::FLayout> StandaloneDefaultLayout = FTabManager::NewLayout("Standalone_MaterialEditor_Layout_v11")
+	->AddArea
+	(
+		FTabManager::NewPrimaryArea() ->SetOrientation(Orient_Vertical)
+		->Split
+		(
+			FTabManager::NewStack()
+			->SetSizeCoefficient(0.1f)
+			->SetHideTabWell( true )
+			->AddTab( InMaterialEditor->GetToolbarTabId(), ETabState::OpenedTab )
+		)
+		->Split
+		(
+			FTabManager::NewSplitter() ->SetOrientation(Orient_Horizontal) ->SetSizeCoefficient(0.9f)
+			->Split
+			(
+				FTabManager::NewSplitter() ->SetOrientation(Orient_Vertical) ->SetSizeCoefficient(0.2f)
+				->Split
+				(
+					FTabManager::NewStack()
+					->SetHideTabWell( true )
+					->AddTab( FMaterialEditorTabs::PreviewTabId, ETabState::OpenedTab )
+				)
+				->Split
+				(
+					FTabManager::NewStack()
+					->AddTab( FMaterialEditorTabs::PropertiesTabId, ETabState::OpenedTab )
+					->AddTab( FMaterialEditorTabs::PreviewSettingsTabId, ETabState::ClosedTab )
+					->AddTab( FMaterialEditorTabs::ParameterDefaultsTabId, ETabState::OpenedTab )
+					->AddTab( FMaterialEditorTabs::CustomPrimitiveTabId, ETabState::ClosedTab )
+					->AddTab( FMaterialEditorTabs::LayerPropertiesTabId, ETabState::ClosedTab )
+					->SetForegroundTab( FMaterialEditorTabs::PropertiesTabId )
+				)
+			)
+			->Split
+			(
+				FTabManager::NewSplitter() ->SetOrientation( Orient_Vertical )
+				->SetSizeCoefficient(0.80f)
+				->Split
+				(
+					FTabManager::NewStack() 
+					->SetSizeCoefficient(0.8f)
+					->SetHideTabWell( true )
+					->AddTab(FMaterialEditorTabs::GraphEditor, ETabState::OpenedTab)
+				)
+				->Split
+				(
+					FTabManager::NewStack()
+					->SetSizeCoefficient( 0.20f )
+						->AddTab( FMaterialStats::GetGridStatsTabName(), ETabState::ClosedTab )
+						->AddTab( FMaterialStats::GetGridOldStatsTabName(), ETabState::ClosedTab )
+					->AddTab( FMaterialEditorTabs::FindTabId, ETabState::ClosedTab )
+				)
+			)
+			->Split
+			(
+				FTabManager::NewSplitter() ->SetOrientation(Orient_Horizontal) ->SetSizeCoefficient(0.2f)
+				->Split
+				(
+					FTabManager::NewStack()
+					->AddTab( FMaterialEditorTabs::PaletteTabId, ETabState::OpenedTab )
+				)
+			)
+		)
+	);
+
+	return StandaloneDefaultLayout;
+}
+
+FMaterialEditorApplicationMode::FMaterialEditorApplicationMode(TSharedPtr<FMaterialEditor> InMaterialEditor)
+	: FApplicationMode(FMaterialEditorApplicationModes::StandardMaterialEditorMode, FMaterialEditorApplicationModes::GetLocalizedMode)
+{
+	MyMaterialEditor = InMaterialEditor;
+	MaterialEditorTabFactories.RegisterFactory(InMaterialEditor->GraphEditorTabFactoryPtr.Pin().ToSharedRef());
+	TabLayout = FMaterialEditorApplicationModes::GetDefaultEditorLayout(InMaterialEditor);
+}
+
+void FMaterialEditorApplicationMode::RegisterTabFactories(TSharedPtr<FTabManager> InTabManager)
+{
+	TSharedPtr<FMaterialEditor> MaterialEditor = MyMaterialEditor.Pin();
+
+	MaterialEditor->RegisterToolbarTab(InTabManager.ToSharedRef());
+
+	// Mode-specific setup
+	MaterialEditor->PushTabFactories(MaterialEditorTabFactories);
+
+	FApplicationMode::RegisterTabFactories(InTabManager);
+}
+
+#undef LOCTEXT_NAMESPACE
\ No newline at end of file
diff --git a/Engine/Source/Editor/MaterialEditor/Private/MaterialEditorTabs.cpp b/Engine/Source/Editor/MaterialEditor/Private/MaterialEditorTabs.cpp
new file mode 100644
--- /dev/null
+++ b/Engine/Source/Editor/MaterialEditor/Private/MaterialEditorTabs.cpp
@@ -0,0 +1,13 @@
+// Copyright Epic Games, Inc. All Rights Reserved.
+
+#include "MaterialEditorTabs.h"
+
+const FName FMaterialEditorTabs::PreviewTabId(TEXT("MaterialEditor_Preview"));
+const FName FMaterialEditorTabs::PropertiesTabId(TEXT("MaterialEditor_MaterialProperties"));
+const FName FMaterialEditorTabs::PaletteTabId(TEXT("MaterialEditor_Palette"));
+const FName FMaterialEditorTabs::FindTabId(TEXT("MaterialEditor_Find"));
+const FName FMaterialEditorTabs::GraphEditor(TEXT("Document"));
+const FName FMaterialEditorTabs::PreviewSettingsTabId(TEXT("MaterialEditor_PreviewSettings"));
+const FName FMaterialEditorTabs::ParameterDefaultsTabId(TEXT("MaterialEditor_ParameterDefaults"));
+const FName FMaterialEditorTabs::CustomPrimitiveTabId(TEXT("MaterialEditor_CustomPrimitiveData"));
+const FName FMaterialEditorTabs::LayerPropertiesTabId(TEXT("MaterialInstanceEditor_MaterialLayerProperties"));
\ No newline at end of file
diff --git a/Engine/Source/Editor/MaterialEditor/Private/MaterialEditorTabs.h b/Engine/Source/Editor/MaterialEditor/Private/MaterialEditorTabs.h
new file mode 100644
--- /dev/null
+++ b/Engine/Source/Editor/MaterialEditor/Private/MaterialEditorTabs.h
@@ -0,0 +1,24 @@
+// Copyright Epic Games, Inc. All Rights Reserved.
+
+#pragma once
+
+#include "CoreMinimal.h"
+
+/**  Material editor tabs identifiers */
+struct FMaterialEditorTabs
+{
+	/**	The tab id for the Material fixture types tab */
+	static const FName PreviewTabId;
+	static const FName PropertiesTabId;
+	static const FName PaletteTabId;
+	static const FName FindTabId;
+	static const FName GraphEditor;
+	static const FName PreviewSettingsTabId;
+	static const FName ParameterDefaultsTabId;
+	static const FName CustomPrimitiveTabId;
+	static const FName LayerPropertiesTabId;
+
+	// Disable default constructor
+	FMaterialEditorTabs() = delete;
+};
+
diff --git a/Engine/Source/Editor/MaterialEditor/Private/MaterialEditorUtilities.cpp b/Engine/Source/Editor/MaterialEditor/Private/MaterialEditorUtilities.cpp
--- a/Engine/Source/Editor/MaterialEditor/Private/MaterialEditorUtilities.cpp
+++ b/Engine/Source/Editor/MaterialEditor/Private/MaterialEditorUtilities.cpp
@@ -5,6 +5,7 @@
 #include "EdGraph/EdGraph.h"
 #include "Materials/Material.h"
 #include "MaterialGraph/MaterialGraphSchema.h"
+#include "MaterialGraph/MaterialGraphNode_Composite.h"
 #include "IMaterialEditor.h"
 
 #include "Materials/MaterialExpressionFunctionInput.h"
@@ -14,6 +15,7 @@
 #include "Materials/MaterialExpressionStaticBool.h"
 #include "Materials/MaterialExpressionStaticSwitch.h"
 #include "Materials/MaterialExpressionComment.h"
+#include "Materials/MaterialExpressionComposite.h"
 #include "Materials/MaterialExpressionTextureSample.h"
 #include "Materials/MaterialExpressionTextureSampleParameter.h"
 #include "Materials/MaterialExpressionRuntimeVirtualTextureSampleParameter.h"
@@ -45,9 +47,19 @@ UMaterialExpression* FMaterialEditorUtilities::CreateNewMaterialExpression(const
 	TSharedPtr<class IMaterialEditor> MaterialEditor = GetIMaterialEditorForObject(Graph);
 	if (MaterialEditor.IsValid())
 	{
-		return MaterialEditor->CreateNewMaterialExpression(NewExpressionClass, NodePos, bAutoSelect, bAutoAssignResource);
+		return MaterialEditor->CreateNewMaterialExpression(NewExpressionClass, NodePos, bAutoSelect, bAutoAssignResource, Graph);
 	}
-	return NULL;
+	return nullptr;
+}
+
+UMaterialExpressionComposite* FMaterialEditorUtilities::CreateNewMaterialExpressionComposite(const class UEdGraph* Graph, const FVector2D& NodePos)
+{
+	TSharedPtr<class IMaterialEditor> MaterialEditor = GetIMaterialEditorForObject(Graph);
+	if (MaterialEditor.IsValid())
+	{
+		return MaterialEditor->CreateNewMaterialExpressionComposite(NodePos, Graph);
+	}
+	return nullptr;
 }
 
 UMaterialExpressionComment* FMaterialEditorUtilities::CreateNewMaterialExpressionComment(const class UEdGraph* Graph, const FVector2D& NodePos)
@@ -55,9 +67,9 @@ UMaterialExpressionComment* FMaterialEditorUtilities::CreateNewMaterialExpressio
 	TSharedPtr<class IMaterialEditor> MaterialEditor = GetIMaterialEditorForObject(Graph);
 	if (MaterialEditor.IsValid())
 	{
-		return MaterialEditor->CreateNewMaterialExpressionComment(NodePos);
+		return MaterialEditor->CreateNewMaterialExpressionComment(NodePos, Graph);
 	}
-	return NULL;
+	return nullptr;
 }
 
 void FMaterialEditorUtilities::ForceRefreshExpressionPreviews(const class UEdGraph* Graph)
@@ -579,6 +591,15 @@ TSharedPtr<class IMaterialEditor> FMaterialEditorUtilities::GetIMaterialEditorFo
 	// Find the associated Material
 	UMaterial* Material = Cast<UMaterial>(ObjectToFocusOn->GetOuter());
 
+	// May be inspecting a subgraph, in which case, get the material from the composite node 
+	if (!Material)
+	{
+		if (UMaterialGraphNode_Composite* Composite = Cast<UMaterialGraphNode_Composite>(ObjectToFocusOn->GetOuter()))
+		{
+			Material = Composite->MaterialExpression->Material;
+		}
+	}
+
 	TSharedPtr<IMaterialEditor> MaterialEditor;
 	if (Material != NULL)
 	{
@@ -591,6 +612,16 @@ TSharedPtr<class IMaterialEditor> FMaterialEditorUtilities::GetIMaterialEditorFo
 	return MaterialEditor;
 }
 
+void FMaterialEditorUtilities::BringFocusAttentionOnObject(const UObject* ObjectToFocusOn)
+{
+	TSharedPtr<IMaterialEditor> MaterialEditor = GetIMaterialEditorForObject(ObjectToFocusOn);
+	if (MaterialEditor.IsValid())
+	{
+		MaterialEditor->FocusWindow();
+		MaterialEditor->JumpToHyperlink(ObjectToFocusOn);
+	}
+}
+
 void FMaterialEditorUtilities::AddMaterialExpressionCategory(FGraphActionMenuBuilder& ActionMenuBuilder, FText CategoryName, TArray<struct FMaterialExpression>* MaterialExpressions, bool bMaterialFunction)
 {
 	// Get type of dragged pin
@@ -603,7 +634,7 @@ void FMaterialEditorUtilities::AddMaterialExpressionCategory(FGraphActionMenuBui
 	for (int32 Index = 0; Index < MaterialExpressions->Num(); ++Index)
 	{
 		const FMaterialExpression& MaterialExpression = (*MaterialExpressions)[Index];
-		if (IsAllowedExpressionType(MaterialExpression.MaterialClass, bMaterialFunction))
+		if (IsAllowedExpressionType(MaterialExpression.MaterialClass, bMaterialFunction) && MaterialExpression.MaterialClass != UMaterialExpressionComposite::StaticClass())
 		{
 			if (!ActionMenuBuilder.FromPin || HasCompatibleConnection(MaterialExpression.MaterialClass, FromPinType, ActionMenuBuilder.FromPin->Direction, bMaterialFunction))
 			{
diff --git a/Engine/Source/Editor/MaterialEditor/Private/MaterialInstanceEditor.cpp b/Engine/Source/Editor/MaterialEditor/Private/MaterialInstanceEditor.cpp
--- a/Engine/Source/Editor/MaterialEditor/Private/MaterialInstanceEditor.cpp
+++ b/Engine/Source/Editor/MaterialEditor/Private/MaterialInstanceEditor.cpp
@@ -800,7 +800,7 @@ void FMaterialInstanceEditor::UpdatePreviewViewportsVisibility()
 
 void FMaterialInstanceEditor::RegisterToolBar()
 {
-	const FName MenuName = GetToolMenuToolbarName();
+	FName MenuName = FAssetEditorToolkit::GetToolMenuToolbarName();
 	if (!UToolMenus::Get()->IsMenuRegistered(MenuName))
 	{
 		UToolMenu* ToolBar = UToolMenus::Get()->RegisterMenu(MenuName, "AssetEditor.DefaultToolBar", EMultiBoxType::ToolBar);
diff --git a/Engine/Source/Editor/MaterialEditor/Private/SMaterialEditorTitleBar.cpp b/Engine/Source/Editor/MaterialEditor/Private/SMaterialEditorTitleBar.cpp
--- a/Engine/Source/Editor/MaterialEditor/Private/SMaterialEditorTitleBar.cpp
+++ b/Engine/Source/Editor/MaterialEditor/Private/SMaterialEditorTitleBar.cpp
@@ -4,49 +4,152 @@
 #include "Widgets/Text/STextBlock.h"
 #include "Widgets/Layout/SSpacer.h"
 #include "Widgets/Views/SListView.h"
+
+#include "Widgets/SBoxPanel.h"
+#include "Widgets/Layout/SBorder.h"
+#include "Widgets/Layout/SScrollBox.h"
+#include "Widgets/Layout/SSeparator.h"
+#include "Widgets/Images/SImage.h"
+
 #include "Styling/CoreStyle.h"
 #include "EditorStyleSet.h"
+#include "EdGraph/EdGraph.h"
+#include "EdGraph/EdGraphSchema.h"
+#include "Widgets/Navigation/SBreadcrumbTrail.h"
+
+#include "MaterialGraph/MaterialGraph.h"
+
+#define LOCTEXT_NAMESPACE "SMaterialEditorTitleBar"
 
 void SMaterialEditorTitleBar::Construct(const FArguments& InArgs)
 {
-	Visibility = EVisibility::HitTestInvisible;
+	
+	EdGraphObj = InArgs._EdGraphObj;
+	OnDifferentGraphCrumbClicked = InArgs._OnDifferentGraphCrumbClicked;
 
+	check(EdGraphObj);
+
+	// Set-up shared breadcrumb defaults
+	FMargin BreadcrumbTrailPadding = FMargin(4.f, 2.f);
+	const FSlateBrush* BreadcrumbButtonImage = FEditorStyle::GetBrush("BreadcrumbTrail.Delimiter");
+	
 	this->ChildSlot
 	[
-		SNew(SBorder)
-		.BorderImage( FEditorStyle::GetBrush( TEXT("Graph.TitleBackground") ) )
+		SNew(SVerticalBox)
+		+ SVerticalBox::Slot()
 		.HAlign(HAlign_Fill)
+		.AutoHeight()
 		[
-			SNew(SVerticalBox)
-			// Title text/icon
-			+SVerticalBox::Slot()
-			.HAlign(HAlign_Center)
-			.Padding(10)
-			.AutoHeight()
-			[
-				SNew(STextBlock)
-				.TextStyle( FEditorStyle::Get(), TEXT("GraphBreadcrumbButtonText") )
-				.Text( InArgs._TitleText )
-			]
-			+SVerticalBox::Slot()
-			.VAlign(VAlign_Top)
-			.AutoHeight()
+			SNew(SBorder)
+			.BorderImage( FEditorStyle::GetBrush( TEXT("Graph.TitleBackground") ) )
+			.HAlign(HAlign_Fill)
+			.AddMetaData<FTagMetaData>(FTagMetaData(TEXT("EventGraphTitleBar")))
 			[
 				SNew(SHorizontalBox)
+				+ SHorizontalBox::Slot()
+				.AutoWidth()
+				[
+					SNew(SSeparator)
+					.Orientation(Orient_Vertical)
+				]
+				+ SHorizontalBox::Slot()
+				.AutoWidth()
+				[
+					InArgs._HistoryNavigationWidget.ToSharedRef()
+				]
+				+ SHorizontalBox::Slot()
+				.AutoWidth()
+				[
+					SNew(SSeparator)
+					.Orientation(Orient_Vertical)
+				]
+				// Title text/icon
 				+SHorizontalBox::Slot()
-				.HAlign(HAlign_Left)
-				.FillWidth(1.f)
-				.Padding(5,0)
 				[
-					SAssignNew(MaterialInfoList, SListView<TSharedPtr<FMaterialInfo>>)
-					.ListItemsSource(InArgs._MaterialInfoList)
-					.OnGenerateRow(this, &SMaterialEditorTitleBar::MakeMaterialInfoWidget)
-					.SelectionMode( ESelectionMode::None )
-					.Visibility((InArgs._MaterialInfoList != NULL) ? EVisibility::Visible : EVisibility::Collapsed)
+					SNew(SHorizontalBox)
+					+SHorizontalBox::Slot()
+					.FillWidth(1.f)
+					[
+						SNew(SHorizontalBox)
+						+SHorizontalBox::Slot()
+						.AutoWidth()
+						.Padding( 10.0f,5.0f )
+						.VAlign(VAlign_Center)
+						[
+							SNew(SImage)
+							.Image( this, &SMaterialEditorTitleBar::GetTypeGlyph ) 
+						]
+						+SHorizontalBox::Slot()
+						.FillWidth(1.f)
+						[
+							SAssignNew(BreadcrumbTrailScrollBox, SScrollBox)
+							.Orientation(Orient_Horizontal)
+							.ScrollBarVisibility(EVisibility::Collapsed)
+
+							+SScrollBox::Slot()
+							.Padding(0.f)
+							.VAlign(VAlign_Center)
+							[
+								SNew(SHorizontalBox)
+								//show fake 'root' breadcrumb for the title
+								+SHorizontalBox::Slot()
+								.AutoWidth()
+								.VAlign(VAlign_Center)
+								.Padding(BreadcrumbTrailPadding)
+								[
+									SNew(STextBlock)
+									.Text(InArgs._TitleText)
+									.TextStyle( FEditorStyle::Get(), TEXT("GraphBreadcrumbButtonText") )
+									.Visibility( EVisibility::Visible )
+								]
+								+SHorizontalBox::Slot()
+								.AutoWidth()
+								.VAlign(VAlign_Center)
+								[
+									SNew(SImage)
+									.Image( BreadcrumbButtonImage )
+									.Visibility( EVisibility::Visible )
+								]
+
+								// New style breadcrumb
+								+SHorizontalBox::Slot()
+								.AutoWidth()
+								.VAlign(VAlign_Center)
+								[
+									SAssignNew(BreadcrumbTrail, SBreadcrumbTrail<UEdGraph*>)
+									.ButtonStyle(FEditorStyle::Get(), "GraphBreadcrumbButton")
+									.TextStyle(FEditorStyle::Get(), "GraphBreadcrumbButtonText")
+									.ButtonContentPadding( BreadcrumbTrailPadding )
+									.DelimiterImage( BreadcrumbButtonImage )
+									.PersistentBreadcrumbs( true )
+									.OnCrumbClicked( this, &SMaterialEditorTitleBar::OnBreadcrumbClicked )
+								]
+							]
+						]
+					]
 				]
 			]
 		]
+		+SVerticalBox::Slot()
+		.VAlign(VAlign_Top)
+		.AutoHeight()
+		[
+			SNew(SHorizontalBox)
+			+SHorizontalBox::Slot()
+			.HAlign(HAlign_Center)
+			.AutoWidth()
+			[
+				SAssignNew(MaterialInfoList, SListView<TSharedPtr<FMaterialInfo>>)
+				.ListItemsSource(InArgs._MaterialInfoList)
+				.OnGenerateRow(this, &SMaterialEditorTitleBar::MakeMaterialInfoWidget)
+				.SelectionMode( ESelectionMode::None )
+				.Visibility(EVisibility::Collapsed)
+			]
+		]
 	];
+
+	RebuildBreadcrumbTrail();
+	BreadcrumbTrailScrollBox->ScrollToEnd();
 }
 
 TSharedRef<ITableRow> SMaterialEditorTitleBar::MakeMaterialInfoWidget(TSharedPtr<FMaterialInfo> Item, const TSharedRef<STableViewBase>& OwnerTable)
@@ -80,5 +183,65 @@ TSharedRef<ITableRow> SMaterialEditorTitleBar::MakeMaterialInfoWidget(TSharedPtr
 
 void SMaterialEditorTitleBar::RequestRefresh()
 {
-	MaterialInfoList->RequestListRefresh();
+	if (MaterialInfoList)
+	{
+		MaterialInfoList->RequestListRefresh();
+	}
+
+	RebuildBreadcrumbTrail();
 }
+
+const FSlateBrush* SMaterialEditorTitleBar::GetTypeGlyph() const
+{
+	check(EdGraphObj != nullptr);
+	UMaterialGraph* MaterialGraph = CastChecked<UMaterialGraph>(EdGraphObj);
+	return FEditorStyle::GetBrush( MaterialGraph->SubgraphExpression ? TEXT("GraphEditor.SubGraph_24x") : TEXT("GraphEditor.EventGraph_24x") );
+}
+
+void SMaterialEditorTitleBar::OnBreadcrumbClicked(UEdGraph* const& Item)
+{
+	OnDifferentGraphCrumbClicked.ExecuteIfBound(Item);
+}
+
+void SMaterialEditorTitleBar::RebuildBreadcrumbTrail()
+{
+	// Build up a stack of graphs so we can pop them in reverse order and create breadcrumbs
+	TArray<UEdGraph*> Stack;
+	for (UEdGraph* OuterChain = EdGraphObj; OuterChain != nullptr; OuterChain = UEdGraph::GetOuterGraph(OuterChain))
+	{
+		Stack.Push(OuterChain);
+	}
+
+	BreadcrumbTrail->ClearCrumbs(false);
+
+	//Get the last object in the array
+	UEdGraph* LastObj = nullptr;
+	if( Stack.Num() > 0 )
+	{
+		LastObj = Stack[Stack.Num() -1];
+	}
+
+	while (Stack.Num() > 0)
+	{
+		UEdGraph* Graph = Stack.Pop();
+		
+		auto TitleText = TAttribute<FText>::Create(TAttribute<FText>::FGetter::CreateStatic<const UEdGraph*, const UEdGraph*>(&SMaterialEditorTitleBar::GetTitleForOneCrumb, LastObj, Graph));
+
+		BreadcrumbTrail->PushCrumb(TitleText, Graph);
+	}
+}
+
+FText SMaterialEditorTitleBar::GetTitleForOneCrumb(const UEdGraph* BaseGraph, const UEdGraph* CurrGraph)
+{
+	const UEdGraphSchema* Schema = CurrGraph->GetSchema();
+
+	FGraphDisplayInfo DisplayInfo;
+	Schema->GetGraphDisplayInformation(*CurrGraph, /*out*/ DisplayInfo);
+
+	FFormatNamedArguments Args;
+	Args.Add(TEXT("BreadcrumbDisplayName"), BaseGraph == CurrGraph ? LOCTEXT("BaseMaterialGraph", "Material Graph") : DisplayInfo.DisplayName);
+	Args.Add(TEXT("BreadcrumbNotes"), FText::FromString(DisplayInfo.GetNotesAsString()));
+	return FText::Format(LOCTEXT("BreadcrumbTitle", "{BreadcrumbDisplayName} {BreadcrumbNotes}"), Args);
+}
+
+#undef LOCTEXT_NAMESPACE
\ No newline at end of file
diff --git a/Engine/Source/Editor/MaterialEditor/Private/SMaterialEditorTitleBar.h b/Engine/Source/Editor/MaterialEditor/Private/SMaterialEditorTitleBar.h
--- a/Engine/Source/Editor/MaterialEditor/Private/SMaterialEditorTitleBar.h
+++ b/Engine/Source/Editor/MaterialEditor/Private/SMaterialEditorTitleBar.h
@@ -8,6 +8,10 @@
 #include "MaterialEditor.h"
 #include "Widgets/Views/STableViewBase.h"
 #include "Widgets/Views/STableRow.h"
+#include "BlueprintUtilities.h"
+
+class UEdGraph;
+class SScrollBox;
 
 //////////////////////////////////////////////////////////////////////////
 // SMaterialEditorTitleBar
@@ -17,12 +21,15 @@ class SMaterialEditorTitleBar : public SCompoundWidget
 public:
 	SLATE_BEGIN_ARGS( SMaterialEditorTitleBar )
 		: _TitleText()
-		, _MaterialInfoList(NULL)
+		, _EdGraphObj(nullptr)
+		, _MaterialInfoList(nullptr)
 	{}
 
 		SLATE_ATTRIBUTE( FText, TitleText )
-
+		SLATE_ARGUMENT( UEdGraph*, EdGraphObj )
 		SLATE_ARGUMENT( const TArray<TSharedPtr<FMaterialInfo>>*, MaterialInfoList )
+		SLATE_EVENT( FEdGraphEvent, OnDifferentGraphCrumbClicked )
+		SLATE_ARGUMENT( TSharedPtr<SWidget>, HistoryNavigationWidget )
 	SLATE_END_ARGS()
 
 	void Construct(const FArguments& InArgs);
@@ -33,6 +40,30 @@ public:
 	/** Request a refresh of the list view */
 	void RequestRefresh();
 
-protected:
+	/** Material info we display when compile errors occur. */
 	TSharedPtr<SListView<TSharedPtr<FMaterialInfo>>>	MaterialInfoList;
+
+protected:
+
+	TSharedPtr<SScrollBox> BreadcrumbTrailScrollBox;
+
+	/** Edited graph */
+	UEdGraph* EdGraphObj;
+
+	/** Breadcrumb trail widget */
+	TSharedPtr< SBreadcrumbTrail<UEdGraph*> > BreadcrumbTrail;
+
+	/** Callback to call when the user wants to change the active graph via the breadcrumb trail */
+	FEdGraphEvent OnDifferentGraphCrumbClicked;
+
+	/** Get the icon to use */
+	const FSlateBrush* GetTypeGlyph() const;
+
+	/** Function to fetch outer class which is of type UEGraph. */
+	UEdGraph* GetOuterGraph(UObject* Obj);
+
+	/** Helper methods */
+	void OnBreadcrumbClicked(UEdGraph* const& Item);
+	void RebuildBreadcrumbTrail();
+	static FText GetTitleForOneCrumb(const UEdGraph* BaseGraph, const UEdGraph* CurrGraph);
 };
diff --git a/Engine/Source/Editor/MaterialEditor/Private/Tabs/MaterialEditorTabFactories.cpp b/Engine/Source/Editor/MaterialEditor/Private/Tabs/MaterialEditorTabFactories.cpp
new file mode 100644
--- /dev/null
+++ b/Engine/Source/Editor/MaterialEditor/Private/Tabs/MaterialEditorTabFactories.cpp
@@ -0,0 +1,155 @@
+// Copyright Epic Games, Inc. All Rights Reserved.
+
+#include "Tabs/MaterialEditorTabFactories.h"
+#include "MaterialEditorTabs.h"
+#include "MaterialEditor.h"
+#include "MaterialGraph/MaterialGraph.h"
+
+#include "WorkflowOrientedApp/ApplicationMode.h"
+#include "Framework/Docking/TabManager.h"
+#include "Misc/BlacklistNames.h"
+
+#include "Widgets/Docking/SDockTab.h"
+
+#define LOCTEXT_NAMESPACE "MaterialWorkflowTabFactory"
+
+/////////////////////////////////////////////////////
+// FMaterialGraphEditorSummoner
+
+void FMaterialGraphEditorSummoner::OnTabActivated(TSharedPtr<SDockTab> Tab) const
+{
+	TSharedRef<SGraphEditor> GraphEditor = StaticCastSharedRef<SGraphEditor>(Tab->GetContent());
+	MaterialEditorPtr.Pin()->OnGraphEditorFocused(GraphEditor);
+}
+
+void FMaterialGraphEditorSummoner::OnTabBackgrounded(TSharedPtr<SDockTab> Tab) const
+{
+	TSharedRef<SGraphEditor> GraphEditor = StaticCastSharedRef<SGraphEditor>(Tab->GetContent());
+	MaterialEditorPtr.Pin()->OnGraphEditorBackgrounded(GraphEditor);
+}
+
+void FMaterialGraphEditorSummoner::OnTabRefreshed(TSharedPtr<SDockTab> Tab) const
+{
+	TSharedRef<SGraphEditor> GraphEditor = StaticCastSharedRef<SGraphEditor>(Tab->GetContent());
+	GraphEditor->NotifyGraphChanged();
+}
+
+void FMaterialGraphEditorSummoner::SaveState(TSharedPtr<SDockTab> Tab, TSharedPtr<FTabPayload> Payload) const
+{
+	TSharedRef<SGraphEditor> GraphEditor = StaticCastSharedRef<SGraphEditor>(Tab->GetContent());
+
+	FVector2D ViewLocation;
+	float ZoomAmount;
+	GraphEditor->GetViewLocation(ViewLocation, ZoomAmount);
+
+	UEdGraph* Graph = Payload->IsValid() ? FTabPayload_UObject::CastChecked<UEdGraph>(Payload) : nullptr;
+}
+
+FMaterialGraphEditorSummoner::FMaterialGraphEditorSummoner(TSharedPtr<class FMaterialEditor> InMaterialEditorPtr, FOnCreateGraphEditorWidget CreateGraphEditorWidgetCallback) : FDocumentTabFactoryForObjects<UEdGraph>("Document", InMaterialEditorPtr)
+, MaterialEditorPtr(InMaterialEditorPtr)
+, OnCreateGraphEditorWidget(CreateGraphEditorWidgetCallback)
+{
+}
+
+FTabSpawnerEntry& FMaterialGraphEditorSummoner::RegisterTabSpawner(TSharedRef<FTabManager> InTabManager, const FApplicationMode* CurrentApplicationMode) const
+{
+	FWorkflowTabSpawnInfo SpawnInfo;
+	SpawnInfo.TabManager = InTabManager;
+	SpawnInfo.Payload = FTabPayload_UObject::Make(MaterialEditorPtr.Pin()->Material->MaterialGraph);
+
+	TWeakPtr<FTabManager> WeakTabManager(InTabManager);
+	FTabSpawnerEntry& SpawnerEntry = InTabManager->RegisterTabSpawner(GetIdentifier(), FOnSpawnTab::CreateSP(this, &FMaterialGraphEditorSummoner::OnSpawnTab, WeakTabManager), FCanSpawnTab::CreateSP(this, &FMaterialGraphEditorSummoner::CanSpawnTab, WeakTabManager))
+		.SetDisplayName(LOCTEXT("GraphTabName", "Graph"));
+
+	if (CurrentApplicationMode)
+	{
+		SpawnerEntry.SetGroup(CurrentApplicationMode->GetWorkspaceMenuCategory());
+	}
+
+	// Add the tab icon to the menu entry if one was provided
+	const FSlateIcon& TabSpawnerIcon = GetTabSpawnerIcon(SpawnInfo);
+	if (TabSpawnerIcon.IsSet())
+	{
+		SpawnerEntry.SetIcon(TabSpawnerIcon);
+	}
+
+	return SpawnerEntry;
+}
+
+TSharedRef<SDockTab> FMaterialGraphEditorSummoner::SpawnTab(const FWorkflowTabSpawnInfo& Info) const
+{
+	TSharedPtr<FDocumentTracker> DocumentManager = MaterialEditorPtr.Pin()->DocumentManager;
+
+	// Call 'CreateTabBodyForObject' indirectly by making the material document manager open the tab, this preserves history while letting us manage the tab checkmark
+	return DocumentManager->GetActiveTab().IsValid()
+		? DocumentManager->GetActiveTab().ToSharedRef()
+		: DocumentManager->OpenDocument(Info.Payload, FDocumentTracker::SpawnManagedDocument).ToSharedRef();
+}
+
+TSharedRef<SDockTab> FMaterialGraphEditorSummoner::OnSpawnTab(const FSpawnTabArgs& SpawnArgs, TWeakPtr<FTabManager> WeakTabManager) const
+{
+	FWorkflowTabSpawnInfo SpawnInfo;
+	SpawnInfo.TabManager = WeakTabManager.Pin();
+	SpawnInfo.Payload = FTabPayload_UObject::Make(MaterialEditorPtr.Pin()->Material->MaterialGraph);
+
+	return SpawnInfo.TabManager.IsValid() ? SpawnTab(SpawnInfo) : SNew(SDockTab);
+}
+
+TSharedRef<SWidget> FMaterialGraphEditorSummoner::CreateTabBodyForObject(const FWorkflowTabSpawnInfo& Info, UEdGraph* DocumentID) const
+{
+	check(Info.TabInfo.IsValid());
+	return OnCreateGraphEditorWidget.Execute(Info.TabInfo.ToSharedRef(), DocumentID);
+}
+
+const FSlateBrush* FMaterialGraphEditorSummoner::GetTabIconForObject(const FWorkflowTabSpawnInfo& Info, UEdGraph* DocumentID) const
+{
+	return FEditorStyle::GetBrush(TEXT("GraphEditor.EventGraph_16x"));
+}
+
+TSharedRef<FGenericTabHistory> FMaterialGraphEditorSummoner::CreateTabHistoryNode(TSharedPtr<FTabPayload> Payload)
+{
+	return MakeShareable(new FMaterialGraphTabHistory(SharedThis(this), Payload));
+}
+
+/////////////////////////////////////////////////////
+// FMaterialGraphTabHistory
+
+void FMaterialGraphTabHistory::EvokeHistory(TSharedPtr<FTabInfo> InTabInfo, bool bPrevTabMatches)
+{
+	FWorkflowTabSpawnInfo SpawnInfo;
+	SpawnInfo.Payload = Payload;
+	SpawnInfo.TabInfo = InTabInfo;
+
+	if(bPrevTabMatches)
+	{
+		TSharedPtr<SDockTab> DockTab = InTabInfo->GetTab().Pin();
+		GraphEditor = StaticCastSharedRef<SGraphEditor>(DockTab->GetContent());
+	}
+	else
+	{
+		TSharedRef< SGraphEditor > GraphEditorRef = StaticCastSharedRef< SGraphEditor >(FactoryPtr.Pin()->CreateTabBody(SpawnInfo));
+		GraphEditor = GraphEditorRef;
+		FactoryPtr.Pin()->UpdateTab(InTabInfo->GetTab().Pin(), SpawnInfo, GraphEditorRef);
+	}
+}
+
+void FMaterialGraphTabHistory::SaveHistory()
+{
+	if (IsHistoryValid())
+	{
+		check(GraphEditor.IsValid());
+		GraphEditor.Pin()->GetViewLocation(SavedLocation, SavedZoomAmount);
+		GraphEditor.Pin()->GetViewBookmark(SavedBookmarkId);
+	}
+}
+
+void FMaterialGraphTabHistory::RestoreHistory()
+{
+	if (IsHistoryValid())
+	{
+		check(GraphEditor.IsValid());
+		GraphEditor.Pin()->SetViewLocation(SavedLocation, SavedZoomAmount, SavedBookmarkId);
+	}
+}
+
+#undef LOCTEXT_NAMESPACE
\ No newline at end of file
diff --git a/Engine/Source/Editor/MaterialEditor/Private/Tabs/MaterialEditorTabFactories.h b/Engine/Source/Editor/MaterialEditor/Private/Tabs/MaterialEditorTabFactories.h
new file mode 100644
--- /dev/null
+++ b/Engine/Source/Editor/MaterialEditor/Private/Tabs/MaterialEditorTabFactories.h
@@ -0,0 +1,124 @@
+// Copyright Epic Games, Inc. All Rights Reserved.
+
+#pragma once
+
+#include "CoreMinimal.h"
+#include "Misc/Attribute.h"
+#include "Widgets/SWidget.h"
+#include "Textures/SlateIcon.h"
+
+#include "GraphEditor.h"
+#include "EdGraph/EdGraph.h"
+#include "EdGraph/EdGraphSchema.h"
+
+#include "WorkflowOrientedApp/WorkflowTabFactory.h"
+#include "WorkflowOrientedApp/WorkflowTabManager.h"
+#include "WorkflowOrientedApp/WorkflowUObjectDocuments.h"
+
+#define LOCTEXT_NAMESPACE "MaterialEditorTabFactories"
+
+class FAssetEditorToolkit;
+class SDockTab;
+
+/////////////////////////////////////////////////////
+// FMaterialGraphEditorSummoner
+
+struct FMaterialGraphEditorSummoner : public FDocumentTabFactoryForObjects<UEdGraph>
+{
+public:
+	DECLARE_DELEGATE_RetVal_TwoParams(TSharedRef<SGraphEditor>, FOnCreateGraphEditorWidget, TSharedRef<FTabInfo>, UEdGraph*);
+public:
+	FMaterialGraphEditorSummoner(TSharedPtr<class FMaterialEditor> InMaterialEditorPtr, FOnCreateGraphEditorWidget CreateGraphEditorWidgetCallback);
+
+	virtual void OnTabActivated(TSharedPtr<SDockTab> Tab) const override;
+
+	virtual void OnTabBackgrounded(TSharedPtr<SDockTab> Tab) const override;
+
+	virtual void OnTabRefreshed(TSharedPtr<SDockTab> Tab) const override;
+
+	virtual void SaveState(TSharedPtr<SDockTab> Tab, TSharedPtr<FTabPayload> Payload) const override;
+
+	virtual FTabSpawnerEntry& RegisterTabSpawner(TSharedRef<FTabManager> TabManager, const FApplicationMode* CurrentApplicationMode) const override;
+
+protected:
+	virtual TAttribute<FText> ConstructTabNameForObject(UEdGraph* DocumentID) const override
+	{
+		auto GetGraphDisplayName = [](UEdGraph* Graph)
+		{
+			if (Graph)
+			{
+				if (const UEdGraphSchema* Schema = Graph->GetSchema())
+				{
+					FGraphDisplayInfo Info;
+					Schema->GetGraphDisplayInformation(*Graph, /*out*/ Info);
+
+					return Info.DisplayName;
+				}
+				else
+				{
+					// if we don't have a schema, we're dealing with a malformed (or incomplete graph)...
+					// possibly in the midst of some transaction - here we return the object's outer path 
+					// so we can at least get some context as to which graph we're referring
+					return FText::FromString(Graph->GetPathName());
+				}
+			}
+
+			return LOCTEXT("UnknownGraphName", "UNKNOWN");
+		};
+
+		return TAttribute<FText>::Create(TAttribute<FText>::FGetter::CreateLambda(GetGraphDisplayName, DocumentID));
+	}
+
+	virtual TSharedRef<SDockTab> SpawnTab(const FWorkflowTabSpawnInfo& Info) const override;
+
+	virtual TSharedRef<SDockTab> OnSpawnTab(const FSpawnTabArgs& SpawnArgs, TWeakPtr<FTabManager> WeakTabManager) const override;
+
+	virtual TSharedRef<SWidget> CreateTabBodyForObject(const FWorkflowTabSpawnInfo& Info, UEdGraph* DocumentID) const override;
+
+	virtual const FSlateBrush* GetTabIconForObject(const FWorkflowTabSpawnInfo& Info, UEdGraph* DocumentID) const override;
+
+	virtual TSharedRef<FGenericTabHistory> CreateTabHistoryNode(TSharedPtr<FTabPayload> Payload) override;
+
+protected:
+	TWeakPtr<class FMaterialEditor> MaterialEditorPtr;
+	FOnCreateGraphEditorWidget OnCreateGraphEditorWidget;
+};
+
+
+/////////////////////////////////////////////////////
+// FMaterialGraphTabHistory
+
+struct FMaterialGraphTabHistory : public FGenericTabHistory
+{
+public:
+	/**
+	 * @param InFactory		The factory used to regenerate the content
+	 * @param InPayload		The payload object used to regenerate the content
+	 */
+	FMaterialGraphTabHistory(TSharedPtr<FDocumentTabFactory> InFactory, TSharedPtr<FTabPayload> InPayload)
+		: FGenericTabHistory(InFactory, InPayload)
+		, SavedLocation(FVector2D::ZeroVector)
+		, SavedZoomAmount(INDEX_NONE)
+	{
+
+	}
+
+	virtual void EvokeHistory(TSharedPtr<FTabInfo> InTabInfo, bool bPrevTabMatches) override;
+
+	virtual void SaveHistory() override;
+
+	virtual void RestoreHistory() override;
+
+private:
+	/** The graph editor represented by this history node. While this node is inactive, the graph editor is invalid */
+	TWeakPtr< class SGraphEditor > GraphEditor;
+	/** Saved location the graph editor was at when this history node was last visited */
+	FVector2D SavedLocation;
+	/** Saved zoom the graph editor was at when this history node was last visited */
+	float SavedZoomAmount;
+	/** Saved bookmark ID the graph editor was at when this history node was last visited */
+	FGuid SavedBookmarkId;
+};
+
+
+#undef LOCTEXT_NAMESPACE
diff --git a/Engine/Source/Editor/MaterialEditor/Public/FindInMaterial.h b/Engine/Source/Editor/MaterialEditor/Public/FindInMaterial.h
--- a/Engine/Source/Editor/MaterialEditor/Public/FindInMaterial.h
+++ b/Engine/Source/Editor/MaterialEditor/Public/FindInMaterial.h
@@ -3,6 +3,7 @@
 #pragma once
 
 #include "CoreMinimal.h"
+#include "EdGraph/EdGraph.h"
 #include "EdGraph/EdGraphPin.h"
 #include "Input/Reply.h"
 #include "Widgets/SWidget.h"
@@ -110,6 +111,9 @@ protected:
 	/** Find any results that contain all of the tokens */
 	void MatchTokens(const TArray<FString>& Tokens);
 
+	/** Find any results that contain all of the tokens in provided graph and subgraphs */
+	void MatchTokensInGraph(const UEdGraph* Graph, const TArray<FString>& Tokens);
+
 	/** Determines if a string matches the search tokens */
 	static bool StringMatchesSearchTokens(const TArray<FString>& Tokens, const FString& ComparisonString);
 
diff --git a/Engine/Source/Editor/MaterialEditor/Public/IMaterialEditor.h b/Engine/Source/Editor/MaterialEditor/Public/IMaterialEditor.h
--- a/Engine/Source/Editor/MaterialEditor/Public/IMaterialEditor.h
+++ b/Engine/Source/Editor/MaterialEditor/Public/IMaterialEditor.h
@@ -4,6 +4,7 @@
 
 #include "CoreMinimal.h"
 #include "Toolkits/AssetEditorToolkit.h"
+#include "WorkflowOrientedApp/WorkflowCentricApplication.h"
 
 class FCanvas;
 class FMaterialRenderProxy;
@@ -11,12 +12,13 @@ class FViewport;
 class IMaterialEditorModule;
 class UMaterialExpression;
 class UMaterialExpressionComment;
+class UMaterialExpressionComposite;
 class UMaterialInterface;
 
 /**
  * Public interface to Material Editor
  */
-class IMaterialEditor : public FAssetEditorToolkit, public IHasMenuExtensibility, public IHasToolBarExtensibility
+class IMaterialEditor : public FWorkflowCentricApplication, public IHasMenuExtensibility, public IHasToolBarExtensibility
 {
 public:
 	/**
@@ -26,10 +28,21 @@ public:
 	 * @param	NodePos					Position of the new node.
 	 * @param	bAutoSelect				If true, deselect all expressions and select the newly created one.
 	 * @param	bAutoAssignResource		If true, assign resources to new expression.
+	 * @param	Graph            		Graph to create new expression within. 
 	 *
 	 * @return	UMaterialExpression*	Newly created material expression
 	 */
-	virtual UMaterialExpression* CreateNewMaterialExpression(UClass* NewExpressionClass, const FVector2D& NodePos, bool bAutoSelect, bool bAutoAssignResource) {return NULL;}
+	virtual UMaterialExpression* CreateNewMaterialExpression(UClass* NewExpressionClass, const FVector2D& NodePos, bool bAutoSelect, bool bAutoAssignResource, const class UEdGraph* Graph = nullptr) {return nullptr;}
+
+	/**
+	 * Creates a new material expression composite, along with related child pinbase expressions
+	 *
+	 * @param	NodePos							Position of the new composite
+	 * @param	Graph							Graph to create new expression within.
+	 *
+	 * @return	UMaterialExpressionComposite*	Newly created material expression composite
+	 */
+	virtual UMaterialExpressionComposite* CreateNewMaterialExpressionComposite(const FVector2D& NodePos, const class UEdGraph* Graph = nullptr) { return nullptr; }
 
 	/**
 	 * Creates a new material expression comment
@@ -37,8 +50,9 @@ public:
 	 * @param	NodePos						Position of the new comment
 	 *
 	 * @return	UMaterialExpressionComment*	Newly created material expression comment
+	 * @param	Graph            			Graph to create new expression within.
 	 */
-	virtual UMaterialExpressionComment* CreateNewMaterialExpressionComment(const FVector2D& NodePos) {return NULL;}
+	virtual UMaterialExpressionComment* CreateNewMaterialExpressionComment(const FVector2D& NodePos, const class UEdGraph* Graph = nullptr) {return nullptr;}
 
 	/**
 	 * Refreshes all material expression previews, regardless of whether or not realtime previews are enabled.
@@ -76,11 +90,18 @@ public:
 	 */
 	virtual void UpdateMaterialAfterGraphChange() {};
 
+	/**
+	 * Jumps to a specific node or graph
+	 *
+	 * @param	ObjectReference			Object to jump to
+	 */
+	virtual void JumpToHyperlink(const UObject* ObjectReference) {};
+
 	/** Checks whether nodes can currently be pasted */
 	virtual bool CanPasteNodes() const {return false;}
 
 	/** Paste nodes at a specific location */
-	virtual void PasteNodesHere(const FVector2D& Location) {};
+	virtual void PasteNodesHere(const FVector2D& Location, const class UEdGraph* Graph = nullptr) {};
 
 	/** Gets the number of selected nodes */
 	virtual int32 GetNumberOfSelectedNodes() const {return 0;}
diff --git a/Engine/Source/Editor/MaterialEditor/Public/MaterialEditingLibrary.h b/Engine/Source/Editor/MaterialEditor/Public/MaterialEditingLibrary.h
--- a/Engine/Source/Editor/MaterialEditor/Public/MaterialEditingLibrary.h
+++ b/Engine/Source/Editor/MaterialEditor/Public/MaterialEditingLibrary.h
@@ -70,6 +70,16 @@ public:
 	UFUNCTION(BlueprintCallable, Category = "MaterialEditing")
 	static UMaterialExpression* CreateMaterialExpression(UMaterial* Material, TSubclassOf<UMaterialExpression> ExpressionClass, int32 NodePosX=0, int32 NodePosY=0);
 
+	/** 
+	 *	Duplicates the provided material expression adding it to the same material / material function, and copying parameters.
+	 *  Note: Does not duplicate transient properties (Ex: GraphNode).
+	 *
+	 *	@param	Material			Material asset to add an expression to
+	 *	@param	MaterialFunction	Specified if adding an expression to a MaterialFunction, used as Outer for new expression object
+	 *	@param	SourceExpression	Expression to be duplicated
+	 */
+	UFUNCTION(BlueprintCallable, Category = "MaterialEditing")
+	static UMaterialExpression* DuplicateMaterialExpression(UMaterial* Material, UMaterialFunction* MaterialFunction, UMaterialExpression* Expression);
 
 	/** 
 	 *	Enable a particular usage for the supplied material (e.g. SkeletalMesh, ParticleSprite etc)
diff --git a/Engine/Source/Editor/MaterialEditor/Public/MaterialEditorModes.h b/Engine/Source/Editor/MaterialEditor/Public/MaterialEditorModes.h
new file mode 100644
--- /dev/null
+++ b/Engine/Source/Editor/MaterialEditor/Public/MaterialEditorModes.h
@@ -0,0 +1,47 @@
+// Copyright Epic Games, Inc. All Rights Reserved.
+
+#pragma once
+
+#include "CoreMinimal.h"
+#include "WorkflowOrientedApp/WorkflowTabManager.h"
+#include "MaterialEditor.h"
+#include "WorkflowOrientedApp/ApplicationMode.h"
+
+struct MATERIALEDITOR_API FMaterialEditorApplicationModes
+{
+	// Mode identifiers
+	static const FName StandardMaterialEditorMode;
+	static FText GetLocalizedMode(const FName InMode)
+	{
+		static TMap< FName, FText > LocModes;
+
+		if (LocModes.Num() == 0)
+		{
+			LocModes.Add(StandardMaterialEditorMode, NSLOCTEXT("MaterialEditor", "StandardMaterialEditorMode", "Graph"));
+		}
+
+		check(InMode != NAME_None);
+		const FText* OutDesc = LocModes.Find(InMode);
+		check(OutDesc);
+		return *OutDesc;
+	}
+	static TSharedPtr<FTabManager::FLayout> GetDefaultEditorLayout(TSharedPtr<class FMaterialEditor> InMaterialEditor);
+private:
+	FMaterialEditorApplicationModes() {}
+};
+
+// Even though we currently only have one mode of operation, we still need a application mode to handle factory tab spawning
+class MATERIALEDITOR_API FMaterialEditorApplicationMode : public FApplicationMode
+{
+public:
+	FMaterialEditorApplicationMode(TSharedPtr<class FMaterialEditor> InMaterialEditor);
+
+	virtual void RegisterTabFactories(TSharedPtr<FTabManager> InTabManager) override;
+public:
+
+protected:
+	TWeakPtr<FMaterialEditor> MyMaterialEditor;
+
+	// Set of spawnable tabs handled by workflow factories
+	FWorkflowAllowedTabSet MaterialEditorTabFactories;
+};
\ No newline at end of file
diff --git a/Engine/Source/Editor/MaterialEditor/Public/MaterialEditorUtilities.h b/Engine/Source/Editor/MaterialEditor/Public/MaterialEditorUtilities.h
--- a/Engine/Source/Editor/MaterialEditor/Public/MaterialEditorUtilities.h
+++ b/Engine/Source/Editor/MaterialEditor/Public/MaterialEditorUtilities.h
@@ -11,6 +11,7 @@
 class IMaterialEditor;
 class UMaterial;
 class UMaterialExpressionComment;
+class UMaterialExpressionComposite;
 class UMaterialExpressionFunctionInput;
 class UMaterialInstance;
 class UMaterialInterface;
@@ -53,6 +54,16 @@ public:
 	 */
 	static UMaterialExpression* CreateNewMaterialExpression(const class UEdGraph* Graph, UClass* NewExpressionClass, const FVector2D& NodePos, bool bAutoSelect, bool bAutoAssignResource);
 
+	/**
+	 * Creates a new material expression composite on the material represented by a graph.
+	 *
+	 * @param	Graph	Graph to add comment to
+	 * @param	NodePos	Position to place new comment at
+	 *
+	 * @return	UMaterialExpressionComment*	Newly created comment
+	 */
+	static UMaterialExpressionComposite* CreateNewMaterialExpressionComposite(const class UEdGraph* Graph, const FVector2D& NodePos);
+
 	/**
 	 * Creates a new material expression comment on the material represented by a graph.
 	 *
@@ -194,6 +205,9 @@ public:
 	/** Get IMaterialEditor for given object, if it exists */
 	static TSharedPtr<class IMaterialEditor> GetIMaterialEditorForObject(const UObject* ObjectToFocusOn);
 
+	/** Get IMaterialEditor to focus on given object, if it exists */
+	static void BringFocusAttentionOnObject(const UObject* ObjectToFocusOn);
+
 
 	/** Commands for the Parents menu */
 	static void OnOpenMaterial(const FAssetData InMaterial);
diff --git a/Engine/Source/Editor/UnrealEd/Classes/MaterialGraph/MaterialGraph.h b/Engine/Source/Editor/UnrealEd/Classes/MaterialGraph/MaterialGraph.h
--- a/Engine/Source/Editor/UnrealEd/Classes/MaterialGraph/MaterialGraph.h
+++ b/Engine/Source/Editor/UnrealEd/Classes/MaterialGraph/MaterialGraph.h
@@ -9,6 +9,7 @@
 #include "MaterialGraph.generated.h"
 
 class UMaterialExpressionComment;
+class UMaterialExpressionComposite;
 
 DECLARE_DELEGATE_RetVal( bool, FRealtimeStateGetter );
 DECLARE_DELEGATE( FSetMaterialDirty );
@@ -114,6 +115,10 @@ class UNREALED_API UMaterialGraph : public UEdGraph
 	UPROPERTY()
 	class UMaterialGraphNode_Root*	RootNode;
 
+	/** Expression this subgraph represents (NULL if not subgraph, Material [Function] still populated) */
+	UPROPERTY()
+	UMaterialExpression*			SubgraphExpression;
+
 	/** List of Material Inputs (not set up for Material Functions) */
 	TArray<FMaterialInputInfo> MaterialInputs;
 
@@ -155,6 +160,15 @@ public:
 	 */
 	class UMaterialGraphNode_Comment*	AddComment(UMaterialExpressionComment* Comment, bool bIsUserInvoked = false);
 
+	/**
+	 * Add a Subgraph to the Graph
+	 *
+	 * @param	InSubgraphExpression, expression that will represent the new subgraph in this graph.
+	 *
+	 * @return	UMaterialGraph* Newly created subgraph
+	 */
+	UMaterialGraph* AddSubGraph(UMaterialExpression* InSubgraphExpression);
+
 	/** Link all of the Graph nodes using the Material's connections */
 	void LinkGraphNodesFromMaterial();
 
@@ -181,6 +195,8 @@ private:
 	 */
 	void RemoveAllNodes();
 
+	void RebuildGraphInternal(const TMap<UMaterialExpression*, TArray<UMaterialExpression*>>& SubgraphExpressionMap, const TMap<UMaterialExpression*, TArray<UMaterialExpressionComment*>>& SubgraphCommentMap);
+
 	/**
 	 * Gets a valid output index, matching mask values if necessary
 	 *
diff --git a/Engine/Source/Editor/UnrealEd/Classes/MaterialGraph/MaterialGraphNode.h b/Engine/Source/Editor/UnrealEd/Classes/MaterialGraph/MaterialGraphNode.h
--- a/Engine/Source/Editor/UnrealEd/Classes/MaterialGraph/MaterialGraphNode.h
+++ b/Engine/Source/Editor/UnrealEd/Classes/MaterialGraph/MaterialGraphNode.h
@@ -39,7 +39,7 @@ class UMaterialGraphNode : public UMaterialGraphNode_Base
 
 public:
 	/** Fix up the node's owner after being copied */
-	UNREALED_API void PostCopyNode();
+	UNREALED_API virtual void PostCopyNode();
 
 	/** Get the expression preview for this node */
 	UNREALED_API FMaterialRenderProxy* GetExpressionPreview();
diff --git a/Engine/Source/Editor/UnrealEd/Classes/MaterialGraph/MaterialGraphNode_Composite.h b/Engine/Source/Editor/UnrealEd/Classes/MaterialGraph/MaterialGraphNode_Composite.h
new file mode 100644
--- /dev/null
+++ b/Engine/Source/Editor/UnrealEd/Classes/MaterialGraph/MaterialGraphNode_Composite.h
@@ -0,0 +1,43 @@
+// Copyright Epic Games, Inc. All Rights Reserved.
+
+#pragma once
+
+#include "CoreMinimal.h"
+#include "UObject/ObjectMacros.h"
+#include "MaterialGraph/MaterialGraph.h"
+#include "MaterialGraph/MaterialGraphNode.h"
+
+#include "MaterialGraphNode_Composite.generated.h"
+
+UCLASS(MinimalAPI)
+class UMaterialGraphNode_Composite : public UMaterialGraphNode
+{
+	GENERATED_UCLASS_BODY()
+
+	/** The graph that this composite node is representing */
+	UPROPERTY()
+	UMaterialGraph* BoundGraph;
+
+	//~ Begin UObject Interface
+	virtual void PostEditUndo() override; 
+	//~ End UObject Interface
+
+	//~ Begin UMaterialGraphNode Interface
+	UNREALED_API virtual void PostCopyNode() override;
+	//~ End UMaterialGraphNode Interface
+
+	//~ Begin UEdGraphNode Interface.
+	virtual UObject* GetJumpTargetForDoubleClick() const override;
+	virtual bool CanJumpToDefinition() const override;
+	virtual void JumpToDefinition() const override;
+	virtual void DestroyNode() override; 
+	virtual void PrepareForCopying() override;
+	virtual void PostPasteNode() override; 
+	virtual void OnRenameNode(const FString& NewName) override; 
+	//~ End UEdGraphNode Interface.
+
+private:
+
+	/** Fixes up the input and output pin bases when needed, useful after PostEditUndo which changes which graph these nodes point to */
+	void FixupInputAndOutputPinBases(); 
+};
diff --git a/Engine/Source/Editor/UnrealEd/Classes/MaterialGraph/MaterialGraphNode_PinBase.h b/Engine/Source/Editor/UnrealEd/Classes/MaterialGraph/MaterialGraphNode_PinBase.h
new file mode 100644
--- /dev/null
+++ b/Engine/Source/Editor/UnrealEd/Classes/MaterialGraph/MaterialGraphNode_PinBase.h
@@ -0,0 +1,25 @@
+// Copyright Epic Games, Inc. All Rights Reserved.
+
+#pragma once
+
+#include "CoreMinimal.h"
+#include "UObject/ObjectMacros.h"
+#include "MaterialGraph/MaterialGraph.h"
+#include "MaterialGraph/MaterialGraphNode.h"
+
+#include "MaterialGraphNode_PinBase.generated.h"
+
+UCLASS(MinimalAPI)
+class UMaterialGraphNode_PinBase : public UMaterialGraphNode
+{
+	GENERATED_BODY()
+
+	//~ Begin UMaterialGraphNode Interface
+	UNREALED_API virtual void PostCopyNode() override;
+	//~ End UMaterialGraphNode Interface
+
+	//~ Begin UEdGraphNode Interface.
+	virtual void PrepareForCopying() override;
+	virtual bool CanDuplicateNode() const override { return false; }
+	//~ End UEdGraphNode Interface.
+};
diff --git a/Engine/Source/Editor/UnrealEd/Classes/MaterialGraph/MaterialGraphSchema.h b/Engine/Source/Editor/UnrealEd/Classes/MaterialGraph/MaterialGraphSchema.h
--- a/Engine/Source/Editor/UnrealEd/Classes/MaterialGraph/MaterialGraphSchema.h
+++ b/Engine/Source/Editor/UnrealEd/Classes/MaterialGraph/MaterialGraphSchema.h
@@ -74,6 +74,31 @@ struct UNREALED_API FMaterialGraphSchemaAction_NewFunctionCall : public FEdGraph
 	//~ End FEdGraphSchemaAction Interface
 };
 
+/** Action to add a composite node to the graph */
+USTRUCT()
+struct UNREALED_API FMaterialGraphSchemaAction_NewComposite : public FEdGraphSchemaAction
+{
+	GENERATED_USTRUCT_BODY();
+
+	// Simple type info
+	static FName StaticGetTypeId() { static FName Type("FMaterialGraphSchemaAction_NewComposite"); return Type; }
+	virtual FName GetTypeId() const override { return StaticGetTypeId(); }
+
+	FMaterialGraphSchemaAction_NewComposite()
+		: FEdGraphSchemaAction()
+	{}
+
+	FMaterialGraphSchemaAction_NewComposite(FText InNodeCategory, FText InMenuDesc, FText InToolTip, const int32 InGrouping)
+		: FEdGraphSchemaAction(MoveTemp(InNodeCategory), MoveTemp(InMenuDesc), MoveTemp(InToolTip), InGrouping)
+	{}
+
+	//~ Begin FEdGraphSchemaAction Interface
+	virtual UEdGraphNode* PerformAction(class UEdGraph* ParentGraph, UEdGraphPin* FromPin, const FVector2D Location, bool bSelectNewNode = true) override;
+	//~ End FEdGraphSchemaAction Interface
+
+	static UEdGraphNode* SpawnNode(class UEdGraph* ParentGraph, const FVector2D Location);
+};
+
 /** Action to add a comment node to the graph */
 USTRUCT()
 struct UNREALED_API FMaterialGraphSchemaAction_NewComment : public FEdGraphSchemaAction
@@ -198,6 +223,7 @@ class UMaterialGraphSchema : public UEdGraphSchema
 	virtual void BreakNodeLinks(UEdGraphNode& TargetNode) const override;
 	virtual void BreakPinLinks(UEdGraphPin& TargetPin, bool bSendsNodeNotifcation) const override;
 	virtual void BreakSinglePinLink(UEdGraphPin* SourcePin, UEdGraphPin* TargetPin) const override;
+	virtual bool CanEncapuslateNode(UEdGraphNode const& TestNode) const override;
 	virtual void DroppedAssetsOnGraph(const TArray<struct FAssetData>& Assets, const FVector2D& GraphPosition, UEdGraph* Graph) const override;
 	virtual int32 GetNodeSelectionCount(const UEdGraph* Graph) const override;
 	virtual TSharedPtr<FEdGraphSchemaAction> GetCreateCommentAction() const override;
@@ -212,6 +238,8 @@ class UMaterialGraphSchema : public UEdGraphSchema
 private:
 	/** Adds actions for all Material Functions */
 	void GetMaterialFunctionActions(FGraphActionMenuBuilder& ActionMenuBuilder) const;
+	/** Adds action for creating a composite */
+	void GetCompositeAction(FGraphActionMenuBuilder& ActionMenuBuilder, const UEdGraph* CurrentGraph = NULL) const;
 	/** Adds action for creating a comment */
 	void GetCommentAction(FGraphActionMenuBuilder& ActionMenuBuilder, const UEdGraph* CurrentGraph = NULL) const;
 	/**
diff --git a/Engine/Source/Editor/UnrealEd/Private/Kismet2/BlueprintEditorUtils.cpp b/Engine/Source/Editor/UnrealEd/Private/Kismet2/BlueprintEditorUtils.cpp
--- a/Engine/Source/Editor/UnrealEd/Private/Kismet2/BlueprintEditorUtils.cpp
+++ b/Engine/Source/Editor/UnrealEd/Private/Kismet2/BlueprintEditorUtils.cpp
@@ -79,6 +79,7 @@
 #include "K2Node_TemporaryVariable.h"
 #include "K2Node_Timeline.h"
 #include "K2Node_Knot.h"
+#include "MaterialGraph/MaterialGraphNode_Composite.h"
 #include "AnimGraphNode_StateMachineBase.h"
 #include "AnimStateNodeBase.h"
 #include "AnimStateNode.h"
@@ -2581,6 +2582,7 @@ void FBlueprintEditorUtils::RemoveGraph(UBlueprint* Blueprint, class UEdGraph* G
 			OuterGraph->SubGraphs.Remove(GraphToRemove);
 		}
 		else if (! (Cast<UK2Node_Composite>(TestOuter)	|| 
+					Cast<UMaterialGraphNode_Composite>(TestOuter)	||
 					Cast<UAnimStateNodeBase>(TestOuter)	||
 					Cast<UAnimStateTransitionNode>(TestOuter)	||
 					Cast<UAnimGraphNode_StateMachineBase>(TestOuter)) )
@@ -2610,6 +2612,19 @@ void FBlueprintEditorUtils::RemoveGraph(UBlueprint* Blueprint, class UEdGraph* G
 		}
 	}
 
+	// Materials have their own form of composite, and we haven't made a generalized composite EdGraphNode yet.
+	TArray<UMaterialGraphNode_Composite*> AllMaterialCompositeNodes;
+	GraphToRemove->GetNodesOfClass<UMaterialGraphNode_Composite>(AllMaterialCompositeNodes);
+
+	for (UMaterialGraphNode_Composite* CompNode : AllMaterialCompositeNodes)
+	{
+		UEdGraph* BoundGraph = CompNode->BoundGraph;
+		if (BoundGraph && BoundGraph->GetOuter()->IsA(UMaterialGraphNode_Composite::StaticClass()))
+		{
+			FBlueprintEditorUtils::RemoveGraph(Blueprint, CompNode->BoundGraph, EGraphRemoveFlags::None);
+		}
+	}
+
 	// Animation nodes can contain subgraphs but are not composite nodes, handle their graphs
 	TArray<UAnimStateNodeBase*> AllAnimCompositeNodes;
 	GraphToRemove->GetNodesOfClassEx<UAnimStateNode>(AllAnimCompositeNodes);
@@ -2640,7 +2655,7 @@ void FBlueprintEditorUtils::RemoveGraph(UBlueprint* Blueprint, class UEdGraph* G
 
 	GraphToRemove->GetSchema()->HandleGraphBeingDeleted(*GraphToRemove);
 
-	GraphToRemove->Rename(nullptr, Blueprint->GetOuter(), REN_DoNotDirty | REN_DontCreateRedirectors);
+	GraphToRemove->Rename(nullptr, Blueprint ? Blueprint->GetOuter() : nullptr, REN_DoNotDirty | REN_DontCreateRedirectors);
 	GraphToRemove->ClearFlags(RF_Standalone | RF_Public);
 	GraphToRemove->RemoveFromRoot();
 
diff --git a/Engine/Source/Editor/UnrealEd/Private/MaterialGraph.cpp b/Engine/Source/Editor/UnrealEd/Private/MaterialGraph.cpp
--- a/Engine/Source/Editor/UnrealEd/Private/MaterialGraph.cpp
+++ b/Engine/Source/Editor/UnrealEd/Private/MaterialGraph.cpp
@@ -5,16 +5,22 @@
 
 #include "MaterialGraph/MaterialGraph.h"
 #include "MaterialGraph/MaterialGraphNode_Comment.h"
+#include "MaterialGraph/MaterialGraphNode_Composite.h"
+#include "MaterialGraph/MaterialGraphNode_PinBase.h"
 #include "MaterialGraph/MaterialGraphNode.h"
 #include "MaterialGraph/MaterialGraphNode_Root.h"
 
 #include "Materials/MaterialExpressionComment.h"
+#include "Materials/MaterialExpressionComposite.h"
+#include "Materials/MaterialExpressionPinBase.h"
 #include "Materials/MaterialExpressionFunctionOutput.h"
 #include "Materials/MaterialExpressionCustomOutput.h"
 #include "Materials/MaterialExpressionReroute.h"
 
 #include "MaterialGraphNode_Knot.h"
 
+#include "Kismet2/BlueprintEditorUtils.h"
+
 #define LOCTEXT_NAMESPACE "MaterialGraph"
 
 UMaterialGraph::UMaterialGraph(const FObjectInitializer& ObjectInitializer)
@@ -26,46 +32,63 @@ void UMaterialGraph::RebuildGraph()
 {
 	check(Material);
 
+	// Pre-group expressions & comments per subgraph to avoid unnecessary iteration over all material expressions
+	TMap<UMaterialExpression*, TArray<UMaterialExpression*>> SubgraphExpressionMap;
+	TMap<UMaterialExpression*, TArray<UMaterialExpressionComment*>> SubgraphCommentMap;
+	for (UMaterialExpression* Expression : Material->Expressions)
+	{
+		SubgraphExpressionMap.FindOrAdd(Expression->SubgraphExpression).Add(Expression);
+	}
+	for (UMaterialExpressionComment* Comment : Material->EditorComments)
+	{
+		SubgraphCommentMap.FindOrAdd(Comment->SubgraphExpression).Add(Comment);
+	}
+
+	RebuildGraphInternal(SubgraphExpressionMap, SubgraphCommentMap);
+}
+
+void UMaterialGraph::RebuildGraphInternal(const TMap<UMaterialExpression*, TArray<UMaterialExpression*>>& SubgraphExpressionMap, const TMap<UMaterialExpression*, TArray<UMaterialExpressionComment*>>& SubgraphCommentMap)
+{
 	Modify();
 
 	RemoveAllNodes();
 
-	if (!MaterialFunction)
+	if (!MaterialFunction && !SubgraphExpression)
 	{
 		// This needs to be done before building the new material inputs to guarantee that the shading model field is up to date
 		Material->RebuildShadingModelField();
 
 		// Initialize the material input list.
-		MaterialInputs.Add( FMaterialInputInfo(FMaterialAttributeDefinitionMap::GetDisplayNameForMaterial(MP_BaseColor, Material), MP_BaseColor, LOCTEXT( "BaseColorToolTip", "Defines the overall color of the Material. Each channel is automatically clamped between 0 and 1" ) ) );
-		MaterialInputs.Add( FMaterialInputInfo(FMaterialAttributeDefinitionMap::GetDisplayNameForMaterial(MP_Metallic, Material), MP_Metallic, LOCTEXT( "MetallicToolTip", "Controls how \"metal-like\" your surface looks like") ) );
-		MaterialInputs.Add( FMaterialInputInfo(FMaterialAttributeDefinitionMap::GetDisplayNameForMaterial(MP_Specular, Material), MP_Specular, LOCTEXT("SpecularToolTip", "Used to scale the current amount of specularity on non-metallic surfaces and is a value between 0 and 1, default at 0.5") ) );
-		MaterialInputs.Add( FMaterialInputInfo(FMaterialAttributeDefinitionMap::GetDisplayNameForMaterial(MP_Roughness, Material), MP_Roughness, LOCTEXT("RoughnessToolTip", "Controls how rough the Material is. Roughness of 0 (smooth) is a mirror reflection and 1 (rough) is completely matte or diffuse") ) );
-		MaterialInputs.Add( FMaterialInputInfo(FMaterialAttributeDefinitionMap::GetDisplayNameForMaterial(MP_Anisotropy, Material), MP_Anisotropy, LOCTEXT("AnisotropyToolTip", "Determines the extent the specular highlight is stretched along the tangent. Anisotropy from 0 to 1 results in a specular highlight that stretches from uniform to maximally stretched along the tangent direction.")));
-		MaterialInputs.Add( FMaterialInputInfo(FMaterialAttributeDefinitionMap::GetDisplayNameForMaterial(MP_EmissiveColor, Material), MP_EmissiveColor, LOCTEXT( "EmissiveToolTip", "Controls which parts of your Material will appear to glow" ) ) );
-		MaterialInputs.Add( FMaterialInputInfo(FMaterialAttributeDefinitionMap::GetDisplayNameForMaterial(MP_Opacity, Material), MP_Opacity, LOCTEXT( "OpacityToolTip", "Controls the translucency of the Material" ) ) );
-		MaterialInputs.Add( FMaterialInputInfo(FMaterialAttributeDefinitionMap::GetDisplayNameForMaterial(MP_OpacityMask, Material), MP_OpacityMask, LOCTEXT( "OpacityMaskToolTip", "When in Masked mode, a Material is either completely visible or completely invisible" ) ) );
-		MaterialInputs.Add( FMaterialInputInfo(FMaterialAttributeDefinitionMap::GetDisplayNameForMaterial(MP_Normal, Material), MP_Normal, LOCTEXT( "NormalToolTip", "Takes the input of a normal map" ) ) );
-		MaterialInputs.Add( FMaterialInputInfo(FMaterialAttributeDefinitionMap::GetDisplayNameForMaterial(MP_Tangent, Material), MP_Tangent, LOCTEXT( "TangentToolTip", "Takes the input of a tangent map. Useful for specifying anisotropy direction." ) ) );
-		MaterialInputs.Add( FMaterialInputInfo(FMaterialAttributeDefinitionMap::GetDisplayNameForMaterial(MP_WorldPositionOffset, Material), MP_WorldPositionOffset, LOCTEXT( "WorldPositionOffsetToolTip", "Allows for the vertices of a mesh to be manipulated in world space by the Material" ) ) );
-		MaterialInputs.Add( FMaterialInputInfo(FMaterialAttributeDefinitionMap::GetDisplayNameForMaterial(MP_WorldDisplacement, Material), MP_WorldDisplacement, LOCTEXT( "WorldDisplacementToolTip", "Allows for the tessellation vertices to be manipulated in world space by the Material" ) ) );
-		MaterialInputs.Add( FMaterialInputInfo(FMaterialAttributeDefinitionMap::GetDisplayNameForMaterial(MP_TessellationMultiplier, Material), MP_TessellationMultiplier, LOCTEXT( "TessllationMultiplierToolTip", "Controls the amount tessellation along the surface" ) ) );
-		MaterialInputs.Add( FMaterialInputInfo(FMaterialAttributeDefinitionMap::GetDisplayNameForMaterial(MP_SubsurfaceColor, Material), MP_SubsurfaceColor, LOCTEXT( "SubsurfaceToolTip", "Allows you to add a color to your Material to simulate shifts in color when light passes through the surface" ) ) );
-		MaterialInputs.Add( FMaterialInputInfo(FMaterialAttributeDefinitionMap::GetDisplayNameForMaterial(MP_CustomData0, Material), MP_CustomData0, FMaterialAttributeDefinitionMap::GetDisplayNameForMaterial(MP_CustomData0, Material)));
-		MaterialInputs.Add( FMaterialInputInfo(FMaterialAttributeDefinitionMap::GetDisplayNameForMaterial(MP_CustomData1, Material), MP_CustomData1, FMaterialAttributeDefinitionMap::GetDisplayNameForMaterial(MP_CustomData1, Material)));
-		MaterialInputs.Add( FMaterialInputInfo(FMaterialAttributeDefinitionMap::GetDisplayNameForMaterial(MP_AmbientOcclusion, Material), MP_AmbientOcclusion, LOCTEXT( "AmbientOcclusionToolTip", "Simulate the self-shadowing that happens within crevices of a surface, or of a volume for volumetric clouds only" ) ) );
-		MaterialInputs.Add( FMaterialInputInfo(FMaterialAttributeDefinitionMap::GetDisplayNameForMaterial(MP_Refraction, Material), MP_Refraction, LOCTEXT( "RefractionToolTip", "Takes in a texture or value that simulates the index of refraction of the surface" ) ) );
+		MaterialInputs.Add(FMaterialInputInfo(FMaterialAttributeDefinitionMap::GetDisplayNameForMaterial(MP_BaseColor, Material), MP_BaseColor, LOCTEXT("BaseColorToolTip", "Defines the overall color of the Material. Each channel is automatically clamped between 0 and 1")));
+		MaterialInputs.Add(FMaterialInputInfo(FMaterialAttributeDefinitionMap::GetDisplayNameForMaterial(MP_Metallic, Material), MP_Metallic, LOCTEXT("MetallicToolTip", "Controls how \"metal-like\" your surface looks like")));
+		MaterialInputs.Add(FMaterialInputInfo(FMaterialAttributeDefinitionMap::GetDisplayNameForMaterial(MP_Specular, Material), MP_Specular, LOCTEXT("SpecularToolTip", "Used to scale the current amount of specularity on non-metallic surfaces and is a value between 0 and 1, default at 0.5")));
+		MaterialInputs.Add(FMaterialInputInfo(FMaterialAttributeDefinitionMap::GetDisplayNameForMaterial(MP_Roughness, Material), MP_Roughness, LOCTEXT("RoughnessToolTip", "Controls how rough the Material is. Roughness of 0 (smooth) is a mirror reflection and 1 (rough) is completely matte or diffuse")));
+		MaterialInputs.Add(FMaterialInputInfo(FMaterialAttributeDefinitionMap::GetDisplayNameForMaterial(MP_Anisotropy, Material), MP_Anisotropy, LOCTEXT("AnisotropyToolTip", "Determines the extent the specular highlight is stretched along the tangent. Anisotropy from 0 to 1 results in a specular highlight that stretches from uniform to maximally stretched along the tangent direction.")));
+		MaterialInputs.Add(FMaterialInputInfo(FMaterialAttributeDefinitionMap::GetDisplayNameForMaterial(MP_EmissiveColor, Material), MP_EmissiveColor, LOCTEXT("EmissiveToolTip", "Controls which parts of your Material will appear to glow")));
+		MaterialInputs.Add(FMaterialInputInfo(FMaterialAttributeDefinitionMap::GetDisplayNameForMaterial(MP_Opacity, Material), MP_Opacity, LOCTEXT("OpacityToolTip", "Controls the translucency of the Material")));
+		MaterialInputs.Add(FMaterialInputInfo(FMaterialAttributeDefinitionMap::GetDisplayNameForMaterial(MP_OpacityMask, Material), MP_OpacityMask, LOCTEXT("OpacityMaskToolTip", "When in Masked mode, a Material is either completely visible or completely invisible")));
+		MaterialInputs.Add(FMaterialInputInfo(FMaterialAttributeDefinitionMap::GetDisplayNameForMaterial(MP_Normal, Material), MP_Normal, LOCTEXT("NormalToolTip", "Takes the input of a normal map")));
+		MaterialInputs.Add(FMaterialInputInfo(FMaterialAttributeDefinitionMap::GetDisplayNameForMaterial(MP_Tangent, Material), MP_Tangent, LOCTEXT("TangentToolTip", "Takes the input of a tangent map. Useful for specifying anisotropy direction.")));
+		MaterialInputs.Add(FMaterialInputInfo(FMaterialAttributeDefinitionMap::GetDisplayNameForMaterial(MP_WorldPositionOffset, Material), MP_WorldPositionOffset, LOCTEXT("WorldPositionOffsetToolTip", "Allows for the vertices of a mesh to be manipulated in world space by the Material")));
+		MaterialInputs.Add(FMaterialInputInfo(FMaterialAttributeDefinitionMap::GetDisplayNameForMaterial(MP_WorldDisplacement, Material), MP_WorldDisplacement, LOCTEXT("WorldDisplacementToolTip", "Allows for the tessellation vertices to be manipulated in world space by the Material")));
+		MaterialInputs.Add(FMaterialInputInfo(FMaterialAttributeDefinitionMap::GetDisplayNameForMaterial(MP_TessellationMultiplier, Material), MP_TessellationMultiplier, LOCTEXT("TessllationMultiplierToolTip", "Controls the amount tessellation along the surface")));
+		MaterialInputs.Add(FMaterialInputInfo(FMaterialAttributeDefinitionMap::GetDisplayNameForMaterial(MP_SubsurfaceColor, Material), MP_SubsurfaceColor, LOCTEXT("SubsurfaceToolTip", "Allows you to add a color to your Material to simulate shifts in color when light passes through the surface")));
+		MaterialInputs.Add(FMaterialInputInfo(FMaterialAttributeDefinitionMap::GetDisplayNameForMaterial(MP_CustomData0, Material), MP_CustomData0, FMaterialAttributeDefinitionMap::GetDisplayNameForMaterial(MP_CustomData0, Material)));
+		MaterialInputs.Add(FMaterialInputInfo(FMaterialAttributeDefinitionMap::GetDisplayNameForMaterial(MP_CustomData1, Material), MP_CustomData1, FMaterialAttributeDefinitionMap::GetDisplayNameForMaterial(MP_CustomData1, Material)));
+		MaterialInputs.Add(FMaterialInputInfo(FMaterialAttributeDefinitionMap::GetDisplayNameForMaterial(MP_AmbientOcclusion, Material), MP_AmbientOcclusion, LOCTEXT("AmbientOcclusionToolTip", "Simulate the self-shadowing that happens within crevices of a surface, or of a volume for volumetric clouds only")));
+		MaterialInputs.Add(FMaterialInputInfo(FMaterialAttributeDefinitionMap::GetDisplayNameForMaterial(MP_Refraction, Material), MP_Refraction, LOCTEXT("RefractionToolTip", "Takes in a texture or value that simulates the index of refraction of the surface")));
 
 		for (int32 UVIndex = 0; UVIndex < UE_ARRAY_COUNT(Material->CustomizedUVs); UVIndex++)
 		{
 			//@todo - localize
-			MaterialInputs.Add( FMaterialInputInfo( FText::FromString(FString::Printf(TEXT("Customized UV%u"), UVIndex)), (EMaterialProperty)(MP_CustomizedUVs0 + UVIndex), FText::FromString(FString::Printf( TEXT( "CustomizedUV%uToolTip" ), UVIndex ) ) ) );
+			MaterialInputs.Add(FMaterialInputInfo(FText::FromString(FString::Printf(TEXT("Customized UV%u"), UVIndex)), (EMaterialProperty)(MP_CustomizedUVs0 + UVIndex), FText::FromString(FString::Printf(TEXT("CustomizedUV%uToolTip"), UVIndex))));
 		}
 
 		MaterialInputs.Add(FMaterialInputInfo(FMaterialAttributeDefinitionMap::GetDisplayNameForMaterial(MP_PixelDepthOffset, Material), MP_PixelDepthOffset, LOCTEXT("PixelDepthOffsetToolTip", "Pixel Depth Offset")));
 		MaterialInputs.Add(FMaterialInputInfo(FMaterialAttributeDefinitionMap::GetDisplayNameForMaterial(MP_ShadingModel, Material), MP_ShadingModel, LOCTEXT("ShadingModelToolTip", "Selects which shading model should be used per pixel")));
 
 		//^^^ New material properties go above here. ^^^^
-		MaterialInputs.Add(FMaterialInputInfo(LOCTEXT("MaterialAttributes", "Material Attributes"), MP_MaterialAttributes, LOCTEXT( "MaterialAttributesToolTip", "Material Attributes" ) ));
+		MaterialInputs.Add(FMaterialInputInfo(LOCTEXT("MaterialAttributes", "Material Attributes"), MP_MaterialAttributes, LOCTEXT("MaterialAttributesToolTip", "Material Attributes")));
 
 		// Add Root Node
 		FGraphNodeCreator<UMaterialGraphNode_Root> NodeCreator(*this);
@@ -74,53 +97,105 @@ void UMaterialGraph::RebuildGraph()
 		NodeCreator.Finalize();
 	}
 
-	for (int32 Index = 0; Index < Material->Expressions.Num(); Index++)
+	TArray<UMaterialExpression*> ChildSubGraphExpressions;
+
+	// Composite's use reroutes under the hood that we don't want to create nodes for, gather their expressions for checking
+	TArray<UMaterialExpressionReroute*> CompositeRerouteExpressions;
+	if (UMaterialExpressionComposite* SubgraphParentComposite = Cast<UMaterialExpressionComposite>(SubgraphExpression))
 	{
-		AddExpression(Material->Expressions[Index], false);
+		CompositeRerouteExpressions = SubgraphParentComposite->GetCurrentReroutes();
 	}
 
-	for (int32 Index = 0; Index < Material->EditorComments.Num(); Index++)
+	if (const TArray<UMaterialExpression*>* Expressions = SubgraphExpressionMap.Find(SubgraphExpression))
 	{
-		AddComment(Material->EditorComments[Index]);
+		for (UMaterialExpression* Expression : *Expressions)
+		{
+			if (!CompositeRerouteExpressions.Contains(Expression))
+			{
+				AddExpression(Expression, false);
+
+				//@TODO: Make a better way to check if an expression represents a subgraph than by type.
+				if (Cast<UMaterialExpressionComposite>(Expression))
+				{
+					ChildSubGraphExpressions.Add(Expression);
+				}
+			}
+		}
+	}
+
+	if (const TArray<UMaterialExpressionComment*>* Comments = SubgraphCommentMap.Find(SubgraphExpression))
+	{
+		for (UMaterialExpressionComment* Comment : *Comments)
+		{
+			AddComment(Comment);
+		}
+	}
+
+	for (UMaterialExpression* ChildSubGraphExpression : ChildSubGraphExpressions)
+	{
+		UMaterialGraph* Subgraph = AddSubGraph(ChildSubGraphExpression);
+
+		if (UMaterialGraphNode_Composite* CompositeNode = Cast<UMaterialGraphNode_Composite>(ChildSubGraphExpression->GraphNode))
+		{
+			CompositeNode->BoundGraph = Subgraph;
+			Subgraph->Rename(*CastChecked<UMaterialExpressionComposite>(CompositeNode->MaterialExpression)->SubgraphName);
+		}
+
+		Subgraph->RebuildGraphInternal(SubgraphExpressionMap, SubgraphCommentMap);
 	}
 
 	LinkGraphNodesFromMaterial();
 }
 
-UMaterialGraphNode* UMaterialGraph::AddExpression(UMaterialExpression* Expression, bool bUserInvoked)
+template<typename NodeType>
+static UMaterialGraphNode* InitExpressionNewNode(UMaterialGraph* Graph, UMaterialExpression* Expression, bool bUserInvoked)
 {
-	UMaterialGraphNode* NewNode = NULL;
-	if (Expression && Expression->IsA(UMaterialExpressionReroute::StaticClass()))
+	UMaterialGraphNode* NewNode = nullptr;
+
+	FGraphNodeCreator<NodeType> NodeCreator(*Graph);
+	if (bUserInvoked)
+	{
+		NewNode = NodeCreator.CreateUserInvokedNode();
+	}
+	else
 	{
-		Modify();
-		FGraphNodeCreator<UMaterialGraphNode_Knot> NodeCreator(*this);
 		NewNode = NodeCreator.CreateNode(false);
-		NewNode->MaterialExpression = Expression;
-		NewNode->RealtimeDelegate = RealtimeDelegate;
-		NewNode->MaterialDirtyDelegate = MaterialDirtyDelegate;
-		Expression->GraphNode = NewNode;
-		NodeCreator.Finalize();
 	}
-	else if (Expression)
+	NewNode->MaterialExpression = Expression;
+	NewNode->RealtimeDelegate = Graph->RealtimeDelegate;
+	NewNode->MaterialDirtyDelegate = Graph->MaterialDirtyDelegate;
+	Expression->GraphNode = NewNode;
+	Expression->SubgraphExpression = Graph->SubgraphExpression;
+	NodeCreator.Finalize();
+
+	return NewNode;
+}
+
+UMaterialGraphNode* UMaterialGraph::AddExpression(UMaterialExpression* Expression, bool bUserInvoked)
+{
+	if (Expression)
 	{
 		Modify();
-		FGraphNodeCreator<UMaterialGraphNode> NodeCreator(*this);
-		if(bUserInvoked)
+
+		if (Expression->IsA(UMaterialExpressionReroute::StaticClass()))
 		{
-			NewNode = NodeCreator.CreateUserInvokedNode();
+			return InitExpressionNewNode<UMaterialGraphNode_Knot>(this, Expression, false);
 		}
-		else
+		else if (Expression->IsA(UMaterialExpressionComposite::StaticClass()))
 		{
-			NewNode = NodeCreator.CreateNode(false);
+			return InitExpressionNewNode<UMaterialGraphNode_Composite>(this, Expression, false);
+		}
+		else if (Expression->IsA(UMaterialExpressionPinBase::StaticClass()))
+		{
+			return InitExpressionNewNode<UMaterialGraphNode_PinBase>(this, Expression, false);
+		}
+		else 
+		{
+			return InitExpressionNewNode<UMaterialGraphNode>(this, Expression, bUserInvoked);
 		}
-		NewNode->MaterialExpression = Expression;
-		NewNode->RealtimeDelegate = RealtimeDelegate;
-		NewNode->MaterialDirtyDelegate = MaterialDirtyDelegate;
-		Expression->GraphNode = NewNode;
-		NodeCreator.Finalize();
 	}
 
-	return NewNode;
+	return nullptr;
 }
 
 UMaterialGraphNode_Comment* UMaterialGraph::AddComment(UMaterialExpressionComment* Comment, bool bIsUserInvoked)
@@ -141,14 +216,48 @@ UMaterialGraphNode_Comment* UMaterialGraph::AddComment(UMaterialExpressionCommen
 		NewComment->MaterialExpressionComment = Comment;
 		NewComment->MaterialDirtyDelegate = MaterialDirtyDelegate;
 		Comment->GraphNode = NewComment;
+		Comment->SubgraphExpression = SubgraphExpression;
 		NodeCreator.Finalize();
 	}
 
 	return NewComment;
 }
 
+UMaterialGraph* UMaterialGraph::AddSubGraph(UMaterialExpression* InSubgraphExpression)
+{
+	UMaterialGraph* SubGraph = CastChecked<UMaterialGraph>(FBlueprintEditorUtils::CreateNewGraph(InSubgraphExpression->GraphNode, NAME_None, UMaterialGraph::StaticClass(), Schema));
+	check(SubGraph);
+
+	SubGraph->Material = Material;
+	SubGraph->MaterialFunction = MaterialFunction;
+	SubGraph->RealtimeDelegate = RealtimeDelegate;
+	SubGraph->MaterialDirtyDelegate = MaterialDirtyDelegate;
+	SubGraph->ToggleCollapsedDelegate = ToggleCollapsedDelegate;
+	SubGraph->SubgraphExpression = InSubgraphExpression;
+	SubGraphs.Add(SubGraph);
+
+	// If we are a subgraph ourselves, mark that on the expression.
+	InSubgraphExpression->SubgraphExpression = SubgraphExpression;
+
+	return SubGraph;
+}
+
 void UMaterialGraph::LinkGraphNodesFromMaterial()
 {
+	struct ExpressionMatchesPredicate
+	{
+		ExpressionMatchesPredicate(UMaterialExpressionReroute* InCompositeReroute)
+			: CompositeReroute(InCompositeReroute)
+		{}
+
+		bool operator()(const FCompositeReroute& Reroute)
+		{
+			return Reroute.Expression == CompositeReroute;
+		}
+
+		UMaterialExpressionReroute* CompositeReroute;
+	};
+
 	for (int32 Index = 0; Index < Nodes.Num(); ++Index)
 	{
 		Nodes[Index]->BreakAllNodeLinks();
@@ -164,8 +273,29 @@ void UMaterialGraph::LinkGraphNodesFromMaterial()
 
 			if (ExpressionInput.Expression)
 			{
-				UMaterialGraphNode* GraphNode = CastChecked<UMaterialGraphNode>(ExpressionInput.Expression->GraphNode);
-				InputPin->MakeLinkTo(GraphNode->GetOutputPin(GetValidOutputIndex(&ExpressionInput)));
+				if (UMaterialGraphNode* GraphNode = Cast<UMaterialGraphNode>(ExpressionInput.Expression->GraphNode))
+				{
+					InputPin->MakeLinkTo(GraphNode->GetOutputPin(GetValidOutputIndex(&ExpressionInput)));
+				}
+				else if (UMaterialExpressionReroute* CompositeReroute = CastChecked<UMaterialExpressionReroute>(ExpressionInput.Expression))
+				{
+					// This is an unseen composite reroute expression, find the actual expression output to connect to.
+					UMaterialExpressionComposite* OwningComposite = CastChecked<UMaterialExpressionComposite>(CompositeReroute->SubgraphExpression);
+
+					UMaterialGraphNode* OutputGraphNode;
+					int32 OutputPinIndex = OwningComposite->InputExpressions->ReroutePins.FindLastByPredicate(ExpressionMatchesPredicate(CompositeReroute));
+					if (OutputPinIndex != INDEX_NONE)
+					{
+						OutputGraphNode = CastChecked<UMaterialGraphNode>(OwningComposite->InputExpressions->GraphNode);
+					}
+					else
+					{
+						// Output pin base in the subgraph cannot have outputs, if this reroute isn't in the inputs connect to composite's outputs
+						OutputPinIndex = OwningComposite->OutputExpressions->ReroutePins.FindLastByPredicate(ExpressionMatchesPredicate(CompositeReroute));
+						OutputGraphNode = CastChecked<UMaterialGraphNode>(OwningComposite->GraphNode);
+					}
+					InputPin->MakeLinkTo(OutputGraphNode->GetOutputPin(OutputPinIndex));
+				}
 			}
 		}
 	}
@@ -173,21 +303,44 @@ void UMaterialGraph::LinkGraphNodesFromMaterial()
 	for (int32 Index = 0; Index < Material->Expressions.Num(); Index++)
 	{
 		UMaterialExpression* Expression = Material->Expressions[Index];
-
+		
 		if (Expression)
 		{
 			const TArray<FExpressionInput*> ExpressionInputs = Expression->GetInputs();
 			for (int32 InputIndex = 0; InputIndex < ExpressionInputs.Num(); ++InputIndex)
 			{
-				UEdGraphPin* InputPin = CastChecked<UMaterialGraphNode>(Expression->GraphNode)->GetInputPin(InputIndex);
-
-				// InputPin can be null during a PostEditChange when there is a circular dependency between nodes, and nodes have pins that are dynamically created
-				if (InputPin != nullptr && ExpressionInputs[InputIndex]->Expression
-					// Unclear why this is null sometimes, but this is safer than crashing
-					&& ExpressionInputs[InputIndex]->Expression->GraphNode)
+				if (UMaterialGraphNode* MaterialGraphNode = Cast<UMaterialGraphNode>(Expression->GraphNode))
 				{
-					UMaterialGraphNode* GraphNode = CastChecked<UMaterialGraphNode>(ExpressionInputs[InputIndex]->Expression->GraphNode);
-					InputPin->MakeLinkTo(GraphNode->GetOutputPin(GetValidOutputIndex(ExpressionInputs[InputIndex])));
+					UEdGraphPin* InputPin = MaterialGraphNode->GetInputPin(InputIndex);
+
+					// InputPin can be null during a PostEditChange when there is a circular dependency between nodes, and nodes have pins that are dynamically created
+					if (InputPin != nullptr && ExpressionInputs[InputIndex]->Expression)
+					{
+						// Unclear why this is null sometimes outside of composite reroute, but this is safer than crashing
+						if (UMaterialGraphNode* GraphNode = Cast<UMaterialGraphNode>(ExpressionInputs[InputIndex]->Expression->GraphNode))
+						{
+							InputPin->MakeLinkTo(GraphNode->GetOutputPin(GetValidOutputIndex(ExpressionInputs[InputIndex])));
+						}
+						else if (UMaterialExpressionReroute* CompositeReroute = Cast<UMaterialExpressionReroute>(ExpressionInputs[InputIndex]->Expression))
+						{
+							// This is an unseen composite reroute expression, find the actual expression output to connect to.
+							UMaterialExpressionComposite* OwningComposite = CastChecked<UMaterialExpressionComposite>(CompositeReroute->SubgraphExpression);
+
+							UMaterialGraphNode* OutputGraphNode;
+							int32 OutputPinIndex = OwningComposite->InputExpressions->ReroutePins.FindLastByPredicate(ExpressionMatchesPredicate(CompositeReroute));
+							if (OutputPinIndex != INDEX_NONE)
+							{
+								OutputGraphNode = CastChecked<UMaterialGraphNode>(OwningComposite->InputExpressions->GraphNode);
+							}
+							else
+							{
+								// Output pin base in the subgraph cannot have outputs, if this reroute isn't in the inputs connect to composite's outputs
+								OutputPinIndex = OwningComposite->OutputExpressions->ReroutePins.FindLastByPredicate(ExpressionMatchesPredicate(CompositeReroute));
+								OutputGraphNode = CastChecked<UMaterialGraphNode>(OwningComposite->GraphNode);
+							}
+							InputPin->MakeLinkTo(OutputGraphNode->GetOutputPin(OutputPinIndex));
+						}
+					}
 				}
 			}
 		}
@@ -226,7 +379,7 @@ void UMaterialGraph::LinkMaterialExpressionsFromGraph() const
 					{
 						if (OutputPins[OutPinIndex] == InputPins[PinIndex]->LinkedTo[0])
 						{
-							if (MaterialInput.OutputIndex != OutPinIndex || MaterialInput.Expression != ConnectedNode->MaterialExpression)
+							if (!ConnectedNode->MaterialExpression->IsExpressionConnected(&MaterialInput, OutPinIndex))
 							{
 								ConnectedNode->MaterialExpression->Modify();
 								MaterialInput.Connect(OutPinIndex, ConnectedNode->MaterialExpression);
@@ -281,13 +434,14 @@ void UMaterialGraph::LinkMaterialExpressionsFromGraph() const
 							{
 								if (OutputPins[OutPinIndex] == InputPins[PinIndex]->LinkedTo[0])
 								{
-									if (ExpressionInput && (ExpressionInput->OutputIndex != OutPinIndex || ExpressionInput->Expression != ConnectedNode->MaterialExpression))
+									if (ExpressionInput && !ConnectedNode->MaterialExpression->IsExpressionConnected(ExpressionInput, OutPinIndex))
 									{
 										if (!bModifiedExpression)
 										{
 											bModifiedExpression = true;
 											Expression->Modify();
 										}
+
 										ConnectedNode->MaterialExpression->Modify();
 										ExpressionInput->Connect(OutPinIndex, ConnectedNode->MaterialExpression);
 									}
@@ -333,6 +487,12 @@ void UMaterialGraph::LinkMaterialExpressionsFromGraph() const
 			}
 		}
 	}
+
+	// Also link subgraphs?
+	for (const UEdGraph* SubGraph : SubGraphs)
+	{
+		CastChecked<UMaterialGraph>(SubGraph)->LinkMaterialExpressionsFromGraph();
+	}
 }
 
 bool UMaterialGraph::IsInputActive(UEdGraphPin* GraphPin) const
diff --git a/Engine/Source/Editor/UnrealEd/Private/MaterialGraphNode.cpp b/Engine/Source/Editor/UnrealEd/Private/MaterialGraphNode.cpp
--- a/Engine/Source/Editor/UnrealEd/Private/MaterialGraphNode.cpp
+++ b/Engine/Source/Editor/UnrealEd/Private/MaterialGraphNode.cpp
@@ -9,6 +9,9 @@
 #include "ToolMenus.h"
 #include "MaterialGraph/MaterialGraphSchema.h"
 
+#include "Materials/MaterialExpressionComposite.h"
+#include "Materials/MaterialExpressionPinBase.h"
+
 #include "Materials/MaterialExpressionComponentMask.h"
 #include "Materials/MaterialExpressionConstant.h"
 #include "Materials/MaterialExpressionConstant2Vector.h"
@@ -443,7 +446,10 @@ void UMaterialGraphNode::GetNodeContextMenuActions(UToolMenu* Menu, UGraphNodeCo
 		}
 			
 		{
-			FToolMenuSection& Section = Menu->AddSection("Alignment");
+			FToolMenuSection& Section = Menu->AddSection("EdGraphSchemaOrganization", LOCTEXT("OrganizationHeader", "Organization"));
+			Section.AddMenuEntry(FGraphEditorCommands::Get().CollapseNodes);
+			Section.AddMenuEntry(FGraphEditorCommands::Get().ExpandNodes);
+
 			Section.AddSubMenu(
 				"Alignment",
 				LOCTEXT("AlignmentHeader", "Alignment"),
diff --git a/Engine/Source/Editor/UnrealEd/Private/MaterialGraphNode_Comment.cpp b/Engine/Source/Editor/UnrealEd/Private/MaterialGraphNode_Comment.cpp
--- a/Engine/Source/Editor/UnrealEd/Private/MaterialGraphNode_Comment.cpp
+++ b/Engine/Source/Editor/UnrealEd/Private/MaterialGraphNode_Comment.cpp
@@ -9,6 +9,9 @@
 #include "MaterialGraph/MaterialGraphSchema.h"
 #include "Materials/MaterialExpressionComment.h"
 #include "Framework/Commands/GenericCommands.h"
+#include "GraphEditorActions.h"
+
+#define LOCTEXT_NAMESPACE "MaterialGraphNode_Comment"
 
 /////////////////////////////////////////////////////
 // UMaterialGraphNode_Comment
@@ -79,6 +82,36 @@ void UMaterialGraphNode_Comment::GetNodeContextMenuActions(UToolMenu* Menu, UGra
 			Section.AddMenuEntry(FGenericCommands::Get().Copy);
 			Section.AddMenuEntry(FGenericCommands::Get().Duplicate);
 		}
+		
+		{
+			FToolMenuSection& Section = Menu->AddSection("EdGraphSchemaOrganization", LOCTEXT("OrganizationHeader", "Organization"));
+			Section.AddMenuEntry(FGraphEditorCommands::Get().CollapseNodes);
+			Section.AddMenuEntry(FGraphEditorCommands::Get().ExpandNodes);
+
+			Section.AddSubMenu(
+				"Alignment",
+				LOCTEXT("AlignmentHeader", "Alignment"),
+				FText(),
+				FNewToolMenuDelegate::CreateLambda([](UToolMenu* InMenu)
+				{
+					{
+						FToolMenuSection& SubMenuSection = InMenu->AddSection("EdGraphSchemaAlignment", LOCTEXT("AlignHeader", "Align"));
+						SubMenuSection.AddMenuEntry(FGraphEditorCommands::Get().AlignNodesTop);
+						SubMenuSection.AddMenuEntry(FGraphEditorCommands::Get().AlignNodesMiddle);
+						SubMenuSection.AddMenuEntry(FGraphEditorCommands::Get().AlignNodesBottom);
+						SubMenuSection.AddMenuEntry(FGraphEditorCommands::Get().AlignNodesLeft);
+						SubMenuSection.AddMenuEntry(FGraphEditorCommands::Get().AlignNodesCenter);
+						SubMenuSection.AddMenuEntry(FGraphEditorCommands::Get().AlignNodesRight);
+						SubMenuSection.AddMenuEntry(FGraphEditorCommands::Get().StraightenConnections);
+					}
+
+					{
+						FToolMenuSection& SubMenuSection = InMenu->AddSection("EdGraphSchemaDistribution", LOCTEXT("DistributionHeader", "Distribution"));
+						SubMenuSection.AddMenuEntry(FGraphEditorCommands::Get().DistributeNodesHorizontally);
+						SubMenuSection.AddMenuEntry(FGraphEditorCommands::Get().DistributeNodesVertically);
+					}
+				}));
+		}
 	}
 }
 
@@ -148,3 +181,6 @@ void UMaterialGraphNode_Comment::ResetMaterialExpressionOwner()
 		MaterialExpressionComment->GraphNode = this;
 	}
 }
+
+
+#undef LOCTEXT_NAMESPACE
\ No newline at end of file
diff --git a/Engine/Source/Editor/UnrealEd/Private/MaterialGraphNode_Composite.cpp b/Engine/Source/Editor/UnrealEd/Private/MaterialGraphNode_Composite.cpp
new file mode 100644
--- /dev/null
+++ b/Engine/Source/Editor/UnrealEd/Private/MaterialGraphNode_Composite.cpp
@@ -0,0 +1,168 @@
+// Copyright Epic Games, Inc. All Rights Reserved.
+
+/*=============================================================================
+	MaterialGraphNode_Composite.cpp
+=============================================================================*/
+
+#include "MaterialGraph/MaterialGraphNode_Composite.h"
+#include "ToolMenus.h"
+#include "MaterialGraph/MaterialGraphSchema.h"
+#include "MaterialGraph/MaterialGraphNode.h"
+#include "MaterialGraph/MaterialGraphNode_Comment.h"
+#include "Materials/MaterialExpressionComposite.h"
+#include "Materials/MaterialExpressionPinBase.h"
+#include "MaterialEditingLibrary.h"
+#include "MaterialEditorUtilities.h"
+#include "EdGraphUtilities.h"
+#include "Kismet2/BlueprintEditorUtils.h"
+#include "Framework/Commands/GenericCommands.h"
+
+#define LOCTEXT_NAMESPACE "MaterialGraphNode_Composite"
+
+/////////////////////////////////////////////////////
+// UMaterialGraphNode_Composite
+
+UMaterialGraphNode_Composite::UMaterialGraphNode_Composite(const FObjectInitializer& ObjectInitializer)
+	: Super(ObjectInitializer)
+{
+}
+
+void UMaterialGraphNode_Composite::PostEditUndo()
+{
+	Super::PostEditUndo();
+	FixupInputAndOutputPinBases();
+}
+
+void UMaterialGraphNode_Composite::PostCopyNode()
+{
+	Super::PostCopyNode();
+
+	if (BoundGraph != nullptr)
+	{
+		for (UEdGraphNode* Node : BoundGraph->Nodes)
+		{
+			if (UMaterialGraphNode* MaterialNode = Cast<UMaterialGraphNode>(Node))
+			{
+				MaterialNode->PostCopyNode();
+			}
+			else if (UMaterialGraphNode_Comment* CommentNode = Cast<UMaterialGraphNode_Comment>(Node))
+			{
+				CommentNode->PostCopyNode();
+			}
+		}
+	}
+}
+
+UObject* UMaterialGraphNode_Composite::GetJumpTargetForDoubleClick() const
+{
+	return BoundGraph;
+}
+
+bool UMaterialGraphNode_Composite::CanJumpToDefinition() const
+{
+	return GetJumpTargetForDoubleClick() != nullptr;
+}
+
+void UMaterialGraphNode_Composite::JumpToDefinition() const
+{
+	if (UObject* HyperlinkTarget = GetJumpTargetForDoubleClick())
+	{
+		FMaterialEditorUtilities::BringFocusAttentionOnObject(HyperlinkTarget);
+	}
+}
+
+void UMaterialGraphNode_Composite::DestroyNode()
+{
+	UEdGraph* GraphToRemove = BoundGraph;
+
+	BoundGraph = nullptr;
+	Super::DestroyNode();
+
+	if (GraphToRemove)
+	{
+		GraphToRemove->Modify();
+		FBlueprintEditorUtils::RemoveGraph(nullptr, GraphToRemove, EGraphRemoveFlags::None);
+	}
+}
+
+void UMaterialGraphNode_Composite::PrepareForCopying()
+{
+	if (BoundGraph != nullptr)
+	{
+		for (UEdGraphNode* Node : BoundGraph->Nodes)
+		{
+			Node->PrepareForCopying();
+		}
+	}
+
+	Super::PrepareForCopying();
+}
+
+void UMaterialGraphNode_Composite::PostPasteNode()
+{
+	Super::PostPasteNode();
+
+	if (BoundGraph != nullptr)
+	{
+		UEdGraph* ParentGraph = CastChecked<UEdGraph>(GetOuter());
+		ensure(BoundGraph != ParentGraph);
+
+		// Update the InputSinkNode / OutputSourceNode pointers to point to the new graph
+		TSet<UEdGraphNode*> BoundaryNodes;
+		for (int32 NodeIndex = 0; NodeIndex < BoundGraph->Nodes.Num(); ++NodeIndex)
+		{
+			UEdGraphNode* Node = BoundGraph->Nodes[NodeIndex];
+			BoundaryNodes.Add(Node);
+		}
+
+		ensure(BoundGraph->SubGraphs.Find(ParentGraph) == INDEX_NONE);
+
+		//Nested composites will already be in the SubGraph array
+		if (ParentGraph->SubGraphs.Find(BoundGraph) == INDEX_NONE)
+		{
+			ParentGraph->SubGraphs.Add(BoundGraph);
+		}
+
+		FEdGraphUtilities::PostProcessPastedNodes(BoundaryNodes);
+	}
+}
+
+void UMaterialGraphNode_Composite::OnRenameNode(const FString& NewName)
+{
+	MaterialExpression->Modify();
+	MaterialExpression->SetEditableName(NewName);
+	MaterialExpression->MarkPackageDirty();
+	MaterialExpression->ValidateParameterName();
+
+	CastChecked<UMaterialGraph>(GetGraph())->Material->UpdateExpressionParameterName(MaterialExpression);
+	MaterialDirtyDelegate.ExecuteIfBound();
+}
+
+void UMaterialGraphNode_Composite::FixupInputAndOutputPinBases()
+{
+	if (BoundGraph)
+	{
+		for (UEdGraphNode* Node : BoundGraph->Nodes)
+		{
+			if (UMaterialGraphNode* MaterialNode = Cast<UMaterialGraphNode>(Node))
+			{
+				UMaterialExpressionPinBase* PinBase = Cast<UMaterialExpressionPinBase>(MaterialNode->MaterialExpression);
+
+				if (PinBase && PinBase->PinDirection == EGPD_Input)
+				{
+					CastChecked<UMaterialExpressionComposite>(MaterialExpression)->OutputExpressions = PinBase;
+					PinBase->SubgraphExpression = MaterialExpression;
+				}
+				else if (PinBase && PinBase->PinDirection == EGPD_Output)
+				{
+					CastChecked<UMaterialExpressionComposite>(MaterialExpression)->InputExpressions = PinBase;
+					PinBase->SubgraphExpression = MaterialExpression;
+				}
+			}
+		}
+
+		ReconstructNode();
+	}
+}
+
+#undef LOCTEXT_NAMESPACE
\ No newline at end of file
diff --git a/Engine/Source/Editor/UnrealEd/Private/MaterialGraphNode_PinBase.cpp b/Engine/Source/Editor/UnrealEd/Private/MaterialGraphNode_PinBase.cpp
new file mode 100644
--- /dev/null
+++ b/Engine/Source/Editor/UnrealEd/Private/MaterialGraphNode_PinBase.cpp
@@ -0,0 +1,48 @@
+// Copyright Epic Games, Inc. All Rights Reserved.
+
+/*=============================================================================
+	MaterialGraphNode_PinBase.cpp
+=============================================================================*/
+
+#include "MaterialGraph/MaterialGraphNode_PinBase.h"
+#include "MaterialGraph/MaterialGraphSchema.h"
+#include "MaterialGraph/MaterialGraphNode.h"
+#include "MaterialGraph/MaterialGraphNode_Comment.h"
+#include "Materials/MaterialExpressionPinBase.h"
+#include "Materials/MaterialExpressionReroute.h"
+#include "MaterialEditingLibrary.h"
+#include "MaterialEditorUtilities.h"
+#include "EdGraphUtilities.h"
+
+#define LOCTEXT_NAMESPACE "MaterialGraphNode_PinBase"
+
+/////////////////////////////////////////////////////
+// UMaterialGraphNode_PinBase
+
+void UMaterialGraphNode_PinBase::PostCopyNode()
+{
+	Super::PostCopyNode();
+
+	if (UMaterialExpressionPinBase* PinBase = Cast<UMaterialExpressionPinBase>(MaterialExpression))
+	{
+		for (FCompositeReroute& ReroutePin : PinBase->ReroutePins)
+		{
+			ReroutePin.Expression->Rename(nullptr, GetOuter(), REN_DontCreateRedirectors);
+		}
+	}
+}
+
+void UMaterialGraphNode_PinBase::PrepareForCopying()
+{
+	Super::PrepareForCopying();
+
+	if (UMaterialExpressionPinBase* PinBase = Cast<UMaterialExpressionPinBase>(MaterialExpression))
+	{
+		for (FCompositeReroute& ReroutePin : PinBase->ReroutePins)
+		{
+			ReroutePin.Expression->Rename(nullptr, GetOuter(), REN_DontCreateRedirectors);
+		}
+	}
+}
+
+#undef LOCTEXT_NAMESPACE
\ No newline at end of file
diff --git a/Engine/Source/Editor/UnrealEd/Private/MaterialGraphSchema.cpp b/Engine/Source/Editor/UnrealEd/Private/MaterialGraphSchema.cpp
--- a/Engine/Source/Editor/UnrealEd/Private/MaterialGraphSchema.cpp
+++ b/Engine/Source/Editor/UnrealEd/Private/MaterialGraphSchema.cpp
@@ -21,12 +21,15 @@
 #include "MaterialGraph/MaterialGraph.h"
 #include "Engine/Texture.h"
 #include "MaterialGraph/MaterialGraphNode_Base.h"
+#include "MaterialGraph/MaterialGraphNode_Comment.h"
 #include "MaterialGraph/MaterialGraphNode.h"
 #include "MaterialGraph/MaterialGraphNode_Root.h"
 #include "Materials/MaterialParameterCollection.h"
 
 #include "Materials/MaterialExpressionCollectionParameter.h"
 #include "Materials/MaterialExpressionComment.h"
+#include "Materials/MaterialExpressionComposite.h"
+#include "Materials/MaterialExpressionPinBase.h"
 #include "Materials/MaterialExpressionFunctionInput.h"
 #include "Materials/MaterialExpressionTextureSample.h"
 #include "Materials/MaterialExpressionFunctionOutput.h"
@@ -146,6 +149,28 @@ UEdGraphNode* FMaterialGraphSchemaAction_NewFunctionCall::PerformAction(class UE
 	return NULL;
 }
 
+/////////////////////////////////////////////
+// FMaterialGraphSchemaAction_NewComposite //
+
+UEdGraphNode* FMaterialGraphSchemaAction_NewComposite::PerformAction(UEdGraph* ParentGraph, UEdGraphPin* FromPin, const FVector2D Location, bool bSelectNewNode)
+{
+	return SpawnNode(ParentGraph, Location);
+}
+
+UEdGraphNode* FMaterialGraphSchemaAction_NewComposite::SpawnNode(UEdGraph* ParentGraph, const FVector2D Location)
+{
+	const FScopedTransaction Transaction(NSLOCTEXT("UnrealEd", "MaterialEditorNewComposite", "Material Editor: New Composite"));
+
+	UMaterialExpressionComposite* NewComposite = FMaterialEditorUtilities::CreateNewMaterialExpressionComposite(ParentGraph, Location);
+
+	if (NewComposite)
+	{
+		return NewComposite->GraphNode;
+	}
+
+	return nullptr;
+}
+
 ///////////////////////////////////////////
 // FMaterialGraphSchemaAction_NewComment //
 
@@ -640,6 +665,22 @@ void UMaterialGraphSchema::BreakSinglePinLink(UEdGraphPin* SourcePin, UEdGraphPi
 	}
 }
 
+bool UMaterialGraphSchema::CanEncapuslateNode(UEdGraphNode const& TestNode) const
+{
+	if (TestNode.IsA(UMaterialGraphNode_Comment::StaticClass()))
+	{
+		return true;
+	}
+
+	// Disallow output nodes from encapsulation, everything else (including parameters) is fair game for materials.
+	const UMaterialGraphNode* MaterialGraphNode = Cast<UMaterialGraphNode>(&TestNode);
+
+	return  MaterialGraphNode && MaterialGraphNode->MaterialExpression 
+		&& !MaterialGraphNode->MaterialExpression->IsA(UMaterialExpressionFunctionOutput::StaticClass())
+		&& !MaterialGraphNode->MaterialExpression->IsA(UMaterialExpressionPinBase::StaticClass())
+		&& !TestNode.IsA(UMaterialGraphNode_Root::StaticClass());
+}
+
 void UMaterialGraphSchema::DroppedAssetsOnGraph(const TArray<struct FAssetData>& Assets, const FVector2D& GraphPosition, UEdGraph* Graph) const
 {
 	UMaterialGraph* MaterialGraph = CastChecked<UMaterialGraph>(Graph);
@@ -825,6 +866,18 @@ void UMaterialGraphSchema::GetMaterialFunctionActions(FGraphActionMenuBuilder& A
 	}
 }
 
+void UMaterialGraphSchema::GetCompositeAction(FGraphActionMenuBuilder& ActionMenuBuilder, const UEdGraph* CurrentGraph) const
+{
+	if (!ActionMenuBuilder.FromPin)
+	{
+		const bool bIsManyNodesSelected = CurrentGraph ? (FMaterialEditorUtilities::GetNumberOfSelectedNodes(CurrentGraph) > 0) : false;
+		const FText CompositeDesc = LOCTEXT("CompositeDesc", "New Composite");
+		const FText CompositeToolTip = LOCTEXT("CompositeToolTip", "Create a composite node that holds a subgraph.");
+		TSharedPtr<FMaterialGraphSchemaAction_NewComposite> NewAction(new FMaterialGraphSchemaAction_NewComposite(FText::GetEmpty(), CompositeDesc, CompositeToolTip, 0));
+		ActionMenuBuilder.AddAction(NewAction);
+	}
+}
+
 void UMaterialGraphSchema::GetCommentAction(FGraphActionMenuBuilder& ActionMenuBuilder, const UEdGraph* CurrentGraph) const
 {
 	if (!ActionMenuBuilder.FromPin)
diff --git a/Engine/Source/Editor/UnrealEd/Public/GraphEditor.h b/Engine/Source/Editor/UnrealEd/Public/GraphEditor.h
--- a/Engine/Source/Editor/UnrealEd/Public/GraphEditor.h
+++ b/Engine/Source/Editor/UnrealEd/Public/GraphEditor.h
@@ -470,6 +470,32 @@ public:
 
 	UNREALED_API void FocusCommentNodes(TArray<UEdGraphNode*> &CommentNodes, TArray<UEdGraphNode*> &RelatedNodes);
 
+	virtual void OnCollapseNodes()
+	{
+		if (Implementation.IsValid())
+		{
+			Implementation->OnCollapseNodes();
+		}
+	}
+
+	virtual bool CanCollapseNodes() const
+	{
+		return Implementation.IsValid() ? Implementation->CanCollapseNodes() : false;
+	}
+
+	virtual void OnExpandNodes()
+	{
+		if (Implementation.IsValid())
+		{
+			Implementation->OnExpandNodes();
+		}
+	}
+
+	virtual bool CanExpandNodes() const
+	{
+		return Implementation.IsValid() ? Implementation->CanExpandNodes() : false;
+	}
+
 	virtual void OnAlignTop()
 	{
 		if (Implementation.IsValid())
diff --git a/Engine/Source/Editor/UnrealEd/Public/UnrealEd.h b/Engine/Source/Editor/UnrealEd/Public/UnrealEd.h
--- a/Engine/Source/Editor/UnrealEd/Public/UnrealEd.h
+++ b/Engine/Source/Editor/UnrealEd/Public/UnrealEd.h
@@ -216,6 +216,7 @@ MONOLITHIC_HEADER_BOILERPLATE()
 #include "MaterialGraph/MaterialGraphNode_Base.h"
 #include "MaterialGraph/MaterialGraphNode.h"
 #include "MaterialGraph/MaterialGraphNode_Root.h"
+#include "MaterialGraph/MaterialGraphNode_Composite.h"
 #include "MaterialGraph/MaterialGraphSchema.h"
 #include "Exporters/ModelExporterT3D.h"
 #include "Exporters/ObjectExporterT3D.h"
diff --git a/Engine/Source/Runtime/Engine/Classes/EdGraph/EdGraph.h b/Engine/Source/Runtime/Engine/Classes/EdGraph/EdGraph.h
--- a/Engine/Source/Runtime/Engine/Classes/EdGraph/EdGraph.h
+++ b/Engine/Source/Runtime/Engine/Classes/EdGraph/EdGraph.h
@@ -211,6 +211,9 @@ public:
 	/** Get all children graphs in the specified graph */
 	void GetAllChildrenGraphs(TArray<UEdGraph*>& Graphs) const;
 
+	/** Get parent outer graph, if it exists */
+	static UEdGraph* GetOuterGraph(UObject* Obj);
+
 	/** Util to find a good place for a new node */
 	FVector2D GetGoodPlaceForNewNode();
 
diff --git a/Engine/Source/Runtime/Engine/Classes/Materials/MaterialExpression.h b/Engine/Source/Runtime/Engine/Classes/Materials/MaterialExpression.h
--- a/Engine/Source/Runtime/Engine/Classes/Materials/MaterialExpression.h
+++ b/Engine/Source/Runtime/Engine/Classes/Materials/MaterialExpression.h
@@ -113,6 +113,10 @@ class ENGINE_API UMaterialExpression : public UObject
 	UPROPERTY(transient)
 	UEdGraphNode*	GraphNode;
 
+	/** If exists, expresssion containing this expression within its subgraph. */
+	UPROPERTY()
+	UMaterialExpression* SubgraphExpression;
+
 	/** Text of last error for this expression */
 	FString LastErrorText;
 
@@ -334,10 +338,15 @@ class ENGINE_API UMaterialExpression : public UObject
 #endif // WITH_EDITOR
 
 #if WITH_EDITOR
+	/**
+	 * Check if input exppresion is directly connected to the material.
+	 */
+	virtual bool IsExpressionConnected(FExpressionInput* Input, int32 OutputIndex);
+
 	/**
 	 * Connects the specified input expression to the specified output of this expression.
 	 */
-	void ConnectExpression(FExpressionInput* Input, int32 OutputIndex);
+	virtual void ConnectExpression(FExpressionInput* Input, int32 OutputIndex);
 #endif // WITH_EDITOR
 
 	/** 
diff --git a/Engine/Source/Runtime/Engine/Classes/Materials/MaterialExpressionComposite.h b/Engine/Source/Runtime/Engine/Classes/Materials/MaterialExpressionComposite.h
new file mode 100644
--- /dev/null
+++ b/Engine/Source/Runtime/Engine/Classes/Materials/MaterialExpressionComposite.h
@@ -0,0 +1,94 @@
+// Copyright Epic Games, Inc. All Rights Reserved.
+
+#pragma once
+
+#include "CoreMinimal.h"
+#include "UObject/ObjectMacros.h"
+#include "Materials/MaterialExpression.h"
+#include "Materials/MaterialExpressionReroute.h"
+#include "Materials/MaterialExpressionPinBase.h"
+#include "MaterialExpressionComposite.generated.h"
+
+/**
+ * Composite (subgraph) expression. Exists purely for organzational purposes.
+ * Under the hood uses reroute expressions for graph compilation. See below
+ * to understand how a particular reroute's input / output correlates to
+ * the inputs / outputs of composites and their pin bases.
+ * 
+ *      _________________          _________________
+ *     |   INPUT BASE    |        |   OUTPUT BASE   |
+ *     +--------+--------+        +--------+--------+
+ *     |        |   (>)  |   ->   |  (>)   |        |
+ *     |        |   (>)  |        |  (>)   |        |
+ *     |        |   (>)  |        |  (>)   |        |
+ *     |        |   (>)  |        |  (>)   |        |
+ *     |        |        |        |        |        |
+ *     +--------+--------+        +--------+--------+
+ *     | NODE IN:  NONE  |        | NODE IN:  PINS  |
+ *     | NODE OUT: PINS  |        | NODE OUT: NONE  |
+ *     |_________________|        |_________________|
+ *   
+ *                   ^                |                       
+ *                   |                v                                        
+ *              ____________________________         
+ *             |         COMPOSITE          |
+ *             +---------+--------+---------+
+ *         ->  |    (>)  |        |  (>)    |  ->
+ *             |    (>)  |        |  (>)    |
+ *             |    (>)  |        |  (>)    |
+ *             |    (>)  |        |  (>)    |
+ *             |         |        |         |
+ *             +---------+--------+---------+
+ *             | NODE IN:  TO INPUT BASE    |
+ *             | NODE OUT: FROM OUTPUT BASE |
+ *             |____________________________|
+ *
+ */
+
+UCLASS(MinimalAPI)
+class UMaterialExpressionComposite : public UMaterialExpression
+{
+	GENERATED_UCLASS_BODY()
+	
+	UPROPERTY(EditAnywhere, Category= MaterialExpressionComposite)
+	FString SubgraphName;
+
+	UPROPERTY()
+	UMaterialExpressionPinBase* InputExpressions;
+
+	UPROPERTY()
+	UMaterialExpressionPinBase* OutputExpressions;
+
+#if WITH_EDITOR
+	/** Get all reroute expressions used by this composite & its pin bases */
+	virtual TArray<UMaterialExpressionReroute*> GetCurrentReroutes() const;
+
+	//~ Begin UObject Interface.
+	virtual void PostEditChangeProperty(FPropertyChangedEvent& PropertyChangedEvent) override;
+	//~ End UObject Interface.
+
+	//~ Begin UMaterialExpression Interface
+	virtual bool CanRenameNode() const override { return true; }
+	virtual FString GetEditableName() const;
+	virtual void SetEditableName(const FString& NewName);
+
+	virtual TArray<FExpressionOutput>& GetOutputs() override;
+	virtual const TArray<FExpressionInput*> GetInputs() override;
+	virtual FExpressionInput* GetInput(int32 InputIndex) override;
+	virtual FName GetInputName(int32 InputIndex) const override;
+	virtual uint32 GetInputType(int32 InputIndex) override;
+	virtual uint32 GetOutputType(int32 OutputIndex) override;
+	virtual bool IsExpressionConnected(FExpressionInput* Input, int32 OutputIndex) override;
+	virtual void ConnectExpression(FExpressionInput* Input, int32 OutputIndex) override;
+
+	virtual void GetCaption(TArray<FString>& OutCaptions) const override;
+	//~ End UMaterialExpression Interface
+
+	//~ Begin UObject Interface.
+	virtual bool Modify(bool bAlwaysMarkDirty = true) override;
+	//~ End UObject Interface.
+#endif // WITH_EDITOR
+};
+
+
+
diff --git a/Engine/Source/Runtime/Engine/Classes/Materials/MaterialExpressionPinBase.h b/Engine/Source/Runtime/Engine/Classes/Materials/MaterialExpressionPinBase.h
new file mode 100644
--- /dev/null
+++ b/Engine/Source/Runtime/Engine/Classes/Materials/MaterialExpressionPinBase.h
@@ -0,0 +1,89 @@
+// Copyright Epic Games, Inc. All Rights Reserved.
+
+#pragma once
+
+#include "CoreMinimal.h"
+#include "UObject/ObjectMacros.h"
+#include "Materials/MaterialExpression.h"
+#include "Materials/MaterialExpressionReroute.h"
+#include "MaterialExpressionPinBase.generated.h"
+
+/**
+ * Collection of pins used for tunneling between graphs.
+ * Utilizes reroute expressions to ensure zero overhead in the compiled material.
+ * 
+ *      _________________          _________________
+ *     |   INPUT BASE    |        |   OUTPUT BASE   |
+ *     +--------+--------+        +--------+--------+
+ *     |        |   (>)  |   ->   |  (>)   |        |
+ *     |        |   (>)  |        |  (>)   |        |
+ *     |        |   (>)  |        |  (>)   |        |
+ *     |        |   (>)  |        |  (>)   |        |
+ *     |        |        |        |        |        |
+ *     +--------+--------+        +--------+--------+
+ *     | NODE IN:  NONE  |        | NODE IN:  PINS  |
+ *     | NODE OUT: PINS  |        | NODE OUT: NONE  |
+ *     |_________________|        |_________________|
+ *
+ */
+
+USTRUCT()
+struct FCompositeReroute
+{
+	GENERATED_USTRUCT_BODY()
+
+	UPROPERTY(EditAnywhere, Category=CustomInput, meta = (NoResetToDefault))
+	FName Name;
+
+	UPROPERTY(meta = (NoResetToDefault))
+	UMaterialExpressionReroute* Expression;
+};
+
+UCLASS(MinimalAPI)
+class UMaterialExpressionPinBase : public UMaterialExpression
+{
+	GENERATED_UCLASS_BODY()
+	
+	/** Underlying reroute pins used to compile material. Must call Modify after editing to update output expressions. */
+	UPROPERTY(EditAnywhere, Category = MaterialExpressionPinBase, meta = (NoResetToDefault, TitleProperty = Name))
+	TArray<FCompositeReroute> ReroutePins;
+
+	/** Direction of the pins for this base. */
+	UPROPERTY()
+	TEnumAsByte<enum EEdGraphPinDirection> PinDirection;
+
+#if WITH_EDITOR
+	/** Helper function to clear all reroutes. */
+	virtual void DeleteReroutePins();
+
+	//~ Begin UObject Interface.
+	virtual void PreEditChange(FProperty* PropertyAboutToChange) override;
+	virtual void PostEditChangeProperty(FPropertyChangedEvent& PropertyChangedEvent) override;
+	//~ End UObject Interface.
+
+	//~ Begin UMaterialExpression Interface
+	virtual TArray<FExpressionOutput>& GetOutputs() override;
+	virtual const TArray<FExpressionInput*> GetInputs() override;
+	virtual FExpressionInput* GetInput(int32 InputIndex) override;
+	virtual FName GetInputName(int32 InputIndex) const override;
+	virtual uint32 GetInputType(int32 InputIndex) override;
+	virtual uint32 GetOutputType(int32 OutputIndex) override;
+	virtual bool IsExpressionConnected(FExpressionInput* Input, int32 OutputIndex) override;
+	virtual void ConnectExpression(FExpressionInput* Input, int32 OutputIndex) override;
+
+	virtual void GetCaption(TArray<FString>& OutCaptions) const override;
+	//~ End UMaterialExpression Interface
+
+	//~ Begin UObject Interface.
+	virtual bool Modify(bool bAlwaysMarkDirty = true) override;
+	//~ End UObject Interface.
+
+private:
+	/** Used to remove reroute expressions from the material after we remove them. */
+	TArray<UMaterialExpressionReroute*> PreEditRereouteExpresions;
+
+#endif // WITH_EDITOR
+};
+
+
+
diff --git a/Engine/Source/Runtime/Engine/Private/EdGraph/EdGraph.cpp b/Engine/Source/Runtime/Engine/Private/EdGraph/EdGraph.cpp
--- a/Engine/Source/Runtime/Engine/Private/EdGraph/EdGraph.cpp
+++ b/Engine/Source/Runtime/Engine/Private/EdGraph/EdGraph.cpp
@@ -354,6 +354,26 @@ void UEdGraph::GetAllChildrenGraphs(TArray<UEdGraph*>& Graphs) const
 #endif // WITH_EDITORONLY_DATA
 }
 
+UEdGraph* UEdGraph::GetOuterGraph(UObject* Obj)
+{
+	if (Obj)
+	{
+		UObject* OuterObj = Obj->GetOuter();
+		if (OuterObj)
+		{
+			if (OuterObj->IsA(UEdGraph::StaticClass()))
+			{
+				return Cast<UEdGraph>(OuterObj);
+			}
+			else
+			{
+				return GetOuterGraph(OuterObj);
+			}
+		}
+	}
+	return nullptr;
+}
+
 FVector2D UEdGraph::GetGoodPlaceForNewNode()
 {
 	FVector2D BottomLeft(0,0);
diff --git a/Engine/Source/Runtime/Engine/Private/Materials/MaterialExpressions.cpp b/Engine/Source/Runtime/Engine/Private/Materials/MaterialExpressions.cpp
--- a/Engine/Source/Runtime/Engine/Private/Materials/MaterialExpressions.cpp
+++ b/Engine/Source/Runtime/Engine/Private/Materials/MaterialExpressions.cpp
@@ -66,6 +66,7 @@
 #include "Materials/MaterialExpressionClamp.h"
 #include "Materials/MaterialExpressionCollectionParameter.h"
 #include "Materials/MaterialExpressionComment.h"
+#include "Materials/MaterialExpressionComposite.h"
 #include "Materials/MaterialExpressionComponentMask.h"
 #include "Materials/MaterialExpressionConstant.h"
 #include "Materials/MaterialExpressionConstant2Vector.h"
@@ -131,6 +132,7 @@
 #include "Materials/MaterialExpressionOneMinus.h"
 #include "Materials/MaterialExpressionPanner.h"
 #include "Materials/MaterialExpressionParameter.h"
+#include "Materials/MaterialExpressionPinBase.h"
 #include "Materials/MaterialExpressionPreSkinnedLocalBounds.h"
 #include "Materials/MaterialExpressionPreviousFrameSwitch.h"
 #include "Materials/MaterialExpressionReroute.h"
@@ -245,6 +247,7 @@
 #include "MaterialCompiler.h"
 #if WITH_EDITOR
 #include "MaterialGraph/MaterialGraphNode_Comment.h"
+#include "MaterialGraph/MaterialGraphNode_Composite.h"
 #include "MaterialGraph/MaterialGraphNode.h"
 #include "Framework/Notifications/NotificationManager.h"
 #include "Widgets/Notifications/SNotificationList.h"
@@ -277,6 +280,7 @@ bool IsAllowedExpressionType(const UClass* const Class, const bool bMaterialFunc
 
 	// Exclude comments from the expression list, as well as the base parameter expression, as it should not be used directly
 	const bool bSharedAllowed = Class != UMaterialExpressionComment::StaticClass() 
+		&& Class != UMaterialExpressionPinBase::StaticClass()
 		&& Class != UMaterialExpressionParameter::StaticClass()
 		&& (Class != UMaterialExpressionTextureSampleParameter2DArray::StaticClass() || AllowTextureArrayAssetCreationVar->GetValueOnGameThread() != 0);
 
@@ -593,6 +597,7 @@ UMaterialExpression::UMaterialExpression(const FObjectInitializer& ObjectInitial
 	: Super(ObjectInitializer)
 #if WITH_EDITORONLY_DATA
 	, GraphNode(nullptr)
+	, SubgraphExpression(nullptr)
 #endif // WITH_EDITORONLY_DATA
 {
 #if WITH_EDITORONLY_DATA
@@ -941,30 +946,30 @@ FName UMaterialExpression::GetInputName(int32 InputIndex) const
 		{
 			for (int32 ArrayIndex = 0; ArrayIndex < StructProp->ArrayDim; ArrayIndex++)
 			{
-			if( Index == InputIndex )
-			{
+				if( Index == InputIndex )
+				{
 					FExpressionInput const* Input = StructProp->ContainerPtrToValuePtr<FExpressionInput>(this, ArrayIndex);
 
-						if (!Input->InputName.IsNone())
-						{
-							return Input->InputName;
-						}
-						else
-						{
-							FName StructName = StructProp->GetFName();
-
-					if (StructProp->ArrayDim > 1)
+					if (!Input->InputName.IsNone())
 					{
-								StructName = *FString::Printf(TEXT("%s_%d"), *StructName.ToString(), ArrayIndex);
+						return Input->InputName;
 					}
+					else
+					{
+						FName StructName = StructProp->GetFName();
 
-							return StructName;
+						if (StructProp->ArrayDim > 1)
+						{
+							StructName = *FString::Printf(TEXT("%s_%d"), *StructName.ToString(), ArrayIndex);
 						}
+
+						return StructName;
+					}
+				}
+				Index++;
 			}
-			Index++;
 		}
 	}
-	}
 	return NAME_None;
 }
 
@@ -1153,6 +1158,11 @@ bool UMaterialExpression::MatchesSearchQuery( const TCHAR* SearchQuery )
 	return Desc.Contains(SearchQuery);
 }
 
+bool UMaterialExpression::IsExpressionConnected( FExpressionInput* Input, int32 OutputIndex )
+{
+	return Input->OutputIndex == OutputIndex && Input->Expression == this;
+}
+
 void UMaterialExpression::ConnectExpression( FExpressionInput* Input, int32 OutputIndex )
 {
 	if( Input && OutputIndex >= 0 && OutputIndex < Outputs.Num() )
@@ -10229,6 +10239,425 @@ bool UMaterialExpressionComment::MatchesSearchQuery( const TCHAR* SearchQuery )
 }
 #endif // WITH_EDITOR
 
+///////////////////////////////////////////////////////////////////////////////
+// UMaterialExpressionComposite
+///////////////////////////////////////////////////////////////////////////////
+UMaterialExpressionComposite::UMaterialExpressionComposite(const FObjectInitializer& ObjectInitializer)
+	: Super(ObjectInitializer)
+{
+	// Structure to hold one-time initialization
+	struct FConstructorStatics
+	{
+		FText NAME_Utility;
+		FConstructorStatics()
+			: NAME_Utility(LOCTEXT( "Utility", "Utility" ))
+		{
+		}
+	};
+	static FConstructorStatics ConstructorStatics;
+
+	SubgraphName = "Collapsed Nodes";
+
+#if WITH_EDITORONLY_DATA
+	bShowOutputNameOnPin = true;
+
+	MenuCategories.Add(ConstructorStatics.NAME_Utility);
+	Outputs.Reset();
+#endif
+}
+
+#if WITH_EDITOR
+TArray<UMaterialExpressionReroute*> UMaterialExpressionComposite::GetCurrentReroutes() const
+{
+	TArray<UMaterialExpressionReroute*> RerouteExpressions;
+	for (const FCompositeReroute& InputReroute : InputExpressions->ReroutePins)
+	{
+		RerouteExpressions.Add(InputReroute.Expression);
+	}
+	for (const FCompositeReroute& OutputReroute : OutputExpressions->ReroutePins)
+	{
+		RerouteExpressions.Add(OutputReroute.Expression);
+	}
+	return RerouteExpressions;
+}
+
+void UMaterialExpressionComposite::PostEditChangeProperty(FPropertyChangedEvent& PropertyChangedEvent)
+{
+	if (UMaterialGraphNode_Composite* CompositeNode = Cast<UMaterialGraphNode_Composite>(GraphNode))
+	{
+		if (CompositeNode->BoundGraph && CompositeNode->BoundGraph->GetName() != SubgraphName)
+		{
+			CompositeNode->BoundGraph->Rename(*SubgraphName);
+		}
+	}
+
+	Super::PostEditChangeProperty(PropertyChangedEvent);
+}
+
+FString UMaterialExpressionComposite::GetEditableName() const
+{
+	return SubgraphName;
+}
+
+void UMaterialExpressionComposite::SetEditableName(const FString& NewName)
+{
+	SubgraphName = NewName;
+
+	if (UMaterialGraphNode_Composite* CompositeNode = Cast<UMaterialGraphNode_Composite>(GraphNode))
+	{
+		if (CompositeNode->BoundGraph && CompositeNode->BoundGraph->GetName() != SubgraphName)
+		{
+			CompositeNode->BoundGraph->Rename(*SubgraphName);
+		}
+	}
+}
+
+TArray<FExpressionOutput>& UMaterialExpressionComposite::GetOutputs()
+{	
+	Outputs.Reset();
+
+	// OutputExpressions may be null if we are using the default object
+	if (OutputExpressions)
+	{
+		for (const FCompositeReroute& ReroutePin : OutputExpressions->ReroutePins)
+		{
+			ReroutePin.Expression->GetOutputs()[0].OutputName = ReroutePin.Name;
+			Outputs.Add(ReroutePin.Expression->GetOutputs()[0]);
+		}
+	}
+	return Outputs;
+}
+
+const TArray<FExpressionInput*> UMaterialExpressionComposite::GetInputs()
+{
+	TArray<FExpressionInput*> ExpressionInputs;
+
+	// InputExpressions may be null if we are using the default object
+	if (InputExpressions)
+	{
+		for (const FCompositeReroute& ReroutePin : InputExpressions->ReroutePins)
+		{
+			ExpressionInputs.Add(ReroutePin.Expression->GetInput(0));
+		}
+	}
+	return ExpressionInputs;
+}
+
+FExpressionInput* UMaterialExpressionComposite::GetInput(int32 InputIndex)
+{
+	if (InputIndex >= 0 && InputIndex < InputExpressions->ReroutePins.Num())
+	{
+		return InputExpressions->ReroutePins[InputIndex].Expression->GetInput(0);
+	}
+
+	return nullptr;
+}
+
+FName UMaterialExpressionComposite::GetInputName(int32 InputIndex) const
+{
+	if (InputIndex >= 0 && InputIndex < InputExpressions->ReroutePins.Num())
+	{
+		return InputExpressions->ReroutePins[InputIndex].Name;
+	}
+
+	return FName();
+}
+
+uint32 UMaterialExpressionComposite::GetInputType(int32 InputIndex)
+{
+	if (InputIndex >= 0 && InputIndex < InputExpressions->ReroutePins.Num())
+	{
+		return InputExpressions->ReroutePins[InputIndex].Expression->GetInputType(0);
+	}
+
+	check(false);
+	return MCT_Float;
+}
+
+uint32 UMaterialExpressionComposite::GetOutputType(int32 OutputIndex)
+{
+	if (OutputIndex >= 0 && OutputIndex < OutputExpressions->ReroutePins.Num())
+	{
+		return OutputExpressions->ReroutePins[OutputIndex].Expression->GetOutputType(0);
+	}
+
+	check(false);
+	return MCT_Float;
+}
+
+bool UMaterialExpressionComposite::IsExpressionConnected(FExpressionInput* Input, int32 OutputIndex)
+{
+	if (Input && OutputIndex >= 0 && OutputIndex < OutputExpressions->ReroutePins.Num())
+	{
+		return OutputExpressions->ReroutePins[OutputIndex].Expression == Input->Expression;
+	}
+
+	return false;
+}
+
+void UMaterialExpressionComposite::ConnectExpression(FExpressionInput* Input, int32 OutputIndex)
+{
+	OutputExpressions->ReroutePins[OutputIndex].Expression->ConnectExpression(Input, 0);
+}
+
+void UMaterialExpressionComposite::GetCaption(TArray<FString>& OutCaptions) const
+{
+	OutCaptions.Add(SubgraphName);
+}
+
+bool UMaterialExpressionComposite::Modify(bool bAlwaysMarkDirty)
+{
+	// Modify pin bases so they can update the compilation graph
+	if (InputExpressions)
+	{
+		InputExpressions->Modify(bAlwaysMarkDirty);
+	}
+
+	if (OutputExpressions)
+	{
+		OutputExpressions->Modify(bAlwaysMarkDirty);
+	}
+
+	return Super::Modify(bAlwaysMarkDirty);
+}
+
+#endif // WITH_EDITOR
+
+///////////////////////////////////////////////////////////////////////////////
+// UMaterialExpressionPinBase
+///////////////////////////////////////////////////////////////////////////////
+UMaterialExpressionPinBase::UMaterialExpressionPinBase(const FObjectInitializer& ObjectInitializer)
+	: Super(ObjectInitializer)
+{
+	// Structure to hold one-time initialization
+	struct FConstructorStatics
+	{
+		FText NAME_Utility;
+		FConstructorStatics()
+			: NAME_Utility(LOCTEXT("Utility", "Utility"))
+		{
+		}
+	};
+	static FConstructorStatics ConstructorStatics;
+
+#if WITH_EDITORONLY_DATA
+	bShowOutputNameOnPin = true;
+
+	MenuCategories.Add(ConstructorStatics.NAME_Utility);
+	Outputs.Reset();
+#endif
+}
+
+#if WITH_EDITOR
+
+void UMaterialExpressionPinBase::DeleteReroutePins()
+{
+	Modify();
+	for (FCompositeReroute& Reroute : ReroutePins)
+	{
+		Reroute.Expression->Modify();
+		Material->Expressions.Remove(Reroute.Expression);
+		Reroute.Expression->MarkPendingKill();
+	}
+	ReroutePins.Empty();
+}
+
+void UMaterialExpressionPinBase::PreEditChange(FProperty* PropertyAboutToChange)
+{
+	Super::PreEditChange(PropertyAboutToChange);
+	
+	if (SubgraphExpression && SubgraphExpression->GraphNode)
+	{
+		SubgraphExpression->Modify();
+	}
+
+	PreEditRereouteExpresions.Empty();
+	for (const FCompositeReroute& Reroute : ReroutePins)
+	{
+		PreEditRereouteExpresions.Add(Reroute.Expression);
+	}
+}
+
+void UMaterialExpressionPinBase::PostEditChangeProperty(FPropertyChangedEvent& PropertyChangedEvent)
+{
+	if (PropertyChangedEvent.MemberProperty && GraphNode && SubgraphExpression->GraphNode)
+	{
+		Modify();
+		Material->Modify();
+
+		if (PropertyChangedEvent.ChangeType == EPropertyChangeType::ArrayAdd || PropertyChangedEvent.ChangeType == EPropertyChangeType::Duplicate)
+		{
+			uint32 AddedRerouteIndex = PropertyChangedEvent.GetArrayIndex(PropertyChangedEvent.Property->GetFName().ToString());
+			FCompositeReroute& AddedReroute = ReroutePins[AddedRerouteIndex];
+
+			AddedReroute.Expression = NewObject<UMaterialExpressionReroute>(GetOuter(), UMaterialExpressionReroute::StaticClass(), NAME_None, RF_Transactional);
+			AddedReroute.Expression->SubgraphExpression = SubgraphExpression;
+			
+			if (Material)
+			{
+				Material->Expressions.Add(AddedReroute.Expression);
+				AddedReroute.Expression->Material = Material;
+			}
+			else if (Function)
+			{
+				Function->Modify();
+				Function->FunctionExpressions.Add(AddedReroute.Expression);
+			}
+
+			AddedReroute.Name = AddedReroute.Name.IsNone() ? FName(FString::Printf(TEXT("Pin %u"), AddedRerouteIndex + 1)) : AddedReroute.Name;
+		}
+		else if (PropertyChangedEvent.ChangeType == EPropertyChangeType::ArrayRemove)
+		{
+			uint32 RemovedRerouteIndex = PropertyChangedEvent.GetArrayIndex(PropertyChangedEvent.Property->GetFName().ToString());
+			UMaterialExpression* RemovedReroute = PreEditRereouteExpresions[RemovedRerouteIndex];
+
+			RemovedReroute->Modify();
+			Material->Expressions.Remove(RemovedReroute);
+			RemovedReroute->MarkPendingKill();
+		}
+		else if (PropertyChangedEvent.ChangeType == EPropertyChangeType::ArrayClear)
+		{
+			for (UMaterialExpressionReroute* RemovedReroute : PreEditRereouteExpresions)
+			{
+				RemovedReroute->Modify();
+				Material->Expressions.Remove(RemovedReroute);
+				RemovedReroute->MarkPendingKill();
+			}
+		}
+
+		GraphNode->Modify();
+		GraphNode->BreakAllNodeLinks();
+		GraphNode->ReconstructNode();
+
+		SubgraphExpression->GraphNode->Modify();
+		SubgraphExpression->GraphNode->BreakAllNodeLinks();
+		SubgraphExpression->GraphNode->ReconstructNode();
+
+		Material->MaterialGraph->Modify();
+		Material->MaterialGraph->LinkGraphNodesFromMaterial();
+		Material->MaterialGraph->LinkMaterialExpressionsFromGraph();
+	}
+
+	PreEditRereouteExpresions.Empty();
+	Super::PostEditChangeProperty(PropertyChangedEvent);
+}
+
+TArray<FExpressionOutput>& UMaterialExpressionPinBase::GetOutputs()
+{
+	// Re-compute output expressions, since we can and do change via code.
+	if (PinDirection == EGPD_Output)
+	{
+		Outputs.Reset();
+		for (const FCompositeReroute& ReroutePin : ReroutePins)
+		{
+			ReroutePin.Expression->GetOutputs()[0].OutputName = ReroutePin.Name;
+			Outputs.Add(ReroutePin.Expression->GetOutputs()[0]);
+		}
+	}
+	return Outputs;
+}
+
+const TArray<FExpressionInput*> UMaterialExpressionPinBase::GetInputs()
+{
+	TArray<FExpressionInput*> ExpressionInputs;
+	if (PinDirection == EGPD_Input)
+	{
+		for (const FCompositeReroute& ReroutePin : ReroutePins)
+		{
+			ExpressionInputs.Add(ReroutePin.Expression->GetInput(0));
+		}
+	}
+	return ExpressionInputs;
+}
+
+FExpressionInput* UMaterialExpressionPinBase::GetInput(int32 InputIndex)
+{
+	if (PinDirection == EGPD_Input)
+	{
+		if (InputIndex >= 0 && InputIndex < ReroutePins.Num())
+		{
+			return ReroutePins[InputIndex].Expression->GetInput(0);
+		}
+	}
+
+	return nullptr;
+}
+
+FName UMaterialExpressionPinBase::GetInputName(int32 InputIndex) const
+{
+	if (PinDirection == EGPD_Input)
+	{
+		if (InputIndex >= 0 && InputIndex < ReroutePins.Num())
+		{
+			return ReroutePins[InputIndex].Name;
+		}
+	}
+
+	return FName();
+}
+
+uint32 UMaterialExpressionPinBase::GetInputType(int32 InputIndex)
+{
+	if (InputIndex >= 0 && InputIndex < ReroutePins.Num())
+	{
+		return PinDirection == EGPD_Input ? ReroutePins[InputIndex].Expression->GetInputType(0) : MCT_Float;
+	}
+
+	check(false);
+	return MCT_Float;
+}
+
+uint32 UMaterialExpressionPinBase::GetOutputType(int32 OutputIndex)
+{
+	if (OutputIndex >= 0 && OutputIndex < ReroutePins.Num())
+	{
+		return PinDirection == EGPD_Output ? ReroutePins[OutputIndex].Expression->GetOutputType(0) : MCT_Float;
+	}
+
+	check(false);
+	return MCT_Float;
+}
+
+bool UMaterialExpressionPinBase::IsExpressionConnected(FExpressionInput* Input, int32 OutputIndex)
+{
+	if (PinDirection == EGPD_Output && Input && OutputIndex >= 0 && OutputIndex < ReroutePins.Num())
+	{
+		return ReroutePins[OutputIndex].Expression == Input->Expression;
+	}
+
+	return false;
+}
+
+void UMaterialExpressionPinBase::ConnectExpression(FExpressionInput* Input, int32 OutputIndex)
+{
+	if (PinDirection == EGPD_Output)
+	{
+		ReroutePins[OutputIndex].Expression->ConnectExpression(Input, 0);
+	}
+}
+
+void UMaterialExpressionPinBase::GetCaption(TArray<FString>& OutCaptions) const
+{
+	OutCaptions.Add(PinDirection == EGPD_Output ? "Input" : "Output");
+}
+
+bool UMaterialExpressionPinBase::Modify(bool bAlwaysMarkDirty)
+{
+	// Modify reroute pins so they can update the compilation graph
+	for (const FCompositeReroute& ReroutePin : ReroutePins)
+	{
+		// Reroute pin can not have an expression if just adding new pin.
+		if (ReroutePin.Expression)
+		{
+			ReroutePin.Expression->Modify(bAlwaysMarkDirty);
+		}
+	}
+
+	return Super::Modify(bAlwaysMarkDirty);
+}
+
+#endif // WITH_EDITOR
+
 UMaterialExpressionFresnel::UMaterialExpressionFresnel(const FObjectInitializer& ObjectInitializer)
 	: Super(ObjectInitializer)
 {
diff --git a/Engine/Source/Runtime/Engine/Private/Materials/MaterialShared.cpp b/Engine/Source/Runtime/Engine/Private/Materials/MaterialShared.cpp
--- a/Engine/Source/Runtime/Engine/Private/Materials/MaterialShared.cpp
+++ b/Engine/Source/Runtime/Engine/Private/Materials/MaterialShared.cpp
@@ -255,17 +255,7 @@ int32 FExpressionInput::Compile(class FMaterialCompiler* Compiler)
 
 void FExpressionInput::Connect( int32 InOutputIndex, class UMaterialExpression* InExpression )
 {
-	OutputIndex = InOutputIndex;
-	Expression = InExpression;
-
-	TArray<FExpressionOutput> Outputs;
-	Outputs = Expression->GetOutputs();
-	FExpressionOutput* Output = &Outputs[ OutputIndex ];
-	Mask = Output->Mask;
-	MaskR = Output->MaskR;
-	MaskG = Output->MaskG;
-	MaskB = Output->MaskB;
-	MaskA = Output->MaskA;
+	InExpression->ConnectExpression(this, InOutputIndex);
 }
 #endif // WITH_EDITOR
 
