Fix missing typename for linux builds

diff --git a/Engine/Source/Editor/MaterialEditor/Private/MaterialEditor.cpp b/Engine/Source/Editor/MaterialEditor/Private/MaterialEditor.cpp
--- a/Engine/Source/Editor/MaterialEditor/Private/MaterialEditor.cpp
+++ b/Engine/Source/Editor/MaterialEditor/Private/MaterialEditor.cpp
@@ -6264,7 +6264,7 @@ void FMaterialEditor::DeepCopyExpressions(UMaterialGraph* CopyGraph, UMaterialEx
 	// Duplicate subnodes
 	auto DuplicateExpression = [&](auto* Expression)
 	{
-		using ExpressionType = std::remove_pointer<decltype(Expression)>::type;
+		using ExpressionType = typename std::remove_pointer<typename decltype(Expression)>::type;
 		return Cast<ExpressionType>(UMaterialEditingLibrary::DuplicateMaterialExpression(Material, MaterialFunction, Expression));
 	};
 
