Material lists now correctly display the reset to default button on the right.

SObjectPropertyEntryBox no longer has a CustomResetToDefault argument that was being ignored. Instead, all uses of it now directly use the FDetailWidgetRow OverrideResetToDefault().

diff --git a/Engine/Source/Editor/MaterialEditor/Private/MaterialEditorInstanceDetailCustomization.cpp b/Engine/Source/Editor/MaterialEditor/Private/MaterialEditorInstanceDetailCustomization.cpp
--- a/Engine/Source/Editor/MaterialEditor/Private/MaterialEditorInstanceDetailCustomization.cpp
+++ b/Engine/Source/Editor/MaterialEditor/Private/MaterialEditorInstanceDetailCustomization.cpp
@@ -524,6 +524,8 @@ void FMaterialInstanceParameterDetails::CreateParameterValueWidget(UDEditorParam
 				{
 					TWeakObjectPtr<UMaterialExpressionTextureSampleParameter> SamplerExpression = Expression;
 
+					PropertyRow.OverrideResetToDefault(ResetOverride);
+
 					PropertyRow.CustomWidget()
 					.NameContent()
 					[
@@ -535,7 +537,6 @@ void FMaterialInstanceParameterDetails::CreateParameterValueWidget(UDEditorParam
 						SNew(SObjectPropertyEntryBox)
 						.PropertyHandle(ParameterValueProperty)
 						.AllowedClass(UTexture::StaticClass())
-						.CustomResetToDefault(ResetOverride)
 						.ThumbnailPool(PropertyUtilities.Pin()->GetThumbnailPool())
 						.OnShouldFilterAsset_Lambda([SamplerExpression](const FAssetData& AssetData)
 						{
@@ -813,7 +814,6 @@ void FMaterialInstanceParameterDetails::CreateLabeledTextureParameterValueWidget
 							SNew(SObjectPropertyEntryBox)
 							.PropertyHandle(ParameterValueProperty)
 							.AllowedClass(UTexture::StaticClass())
-							.CustomResetToDefault(ResetOverride)
 							.ThumbnailPool(PropertyUtilities.Pin()->GetThumbnailPool())
 							.OnShouldFilterAsset_Lambda([SamplerExpression](const FAssetData& AssetData)
 							{
@@ -833,6 +833,8 @@ void FMaterialInstanceParameterDetails::CreateLabeledTextureParameterValueWidget
 							})
 						];
 
+					DetailWidgetRow.OverrideResetToDefault(ResetOverride);
+
 					static const FName Red("R");
 					static const FName Green("G");
 					static const FName Blue("B");
diff --git a/Engine/Source/Editor/MaterialEditor/Private/MaterialLayersFunctionsCustomization.cpp b/Engine/Source/Editor/MaterialEditor/Private/MaterialLayersFunctionsCustomization.cpp
--- a/Engine/Source/Editor/MaterialEditor/Private/MaterialLayersFunctionsCustomization.cpp
+++ b/Engine/Source/Editor/MaterialEditor/Private/MaterialLayersFunctionsCustomization.cpp
@@ -423,6 +423,7 @@ void FMaterialLayerFunctionElement::GenerateHeaderRowContent(FDetailWidgetRow& N
 	FunctionInfo.Association = InAssociation;
 
 	NodeRow
+		.OverrideResetToDefault(ResetAssetOverride)
 		.NameContent()
 		[
 			SNew(SHorizontalBox)
@@ -445,7 +446,6 @@ void FMaterialLayerFunctionElement::GenerateHeaderRowContent(FDetailWidgetRow& N
 				.ObjectPath(ParentCustomization, &FMaterialLayersFunctionsCustomization::GetLayerAssetPath, FunctionInfo)
 				.OnObjectChanged(AssetChanged)
 				.OnShouldFilterAsset(AssetFilter)
-				.CustomResetToDefault(ResetAssetOverride)
 				.ThumbnailPool(ParentCustomization->GetPropertyUtilities()->GetThumbnailPool())
 				.DisplayCompactSize(true)
 				.ThumbnailSizeOverride(ThumbnailOverride)
diff --git a/Engine/Source/Editor/MaterialEditor/Private/MaterialPropertyHelpers.cpp b/Engine/Source/Editor/MaterialEditor/Private/MaterialPropertyHelpers.cpp
--- a/Engine/Source/Editor/MaterialEditor/Private/MaterialPropertyHelpers.cpp
+++ b/Engine/Source/Editor/MaterialEditor/Private/MaterialPropertyHelpers.cpp
@@ -869,7 +869,7 @@ void FMaterialPropertyHelpers::ResetToDefault(TSharedPtr<IPropertyHandle> Proper
 	}
 }
 
-void FMaterialPropertyHelpers::ResetLayerAssetToDefault(TSharedPtr<IPropertyHandle> PropertyHandle,  class UDEditorParameterValue* InParameter, TEnumAsByte<EMaterialParameterAssociation> InAssociation, int32 Index, UMaterialEditorInstanceConstant* MaterialEditorInstance)
+void FMaterialPropertyHelpers::ResetLayerAssetToDefault(TSharedPtr<IPropertyHandle> InPropertyHandle, UDEditorParameterValue* InParameter, TEnumAsByte<EMaterialParameterAssociation> InAssociation, int32 Index, UMaterialEditorInstanceConstant* MaterialEditorInstance)
 {
 	
 	const FScopedTransaction Transaction(LOCTEXT("ResetToDefault", "Reset To Default"));
@@ -920,7 +920,7 @@ void FMaterialPropertyHelpers::ResetLayerAssetToDefault(TSharedPtr<IPropertyHand
 	
 }
 
-bool FMaterialPropertyHelpers::ShouldLayerAssetShowResetToDefault(TSharedPtr<IPropertyHandle> PropertyHandle, TSharedPtr<FSortedParamData> InParameterData, UMaterialInterface* InMaterial)
+bool FMaterialPropertyHelpers::ShouldLayerAssetShowResetToDefault(TSharedPtr<IPropertyHandle> InPropertyHandle, TSharedPtr<FSortedParamData> InParameterData, UMaterialInterface* InMaterial)
 {
 	if (!InParameterData->Parameter)
 	{
@@ -973,7 +973,7 @@ bool FMaterialPropertyHelpers::ShouldLayerAssetShowResetToDefault(TSharedPtr<IPr
 	return false;
 }
 
-bool FMaterialPropertyHelpers::ShouldShowResetToDefault(TSharedPtr<IPropertyHandle> PropertyHandle, UDEditorParameterValue* InParameter, UMaterialEditorInstanceConstant* MaterialEditorInstance)
+bool FMaterialPropertyHelpers::ShouldShowResetToDefault(TSharedPtr<IPropertyHandle> InPropertyHandle, UDEditorParameterValue* InParameter, UMaterialEditorInstanceConstant* MaterialEditorInstance)
 {
 	const FMaterialParameterInfo& ParameterInfo = InParameter->ParameterInfo;
 
@@ -1298,7 +1298,7 @@ void FMaterialPropertyHelpers::SetPositionFromCurveAsset(const FAssetData& Asset
 	}
 }
 
-void FMaterialPropertyHelpers::ResetCurveToDefault(TSharedPtr<IPropertyHandle> PropertyHandle, class UDEditorParameterValue* Parameter, UMaterialEditorInstanceConstant* MaterialEditorInstance)
+void FMaterialPropertyHelpers::ResetCurveToDefault(TSharedPtr<IPropertyHandle> PropertyHandle, UDEditorParameterValue* Parameter, UMaterialEditorInstanceConstant* MaterialEditorInstance)
 {
 	const FScopedTransaction Transaction(LOCTEXT("ResetToDefault", "Reset To Default"));
 	Parameter->Modify();
diff --git a/Engine/Source/Editor/MaterialEditor/Private/SMaterialLayersFunctionsTree.cpp b/Engine/Source/Editor/MaterialEditor/Private/SMaterialLayersFunctionsTree.cpp
--- a/Engine/Source/Editor/MaterialEditor/Private/SMaterialLayersFunctionsTree.cpp
+++ b/Engine/Source/Editor/MaterialEditor/Private/SMaterialLayersFunctionsTree.cpp
@@ -34,14 +34,13 @@
 #include "Widgets/Text/SInlineEditableTextBlock.h"
 #include "Materials/MaterialFunctionInstance.h"
 #include "Framework/Application/SlateApplication.h"
-
 #include "Widgets/Input/SEditableTextBox.h"
 #include "Curves/CurveLinearColor.h"
 #include "Curves/CurveLinearColorAtlas.h"
 #include "Subsystems/AssetEditorSubsystem.h"
 #include "Editor.h"
 #include "EditorFontGlyphs.h"
-
+#include "SResetToDefaultPropertyEditor.h"
 
 #define LOCTEXT_NAMESPACE "MaterialLayerCustomization"
 
@@ -230,6 +229,7 @@ void SMaterialLayersFunctionsInstanceTreeItem::Construct(const FArguments& InArg
 
 	TSharedRef<SWidget> LeftSideWidget = SNullWidget::NullWidget;
 	TSharedRef<SWidget> RightSideWidget = SNullWidget::NullWidget;
+	TSharedRef<SWidget> ResetWidget = SNullWidget::NullWidget;
 	FText NameOverride;
 	TSharedRef<SVerticalBox> WrapperWidget = SNew(SVerticalBox);
 
@@ -427,6 +427,11 @@ void SMaterialLayersFunctionsInstanceTreeItem::Construct(const FArguments& InArg
 		FResetToDefaultOverride ResetAssetOverride = FResetToDefaultOverride::Create(IsAssetResetVisible, ResetAssetHandler);
 
 		IDetailTreeNode& Node = *StackParameterData->ParameterNode;
+
+		ResetWidget = SNew(SResetToDefaultPropertyEditor, Node.CreatePropertyHandle())
+			.IsEnabled(IsParamEnabled)
+			.CustomResetToDefault(ResetAssetOverride);
+
 		FNodeWidgets NodeWidgets = Node.CreateNodeWidgets();
 
 		LeftSideWidget = StackParameterData->ParameterHandle->CreatePropertyNameWidget(NameOverride);
@@ -482,7 +487,6 @@ void SMaterialLayersFunctionsInstanceTreeItem::Construct(const FArguments& InArg
 					.ObjectPath(this, &SMaterialLayersFunctionsInstanceTreeItem::GetInstancePath, Tree)
 					.OnShouldFilterAsset(AssetFilter)
 					.OnObjectChanged(AssetChanged)
-					.CustomResetToDefault(ResetAssetOverride)
 					.DisplayCompactSize(true)
 					.NewAssetFactories(FMaterialPropertyHelpers::GetAssetFactories(InAssociation))
 				]
@@ -669,6 +673,9 @@ void SMaterialLayersFunctionsInstanceTreeItem::Construct(const FArguments& InArg
 					TSharedPtr<SWidget> ValueWidget;
 					TSharedPtr<SVerticalBox> NameVerticalBox;
 					const FText ParameterName = FText::FromName(StackParameterData->Parameter->ParameterInfo.Name);
+
+					Row.OverrideResetToDefault(ResetOverride);
+
 					FDetailWidgetRow& CustomWidget = Row.CustomWidget();
 					CustomWidget
 						.FilterString(ParameterName)
@@ -689,7 +696,6 @@ void SMaterialLayersFunctionsInstanceTreeItem::Construct(const FArguments& InArg
 							SNew(SObjectPropertyEntryBox)
 							.PropertyHandle(StackParameterData->ParameterNode->CreatePropertyHandle())
 							.AllowedClass(UTexture::StaticClass())
-							.CustomResetToDefault(ResetOverride)
 							.ThumbnailPool(Tree->GetTreeThumbnailPool())
 							.OnShouldFilterAsset_Lambda([SamplerExpression](const FAssetData& AssetData)
 							{
@@ -998,6 +1004,13 @@ void SMaterialLayersFunctionsInstanceTreeItem::Construct(const FArguments& InArg
 			break;
 		}
 
+		if (ResetWidget == SNullWidget::NullWidget)
+		{
+			const FSlateBrush* DiffersFromDefaultBrush = FEditorStyle::GetBrush("PropertyWindow.DiffersFromDefault");
+			ResetWidget = SNew(SSpacer)
+				.Size(DiffersFromDefaultBrush != nullptr ? DiffersFromDefaultBrush->ImageSize : FVector2D(8.0f, 8.0f));
+		}
+
 		WrapperWidget->AddSlot()
 			.AutoHeight()
 			[
@@ -1009,6 +1022,7 @@ void SMaterialLayersFunctionsInstanceTreeItem::Construct(const FArguments& InArg
 					.Style(FEditorStyle::Get(), "DetailsView.Splitter")
 					.PhysicalSplitterHandleSize(1.0f)
 					.HitDetectionSplitterHandleSize(5.0f)
+					.HighlightedHandleIndex(Tree->ColumnSizeData.HoveredSplitterIndex)
 					+ SSplitter::Slot()
 					.Value(Tree->ColumnSizeData.NameColumnWidth)
 					.OnSlotResized(Tree->ColumnSizeData.OnNameColumnResized)
@@ -1041,6 +1055,18 @@ void SMaterialLayersFunctionsInstanceTreeItem::Construct(const FArguments& InArg
 							RightSideWidget
 						]
 					]
+					+ SSplitter::Slot()
+					.SizeRule(SSplitter::ESizeRule::SizeToContent)
+					[
+						SNew(SHorizontalBox)
+						+ SHorizontalBox::Slot()
+						.VAlign(VAlign_Center)
+						.HAlign(HAlign_Center)
+						.Padding(5.0f)
+						[
+							ResetWidget
+						]
+					]
 				]
 			];
 	}
diff --git a/Engine/Source/Editor/Persona/Private/PreviewSceneCustomizations.cpp b/Engine/Source/Editor/Persona/Private/PreviewSceneCustomizations.cpp
--- a/Engine/Source/Editor/Persona/Private/PreviewSceneCustomizations.cpp
+++ b/Engine/Source/Editor/Persona/Private/PreviewSceneCustomizations.cpp
@@ -176,6 +176,23 @@ void FPreviewSceneDescriptionCustomization::CustomizeDetails(IDetailLayoutBuilde
 		DetailBuilder.EditCategory("Mesh")
 		.AddProperty(SkeletalMeshProperty)
 		.CustomWidget()
+		.OverrideResetToDefault(FResetToDefaultOverride::Create(
+			FIsResetToDefaultVisible::CreateLambda([this](TSharedPtr<IPropertyHandle> PropertyHandle) -> bool 
+			{
+				if (PreviewScene.IsValid())
+				{
+					return PreviewScene.Pin()->GetPreviewMesh() != nullptr;
+				}
+				return false;
+			}),
+			FResetToDefaultHandler::CreateLambda([this](TSharedPtr<IPropertyHandle> PropertyHandle) 
+			{
+				if (PreviewScene.IsValid())
+				{
+					PreviewScene.Pin()->SetPreviewMesh(nullptr, false);
+				}
+			}))
+		)
 		.NameContent()
 		[
 			SNew(SVerticalBox)
@@ -215,21 +232,6 @@ void FPreviewSceneDescriptionCustomization::CustomizeDetails(IDetailLayoutBuilde
 			.OnShouldFilterAsset(this, &FPreviewSceneDescriptionCustomization::HandleShouldFilterAsset, FName("Skeleton"), PersonaToolkit.Pin()->GetContext() == UPhysicsAsset::StaticClass()->GetFName())
 			.OnObjectChanged(this, &FPreviewSceneDescriptionCustomization::HandleMeshChanged)
 			.ThumbnailPool(DetailBuilder.GetThumbnailPool())
-			.CustomResetToDefault(FResetToDefaultOverride::Create(
-				FIsResetToDefaultVisible::CreateLambda([this](TSharedPtr<IPropertyHandle> PropertyHandle) -> bool {
-					if (PreviewScene.IsValid())
-					{
-						return PreviewScene.Pin()->GetPreviewMesh() != nullptr;
-					}
-					return false;
-				}),
-				FResetToDefaultHandler::CreateLambda([this](TSharedPtr<IPropertyHandle> PropertyHandle) {
-					if (PreviewScene.IsValid())
-					{
-						PreviewScene.Pin()->SetPreviewMesh(nullptr, false);
-					}
-				})
-			))
 		];
 
 		// Customize animation blueprint preview
@@ -395,6 +397,7 @@ void FPreviewSceneDescriptionCustomization::CustomizeDetails(IDetailLayoutBuilde
 		DetailBuilder.EditCategory("Additional Meshes")
 		.AddProperty(AdditionalMeshesProperty)
 		.CustomWidget()
+		.OverrideResetToDefault(ResetToDefaultOverride)
 		.NameContent()
 		[
 			AdditionalMeshesProperty->CreatePropertyNameWidget()
@@ -414,7 +417,6 @@ void FPreviewSceneDescriptionCustomization::CustomizeDetails(IDetailLayoutBuilde
 				.PropertyHandle(AdditionalMeshesProperty)
 				.OnShouldFilterAsset(this, &FPreviewSceneDescriptionCustomization::HandleShouldFilterAdditionalMesh, true)
 				.OnObjectChanged(this, &FPreviewSceneDescriptionCustomization::HandleAdditionalMeshesChanged, &DetailBuilder)
-				.CustomResetToDefault(ResetToDefaultOverride)
 				.ThumbnailPool(DetailBuilder.GetThumbnailPool())
 				.NewAssetFactories(FactoriesToUse)
 			]
diff --git a/Engine/Source/Editor/PropertyEditor/Private/MaterialList.cpp b/Engine/Source/Editor/PropertyEditor/Private/MaterialList.cpp
--- a/Engine/Source/Editor/PropertyEditor/Private/MaterialList.cpp
+++ b/Engine/Source/Editor/PropertyEditor/Private/MaterialList.cpp
@@ -144,11 +144,6 @@ public:
 	{
 		FIntPoint ThumbnailSize(64, 64);
 
-		FResetToDefaultOverride ResetToDefaultOverride = FResetToDefaultOverride::Create(
-			FIsResetToDefaultVisible::CreateSP(this, &FMaterialItemView::GetReplaceVisibility),
-			FResetToDefaultHandler::CreateSP(this, &FMaterialItemView::OnResetToBaseClicked)
-		);
-
 		return
 			SNew(SHorizontalBox)
 			+ SHorizontalBox::Slot()
@@ -171,7 +166,6 @@ public:
 						.OnObjectChanged(this, &FMaterialItemView::OnSetObject)
 						.ThumbnailPool(ThumbnailPool)
 						.DisplayCompactSize(bDisplayCompactSize)
-						.CustomResetToDefault(ResetToDefaultOverride)
 						.OwnerAssetDataArray(OwnerAssetDataArray)
 						.CustomContentSlot()
 						[
@@ -220,6 +214,29 @@ public:
 			];
 	}
 
+	/**
+	 * Called to get the visibility of the reset to base button
+	 */
+	bool GetResetToBaseVisibility() const
+	{
+		// Only show the reset to base button if the current material can be replaced
+		return OnMaterialChanged.IsBound() && MaterialItem.bCanBeReplaced;
+	}
+
+	/**
+	 * Called when reset to base is clicked
+	 */
+	void OnResetToBaseClicked()
+	{
+		// Only allow reset to base if the current material can be replaced
+		if( MaterialItem.Material.IsValid() && MaterialItem.bCanBeReplaced )
+		{
+			bool bReplaceAll = false;
+			ReplaceMaterial( nullptr, bReplaceAll );
+			OnResetToDefaultClicked.ExecuteIfBound( MaterialItem.Material.Get(), MaterialItem.SlotIndex );
+		}
+	}
+
 private:
 
 	FMaterialItemView(	const FMaterialListItem& InMaterial, 
@@ -316,34 +333,6 @@ private:
 		}
 	}
 
-	/**
-	 * Called to get the visibility of the replace button
-	 */
-	bool GetReplaceVisibility(TSharedPtr<IPropertyHandle> PropertyHandle) const
-	{
-		// Only show the replace button if the current material can be replaced
-		if (OnMaterialChanged.IsBound() && MaterialItem.bCanBeReplaced)
-		{
-			return true;
-		}
-
-		return false;
-	}
-
-	/**
-	 * Called when reset to base is clicked
-	 */
-	void OnResetToBaseClicked(TSharedPtr<IPropertyHandle> PropertyHandle)
-	{
-		// Only allow reset to base if the current material can be replaced
-		if( MaterialItem.Material.IsValid() && MaterialItem.bCanBeReplaced )
-		{
-			bool bReplaceAll = false;
-			ReplaceMaterial( NULL, bReplaceAll );
-			OnResetToDefaultClicked.ExecuteIfBound( MaterialItem.Material.Get(), MaterialItem.SlotIndex );
-		}
-	}
-
 private:
 	FMaterialListItem MaterialItem;
 	FOnMaterialChanged OnMaterialChanged;
@@ -625,6 +614,13 @@ void FMaterialList::AddMaterialItem( FDetailWidgetRow& Row, int32 CurrentSlot, c
 		ViewedMaterials.Add( NewView );
 	}
 
+	FResetToDefaultOverride ResetToDefaultOverride = FResetToDefaultOverride::Create(
+		TAttribute<bool>(NewView, &FMaterialItemView::GetResetToBaseVisibility),
+		FSimpleDelegate::CreateSP(NewView, &FMaterialItemView::OnResetToBaseClicked)
+	);
+
+	Row.OverrideResetToDefault(ResetToDefaultOverride);
+
 	Row.CopyAction(FUIAction(FExecuteAction::CreateSP(this, &FMaterialList::OnCopyMaterialItem, Item.SlotIndex), FCanExecuteAction::CreateSP(this, &FMaterialList::OnCanCopyMaterialItem, Item.SlotIndex)));
 	Row.PasteAction(FUIAction(FExecuteAction::CreateSP(this, &FMaterialList::OnPasteMaterialItem, Item.SlotIndex)));
 
diff --git a/Engine/Source/Editor/PropertyEditor/Public/PropertyCustomizationHelpers.h b/Engine/Source/Editor/PropertyEditor/Public/PropertyCustomizationHelpers.h
--- a/Engine/Source/Editor/PropertyEditor/Public/PropertyCustomizationHelpers.h
+++ b/Engine/Source/Editor/PropertyEditor/Public/PropertyCustomizationHelpers.h
@@ -202,8 +202,6 @@ public:
 		SLATE_ARGUMENT(bool, DisplayBrowse)
 		/** Whether to enable the content Picker */
 		SLATE_ARGUMENT(bool, EnableContentPicker)
-		/** A custom reset to default override */
-		SLATE_ARGUMENT(TOptional<FResetToDefaultOverride>, CustomResetToDefault)
 		/** Whether or not to display a smaller, compact size for the asset thumbnail */ 
 		SLATE_ARGUMENT(bool, DisplayCompactSize)
 		/** Whether or not to display the asset thumbnail */ 
