PR #7226: Add backgrounds and panning to UI material previews (Contributed by [REDACTED_PERSON])

Unshelved from pending changelist '14008291':

diff --git a/Engine/Source/Editor/EditorStyle/Private/SlateEditorStyle.cpp b/Engine/Source/Editor/EditorStyle/Private/SlateEditorStyle.cpp
--- a/Engine/Source/Editor/EditorStyle/Private/SlateEditorStyle.cpp
+++ b/Engine/Source/Editor/EditorStyle/Private/SlateEditorStyle.cpp
@@ -7174,6 +7174,7 @@ void FSlateEditorStyle::FStyle::SetupToolkitStyles()
 		Set( "MaterialEditor.CameraHome.Small", new IMAGE_BRUSH( "Icons/icon_MatEd_Home_40x", Icon20x20 ) );
 		Set( "MaterialEditor.FindInMaterial", new IMAGE_BRUSH( "Icons/icon_Blueprint_Find_40px", Icon40x40 ) );
 		Set( "MaterialEditor.FindInMaterial.Small", new IMAGE_BRUSH( "Icons/icon_Blueprint_Find_40px", Icon20x20 ) );
+		Set( "MaterialEditor.ViewportBorder", new BORDER_BRUSH("Old/White", FMargin(1), FLinearColor::White ) );
 
 		Set("MaterialEditor.CellListViewRow", FTableRowStyle(NormalTableRowStyle)
 			.SetEvenRowBackgroundBrush(IMAGE_BRUSH("PropertyView/DetailCategoryMiddle", FVector2D(16, 16), FLinearColor(0.5f, 0.5f, 0.5f)))
diff --git a/Engine/Source/Editor/MaterialEditor/Private/SMaterialEditorViewport.cpp b/Engine/Source/Editor/MaterialEditor/Private/SMaterialEditorViewport.cpp
--- a/Engine/Source/Editor/MaterialEditor/Private/SMaterialEditorViewport.cpp
+++ b/Engine/Source/Editor/MaterialEditor/Private/SMaterialEditorViewport.cpp
@@ -3,6 +3,7 @@
 #include "SMaterialEditorViewport.h"
 #include "Widgets/SBoxPanel.h"
 #include "Widgets/SViewport.h"
+#include "Widgets/SToolTip.h"
 #include "Widgets/Text/STextBlock.h"
 #include "Widgets/Layout/SBorder.h"
 #include "Widgets/Images/SImage.h"
@@ -24,12 +25,18 @@
 #include "Widgets/Docking/SDockTab.h"
 #include "Engine/TextureCube.h"
 #include "ComponentAssetBroker.h"
+#include "Modules/ModuleManager.h"
 #include "SlateMaterialBrush.h"
 #include "Widgets/Input/SNumericEntryBox.h"
 #include "AdvancedPreviewScene.h"
 #include "AssetViewerSettings.h"
 #include "Engine/PostProcessVolume.h"
 #include "MaterialEditorSettings.h"
+#include "Brushes/SlateImageBrush.h"
+#include "Brushes/SlateBorderBrush.h"
+#include "ImageUtils.h"
+#include "ISettingsModule.h"
+#include "Framework/Layout/ScrollyZoomy.h"
 
 #define LOCTEXT_NAMESPACE "MaterialEditor"
 #include "UnrealWidget.h"
@@ -765,7 +772,7 @@ void SMaterialEditor3DPreviewViewport::OnPropertyChanged(UObject* ObjectBeingMod
 	}
 }
 
-class SMaterialEditorUIPreviewZoomer : public SPanel
+class SMaterialEditorUIPreviewZoomer : public SPanel, public IScrollableZoomable, public FGCObject
 {
 public:
 
@@ -778,11 +785,24 @@ public:
 		}
 	};
 
-	SLATE_BEGIN_ARGS(SMaterialEditorUIPreviewViewport){}
+	SLATE_BEGIN_ARGS(SMaterialEditorUIPreviewZoomer)
+		: _OnContextMenuRequested()
+		, _OnZoomed()
+		, _InitialPreviewSize(FVector2D(250.f))
+		, _BackgroundSettings()
+		{}
+		SLATE_EVENT(FNoReplyPointerEventHandler, OnContextMenuRequested)
+		SLATE_EVENT(FSimpleDelegate, OnZoomed)
+		SLATE_ARGUMENT(FVector2D, InitialPreviewSize)
+		SLATE_ARGUMENT(FPreviewBackgroundSettings, BackgroundSettings)
 	SLATE_END_ARGS()
 
 	SMaterialEditorUIPreviewZoomer()
-		: ChildSlot(this)
+		: CachedSize(ForceInitToZero)
+		, PhysicalOffset(ForceInitToZero)
+		, ScrollyZoomy(/* bUseInertialScrolling */ false)
+		, ChildSlot(this)
+		, CheckerboardTexture(nullptr)
 	{
 	}
 
@@ -791,35 +811,138 @@ public:
 	virtual void OnArrangeChildren( const FGeometry& AllottedGeometry, FArrangedChildren& ArrangedChildren ) const override;
 	virtual FVector2D ComputeDesiredSize(float) const override;
 	virtual FChildren* GetChildren() override;
-	virtual FReply OnMouseWheel( const FGeometry& MyGeometry, const FPointerEvent& MouseEvent ) override;
 	virtual int32 OnPaint(const FPaintArgs& Args, const FGeometry& AllottedGeometry, const FSlateRect& MyCullingRect, FSlateWindowElementList& OutDrawElements, int32 LayerId, const FWidgetStyle& InWidgetStyle, bool bParentEnabled) const override;
 
-	bool ZoomBy( const float Amount );
+	/** Begin SWidget Interface */
+	virtual void Tick(const FGeometry& AllottedGeometry, const double InCurrentTime, const float InDeltaTime) override;
+	virtual FReply OnMouseButtonDown(const FGeometry& MyGeometry, const FPointerEvent& MouseEvent) override;
+	virtual FReply OnMouseButtonUp(const FGeometry& MyGeometry, const FPointerEvent& MouseEvent) override;
+	virtual FReply OnMouseMove(const FGeometry& MyGeometry, const FPointerEvent& MouseEvent) override;
+	virtual void OnMouseLeave(const FPointerEvent& MouseEvent) override;
+	virtual FReply OnMouseWheel(const FGeometry& MyGeometry, const FPointerEvent& MouseEvent) override;
+	virtual FCursorReply OnCursorQuery(const FGeometry& MyGeometry, const FPointerEvent& CursorEvent) const override;
+	/** End SWidget Interface */
+
+	/** Begin IScrollableZoomable Interface */
+	virtual bool ScrollBy(const FVector2D& Offset) override;
+	/** End IScrollableZoomable Interface */
+
+	bool ZoomBy(const float Amount);
 	float GetZoomLevel() const;
+	bool SetZoomLevel(float Level);
+	FVector2D ComputeZoomedPreviewSize() const;
+
+	FReply HandleScrollEvent(const FPointerEvent& MouseEvent);
+	bool IsCurrentlyScrollable() const;
+	void ScrollToCenter();
+	bool IsCentered() const;
 
 	void SetPreviewSize( const FVector2D PreviewSize );
 	void SetPreviewMaterial(UMaterialInterface* InPreviewMaterial);
+
+	void SetBackgroundSettings(const FPreviewBackgroundSettings& NewSettings);
+	FSlateColor GetBorderColor() const;
+	FMargin GetBorderPadding() const;
+
+	void AddReferencedObjects(FReferenceCollector& Collector);
+
 private:
+
+	FSimpleDelegate OnZoomed;
+	FNoReplyPointerEventHandler OnContextMenuRequested;
+
+	void ModifyCheckerboardTextureColors(const FCheckerboardSettings& Checkerboard);
+	void SetupCheckerboardTexture(const FColor& ColorOne, const FColor& ColorTwo, int32 CheckerSize);
+	void DestroyCheckerboardTexture();
+
+	FSlateColor GetSolidBackgroundColor() const;
+	EVisibility GetVisibilityForBackgroundType(EBackgroundType BackgroundType) const;
+
+	void ClampViewOffset(const FVector2D& ZoomedPreviewSize, const FVector2D& LocalSize);
+	float ClampViewOffsetAxis(const float ZoomedPreviewSize, const float LocalSize, const float CurrentOffset);
+
 	mutable FVector2D CachedSize;
 	float ZoomLevel;
+	FVector2D PhysicalOffset;
+	FScrollyZoomy ScrollyZoomy;
+	bool bCenterInFrame;
 
 	FMaterialPreviewPanelSlot ChildSlot;
+
 	TSharedPtr<FSlateMaterialBrush> PreviewBrush;
+	TSharedPtr<FSlateImageBrush> CheckerboardBrush;
+	const FSlateBrush* BorderBrush;
+	UTexture2D* CheckerboardTexture;
 	TSharedPtr<SImage> ImageWidget;
+	FPreviewBackgroundSettings BackgroundSettings;
 };
 
 
 void SMaterialEditorUIPreviewZoomer::Construct( const FArguments& InArgs, UMaterialInterface* InPreviewMaterial )
 {
-	PreviewBrush = MakeShareable( new FSlateMaterialBrush( *InPreviewMaterial, GetDefault<UMaterialEditorSettings>()->GetPreviewViewportStartingSize() ) );
+	OnContextMenuRequested = InArgs._OnContextMenuRequested;
+	OnZoomed = InArgs._OnZoomed;
+	
+	ZoomLevel = 1.0f;
+	bCenterInFrame = true;
+	
+	BackgroundSettings = InArgs._BackgroundSettings;
+	ModifyCheckerboardTextureColors(BackgroundSettings.Checkerboard);
+
+	PreviewBrush = MakeShared<FSlateMaterialBrush>(*InPreviewMaterial, InArgs._InitialPreviewSize);
+	BorderBrush = FEditorStyle::GetBrush("MaterialEditor.ViewportBorder");	
 
 	ChildSlot
 	[
-		SAssignNew( ImageWidget, SImage )
-		.Image( PreviewBrush.Get() )
+		SNew(SBorder)
+		.BorderImage(FCoreStyle::Get().GetBrush("WhiteBrush"))
+		.BorderBackgroundColor(this, &SMaterialEditorUIPreviewZoomer::GetBorderColor)
+		.Padding(this, &SMaterialEditorUIPreviewZoomer::GetBorderPadding) // Leave space for our border (drawn in OnPaint) to remain visible when scrolled to the edges
+		[
+			SNew(SBorder)
+			.BorderImage(FCoreStyle::Get().GetBrush("WhiteBrush"))
+			.BorderBackgroundColor(this, &SMaterialEditorUIPreviewZoomer::GetSolidBackgroundColor)
+			.Padding(0)
+			[
+				SNew(SOverlay)
+				+ SOverlay::Slot()
+				.HAlign(HAlign_Fill)
+				.VAlign(VAlign_Fill)
+				[
+					SNew(SImage)
+					.Image(CheckerboardBrush.Get())
+					.Visibility(this, &SMaterialEditorUIPreviewZoomer::GetVisibilityForBackgroundType, EBackgroundType::Checkered)
+				]
+				+ SOverlay::Slot()
+				.HAlign(HAlign_Fill)
+				.VAlign(VAlign_Fill)
+				[
+					SAssignNew(ImageWidget, SImage)
+					.Image(PreviewBrush.Get())
+				]
+			]
+		]
 	];
+}
 
-	ZoomLevel = 1.0f;
+FSlateColor SMaterialEditorUIPreviewZoomer::GetBorderColor() const
+{
+	return FLinearColor(BackgroundSettings.BorderColor);
+}
+
+FMargin SMaterialEditorUIPreviewZoomer::GetBorderPadding() const
+{
+	return FMargin(BackgroundSettings.bShowBorder ? 1.f : 0.f);
+}
+
+FSlateColor SMaterialEditorUIPreviewZoomer::GetSolidBackgroundColor() const
+{
+	return FLinearColor(BackgroundSettings.BackgroundColor);
+}
+
+EVisibility SMaterialEditorUIPreviewZoomer::GetVisibilityForBackgroundType(EBackgroundType BackgroundType) const
+{
+	return BackgroundSettings.BackgroundType == BackgroundType ? EVisibility::HitTestInvisible : EVisibility::Collapsed;
 }
 
 void SMaterialEditorUIPreviewZoomer::OnArrangeChildren( const FGeometry& AllottedGeometry, FArrangedChildren& ArrangedChildren ) const 
@@ -829,9 +952,21 @@ void SMaterialEditorUIPreviewZoomer::OnArrangeChildren( const FGeometry& Allotte
 	const TSharedRef<SWidget>& ChildWidget = ChildSlot.GetWidget();
 	if( ChildWidget->GetVisibility() != EVisibility::Collapsed )
 	{
-		const FVector2D& WidgetDesiredSize = ChildWidget->GetDesiredSize();
+		SMaterialEditorUIPreviewZoomer* const MutableThis = const_cast<SMaterialEditorUIPreviewZoomer*>(this);
+
+		FVector2D SizeWithBorder = GetDesiredSize();
+
+		// Ensure we're centered within our current geometry
+		if (bCenterInFrame)
+		{
+			MutableThis->PhysicalOffset = ((CachedSize - SizeWithBorder) * 0.5f).RoundToVector();
+		}
 
-		ArrangedChildren.AddWidget(AllottedGeometry.MakeChild(ChildWidget, FVector2D::ZeroVector, WidgetDesiredSize * ZoomLevel));
+		// Re-clamp since our parent might have changed size
+		MutableThis->ClampViewOffset(SizeWithBorder, CachedSize);
+
+		// Round so that we get a crisp checkerboard at all zoom levels
+		ArrangedChildren.AddWidget(AllottedGeometry.MakeChild(ChildWidget, PhysicalOffset, SizeWithBorder));
 	}
 }
 
@@ -842,38 +977,45 @@ FVector2D SMaterialEditorUIPreviewZoomer::ComputeDesiredSize(float) const
 	const TSharedRef<SWidget>& ChildWidget = ChildSlot.GetWidget();
 	if( ChildWidget->GetVisibility() != EVisibility::Collapsed )
 	{
-		ThisDesiredSize = ChildWidget->GetDesiredSize() * ZoomLevel;
+		ThisDesiredSize = ComputeZoomedPreviewSize() + GetBorderPadding().GetDesiredSize();
 	}
 
 	return ThisDesiredSize;
 }
 
-FChildren* SMaterialEditorUIPreviewZoomer::GetChildren()
+FVector2D SMaterialEditorUIPreviewZoomer::ComputeZoomedPreviewSize() const
 {
-	return &ChildSlot;
+	// Our desired size includes the 1px border (if enabled), but this is purely the size of the actual preview quad
+	return (ImageWidget->GetDesiredSize() * ZoomLevel).RoundToVector();
 }
 
-FReply SMaterialEditorUIPreviewZoomer::OnMouseWheel( const FGeometry& MyGeometry, const FPointerEvent& MouseEvent )
+FChildren* SMaterialEditorUIPreviewZoomer::GetChildren()
 {
-	 ZoomBy(MouseEvent.GetWheelDelta());
-
-	 return FReply::Handled();
+	return &ChildSlot;
 }
 
 int32 SMaterialEditorUIPreviewZoomer::OnPaint(const FPaintArgs& Args, const FGeometry& AllottedGeometry, const FSlateRect& MyCullingRect, FSlateWindowElementList& OutDrawElements, int32 LayerId, const FWidgetStyle& InWidgetStyle, bool bParentEnabled) const
 {
 	LayerId = SPanel::OnPaint(Args, AllottedGeometry, MyCullingRect, OutDrawElements, LayerId, InWidgetStyle, bParentEnabled);
+	
+	if (IsCurrentlyScrollable())
+	{
+		LayerId = ScrollyZoomy.PaintSoftwareCursorIfNeeded(AllottedGeometry, MyCullingRect, OutDrawElements, LayerId);
+	}
+
 	return LayerId;
 }
 
-bool SMaterialEditorUIPreviewZoomer::ZoomBy( const float Amount )
+bool SMaterialEditorUIPreviewZoomer::IsCurrentlyScrollable() const
 {
-	static const float MinZoomLevel = 0.2f;
-	static const float MaxZoomLevel = 4.0f;
+	FVector2D ContentSize = GetDesiredSize();
+	bool bCanScroll = ContentSize.X > CachedSize.X || ContentSize.Y > CachedSize.Y;
+	return bCanScroll;
+}
 
-	const float PrevZoomLevel = ZoomLevel;
-	ZoomLevel = FMath::Clamp(ZoomLevel + (Amount * 0.05f), MinZoomLevel, MaxZoomLevel);
-	return ZoomLevel != PrevZoomLevel;
+bool SMaterialEditorUIPreviewZoomer::ZoomBy( const float Amount )
+{
+	return SetZoomLevel(ZoomLevel + (Amount * 0.05f));
 }
 
 float SMaterialEditorUIPreviewZoomer::GetZoomLevel() const
@@ -881,6 +1023,21 @@ float SMaterialEditorUIPreviewZoomer::GetZoomLevel() const
 	return ZoomLevel;
 }
 
+bool SMaterialEditorUIPreviewZoomer::SetZoomLevel(float NewLevel)
+{
+	static const float MinZoomLevel = 0.2f;
+	static const float MaxZoomLevel = 4.0f;
+
+	const float PrevZoomLevel = ZoomLevel;
+	ZoomLevel = FMath::Clamp(NewLevel, MinZoomLevel, MaxZoomLevel);
+
+	// Fire regardless of whether it actually changed, since still useful to give
+	// the user feedback feedback when attempting to zoom past the limit
+	OnZoomed.ExecuteIfBound();
+
+	const bool bZoomChanged = ZoomLevel != PrevZoomLevel;
+	return bZoomChanged;
+}
 
 void SMaterialEditorUIPreviewZoomer::SetPreviewSize( const FVector2D PreviewSize )
 {
@@ -889,12 +1046,201 @@ void SMaterialEditorUIPreviewZoomer::SetPreviewSize( const FVector2D PreviewSize
 
 void SMaterialEditorUIPreviewZoomer::SetPreviewMaterial(UMaterialInterface* InPreviewMaterial)
 {
-	PreviewBrush = MakeShareable(new FSlateMaterialBrush(*InPreviewMaterial, PreviewBrush->ImageSize));
+	// Just create a new brush to avoid possible invalidation issues from only the resource changing
+	PreviewBrush = MakeShared<FSlateMaterialBrush>(*InPreviewMaterial, PreviewBrush->ImageSize);
 	ImageWidget->SetImage(PreviewBrush.Get());
 }
 
+void SMaterialEditorUIPreviewZoomer::SetBackgroundSettings(const FPreviewBackgroundSettings& NewSettings)
+{
+	const bool bCheckerboardChanged = NewSettings.Checkerboard != BackgroundSettings.Checkerboard;
+
+	BackgroundSettings = NewSettings;
+
+	if (bCheckerboardChanged)
+	{
+		ModifyCheckerboardTextureColors(BackgroundSettings.Checkerboard);
+	}
+}
+
+void SMaterialEditorUIPreviewZoomer::ModifyCheckerboardTextureColors(const FCheckerboardSettings& Checkerboard)
+{
+	DestroyCheckerboardTexture();
+	SetupCheckerboardTexture(Checkerboard.ColorOne, Checkerboard.ColorTwo, Checkerboard.Size);
+
+	if (!CheckerboardBrush.IsValid())
+	{
+		CheckerboardBrush = MakeShared<FSlateImageBrush>(CheckerboardTexture, FVector2D(Checkerboard.Size), FLinearColor::White, ESlateBrushTileType::Both);
+	}
+	else
+	{
+		// TODO: May need to invalidate paint here if the widget isn't aware the brush changed?
+		CheckerboardBrush->SetResourceObject(CheckerboardTexture);
+		CheckerboardBrush->SetImageSize(FVector2D(Checkerboard.Size));
+	}
+}
+
+void SMaterialEditorUIPreviewZoomer::SetupCheckerboardTexture(const FColor& ColorOne, const FColor& ColorTwo, int32 CheckerSize)
+{
+	if (CheckerboardTexture == nullptr)
+	{
+		CheckerboardTexture = FImageUtils::CreateCheckerboardTexture(ColorOne, ColorTwo, CheckerSize);
+	}
+}
+
+void SMaterialEditorUIPreviewZoomer::DestroyCheckerboardTexture()
+{
+	if (CheckerboardTexture)
+	{
+		if (CheckerboardTexture->Resource)
+		{
+			CheckerboardTexture->ReleaseResource();
+		}
+		CheckerboardTexture->MarkPendingKill();
+		CheckerboardTexture = nullptr;
+	}
+}
+
+void SMaterialEditorUIPreviewZoomer::AddReferencedObjects(FReferenceCollector& Collector)
+{
+	if (CheckerboardTexture != nullptr)
+	{
+		Collector.AddReferencedObject(CheckerboardTexture);
+	}
+}
+
+void SMaterialEditorUIPreviewZoomer::Tick(const FGeometry& AllottedGeometry, const double InCurrentTime, const float InDeltaTime)
+{
+	ScrollyZoomy.Tick(InDeltaTime, *this);
+}
+
+FReply SMaterialEditorUIPreviewZoomer::OnMouseButtonDown(const FGeometry& MyGeometry, const FPointerEvent& MouseEvent)
+{
+	return ScrollyZoomy.OnMouseButtonDown(MouseEvent);
+}
+
+FReply SMaterialEditorUIPreviewZoomer::OnMouseButtonUp(const FGeometry& MyGeometry, const FPointerEvent& MouseEvent)
+{
+	if (MouseEvent.GetEffectingButton() == EKeys::RightMouseButton)
+	{
+		// If they didn't drag far enough to trigger a scroll, then treat it like a normal click,
+		// which would show the context menu for rmb
+		bool bWasPanning = ScrollyZoomy.IsRightClickScrolling();
+		if (!bWasPanning)
+		{
+			OnContextMenuRequested.ExecuteIfBound(MyGeometry, MouseEvent);
+		}
+	}
+
+	return ScrollyZoomy.OnMouseButtonUp(AsShared(), MyGeometry, MouseEvent);
+}
+
+FReply SMaterialEditorUIPreviewZoomer::OnMouseMove(const FGeometry& MyGeometry, const FPointerEvent& MouseEvent)
+{
+	// Only pass this on if we're scrollable, otherwise ScrollyZoomy will hide the cursor while rmb is down
+	if (IsCurrentlyScrollable())
+	{
+		return ScrollyZoomy.OnMouseMove(AsShared(), *this, MyGeometry, MouseEvent);
+	}
+
+	return FReply::Unhandled();
+}
+
+void SMaterialEditorUIPreviewZoomer::OnMouseLeave(const FPointerEvent& MouseEvent)
+{
+	ScrollyZoomy.OnMouseLeave(AsShared(), MouseEvent);
+}
+
+FReply SMaterialEditorUIPreviewZoomer::OnMouseWheel(const FGeometry& MyGeometry, const FPointerEvent& MouseEvent)
+{
+	return HandleScrollEvent(MouseEvent);
+}
+
+FCursorReply SMaterialEditorUIPreviewZoomer::OnCursorQuery(const FGeometry& MyGeometry, const FPointerEvent& CursorEvent) const
+{
+	// Only pass this on if we're scrollable, otherwise ScrollyZoomy will hide the cursor while rmb is down
+	if (IsCurrentlyScrollable())
+	{
+		return ScrollyZoomy.OnCursorQuery();
+	}
+
+	return FCursorReply::Unhandled();
+}
+
+bool SMaterialEditorUIPreviewZoomer::ScrollBy(const FVector2D& Offset)
+{
+	const FVector2D PrevPhysicalOffset = PhysicalOffset;
+	
+	PhysicalOffset += Offset.RoundToVector();
+
+	bCenterInFrame = false;
+
+	ClampViewOffset(GetDesiredSize(), CachedSize);
+
+	return PhysicalOffset != PrevPhysicalOffset;
+}
+
+FReply SMaterialEditorUIPreviewZoomer::HandleScrollEvent(const FPointerEvent& MouseEvent)
+{
+	return ScrollyZoomy.OnMouseWheel(MouseEvent, *this);
+}
+
+void SMaterialEditorUIPreviewZoomer::ScrollToCenter()
+{
+	bCenterInFrame = true;
+}
+
+bool SMaterialEditorUIPreviewZoomer::IsCentered() const
+{
+	return bCenterInFrame;
+}
+
+void SMaterialEditorUIPreviewZoomer::ClampViewOffset(const FVector2D& ZoomedPreviewSize, const FVector2D& LocalSize)
+{
+	PhysicalOffset.X = ClampViewOffsetAxis(ZoomedPreviewSize.X, LocalSize.X, PhysicalOffset.X);
+	PhysicalOffset.Y = ClampViewOffsetAxis(ZoomedPreviewSize.Y, LocalSize.Y, PhysicalOffset.Y);
+}
+
+float SMaterialEditorUIPreviewZoomer::ClampViewOffsetAxis(const float ZoomedPreviewSize, const float LocalSize, const float CurrentOffset)
+{
+	if (ZoomedPreviewSize <= LocalSize)
+	{
+		// If the viewport is smaller than the available size, then we can't be scrolled
+		return 0.0f;
+	}
+
+	// Given the size of the viewport, and the current size of the window, work how far we can scroll
+	// Note: This number is negative since scrolling down/right moves the viewport up/left
+	const float MaxScrollOffset = LocalSize - ZoomedPreviewSize;
+	const float MinScrollOffset = 0.f;
+
+	// Clamp the left/top edge
+	if (CurrentOffset < MaxScrollOffset)
+	{
+		return MaxScrollOffset;
+	}
+
+	// Clamp the right/bottom edge
+	if (CurrentOffset > MinScrollOffset)
+	{
+		return MinScrollOffset;
+	}
+
+	return CurrentOffset;
+}
+
 void SMaterialEditorUIPreviewViewport::Construct( const FArguments& InArgs, UMaterialInterface* PreviewMaterial )
 {
+	ZoomLevelFade = FCurveSequence(0.0f, 1.0f);
+	ZoomLevelFade.JumpToEnd();
+
+	UMaterialEditorSettings* Settings = GetMutableDefault<UMaterialEditorSettings>();
+	PreviewSize = Settings->GetPreviewViewportStartingSize();
+
+	// Take a copy of the global background settings at this moment, and listen for changes so we can update our colors as the user changes them
+	BackgroundSettings = Settings->PreviewBackground;
+	Settings->OnPostEditChange.AddSP(this, &SMaterialEditorUIPreviewViewport::HandleSettingsChanged);
+
 	ChildSlot
 	[
 		SNew( SVerticalBox )
@@ -902,62 +1248,107 @@ void SMaterialEditorUIPreviewViewport::Construct( const FArguments& InArgs, UMat
 		.AutoHeight()
 		[
 			SNew( SBorder )
-			.HAlign(HAlign_Right)
+			.HAlign(HAlign_Fill)
 			.VAlign(VAlign_Top)
 			.BorderImage( FEditorStyle::GetBrush("ToolPanel.GroupBorder") )
 			[
-				SNew( SHorizontalBox )
-				+ SHorizontalBox::Slot()
-				.VAlign(VAlign_Center)
-				.Padding( 3.f )
-				.AutoWidth()
-				[
-					SNew( STextBlock )
-					.Text( LOCTEXT("PreviewSize", "Preview Size" ) )
-				]
+				SNew(SHorizontalBox)
 				+ SHorizontalBox::Slot()
+				.FillWidth(1.0f)
 				.VAlign(VAlign_Center)
-				.Padding( 3.f )
-				.MaxWidth( 75 )
 				[
-					SNew( SNumericEntryBox<int32> )
-					.AllowSpin( true )
-					.MinValue(1)
-					.MaxSliderValue( 4096 )
-					.OnValueChanged( this, &SMaterialEditorUIPreviewViewport::OnPreviewXChanged )
-					.OnValueCommitted( this, &SMaterialEditorUIPreviewViewport::OnPreviewXCommitted )
-					.Value( this, &SMaterialEditorUIPreviewViewport::OnGetPreviewXValue )
-					.MinDesiredValueWidth( 75 )
-					.Label()
-					[	
-						SNew( SBox )
-						.VAlign( VAlign_Center )
-						[
-							SNew( STextBlock )
-							.Text( LOCTEXT("PreviewSize_X", "X") )
+					SNew(SHorizontalBox)
+					+ SHorizontalBox::Slot()
+					.VAlign(VAlign_Center)
+					.Padding( 3.f )
+					.AutoWidth()
+					[
+						SNew( STextBlock )
+						.Text( LOCTEXT("PreviewSize", "Preview Size" ) )
+					]
+					+ SHorizontalBox::Slot()
+					.VAlign(VAlign_Center)
+					.Padding( 3.f )
+					.MaxWidth( 75 )
+					[
+						SNew( SNumericEntryBox<int32> )
+						.AllowSpin( true )
+						.MinValue(1)
+						.MaxSliderValue( 4096 )
+						.OnValueChanged( this, &SMaterialEditorUIPreviewViewport::OnPreviewXChanged )
+						.OnValueCommitted( this, &SMaterialEditorUIPreviewViewport::OnPreviewXCommitted )
+						.Value( this, &SMaterialEditorUIPreviewViewport::OnGetPreviewXValue )
+						.MinDesiredValueWidth( 75 )
+						.Label()
+						[	
+							SNew( SBox )
+							.VAlign( VAlign_Center )
+							[
+								SNew( STextBlock )
+								.Text( LOCTEXT("PreviewSize_X", "X") )
+							]
+						]
+					]
+					+ SHorizontalBox::Slot()
+					.VAlign(VAlign_Center)
+					.Padding( 3.f )
+					.MaxWidth( 75 )
+					[
+						SNew( SNumericEntryBox<int32> )
+						.AllowSpin( true )
+						.MinValue(1)
+						.MaxSliderValue( 4096 )
+						.MinDesiredValueWidth( 75 )
+						.OnValueChanged( this, &SMaterialEditorUIPreviewViewport::OnPreviewYChanged )
+						.OnValueCommitted( this, &SMaterialEditorUIPreviewViewport::OnPreviewYCommitted )
+						.Value( this, &SMaterialEditorUIPreviewViewport::OnGetPreviewYValue )
+						.Label()
+						[	
+							SNew( SBox )
+							.VAlign( VAlign_Center )
+							[
+								SNew( STextBlock )
+								.Text( LOCTEXT("PreviewSize_Y", "Y") )
+							]
 						]
 					]
 				]
 				+ SHorizontalBox::Slot()
+				.HAlign(HAlign_Right)
 				.VAlign(VAlign_Center)
-				.Padding( 3.f )
-				.MaxWidth( 75 )
+				.AutoWidth()
 				[
-					SNew( SNumericEntryBox<int32> )
-					.AllowSpin( true )
-					.MinValue(1)
-					.MaxSliderValue( 4096 )
-					.MinDesiredValueWidth( 75 )
-					.OnValueChanged( this, &SMaterialEditorUIPreviewViewport::OnPreviewYChanged )
-					.OnValueCommitted( this, &SMaterialEditorUIPreviewViewport::OnPreviewYCommitted )
-					.Value( this, &SMaterialEditorUIPreviewViewport::OnGetPreviewYValue )
-					.Label()
-					[	
-						SNew( SBox )
-						.VAlign( VAlign_Center )
+					SNew( SHorizontalBox )
+					
+					+ SHorizontalBox::Slot()
+					.HAlign(HAlign_Left)
+					.VAlign(VAlign_Center)
+					.AutoWidth()
+					.Padding(FMargin(0.f, 0.f, 4.f, 0.f))
+					[
+						SNew(STextBlock)
+						.Text(this, &SMaterialEditorUIPreviewViewport::GetZoomText)
+						.ColorAndOpacity(this, &SMaterialEditorUIPreviewViewport::GetZoomTextColorAndOpacity)
+						//.Visibility(EVisibility::SelfHitTestInvisible)
+						.ToolTip(SNew(SToolTip).Text(this, &SMaterialEditorUIPreviewViewport::GetDisplayedAtSizeText))
+					]
+					+ SHorizontalBox::Slot()
+					.HAlign(HAlign_Right)
+					.AutoWidth()
+					[
+						SNew( SComboButton )
+						.ContentPadding(0)
+						.ForegroundColor( FSlateColor::UseForeground() )
+						.ButtonStyle( FEditorStyle::Get(), "ToggleButton" )
+						.AddMetaData<FTagMetaData>(FTagMetaData(TEXT("ViewOptions")))
+						.MenuContent()
+						[
+							BuildViewOptionsMenu().MakeWidget()
+						]
+						.ButtonContent()
 						[
-							SNew( STextBlock )
-							.Text( LOCTEXT("PreviewSize_Y", "Y") )
+							SNew(SImage)
+							.Image( FEditorStyle::GetBrush("GenericViewButton") )
 						]
 					]
 				]
@@ -965,18 +1356,75 @@ void SMaterialEditorUIPreviewViewport::Construct( const FArguments& InArgs, UMat
 		]
 		+ SVerticalBox::Slot()
 		[
-			SNew( SBorder )
+			SAssignNew( PreviewArea, SBorder )
+			.Padding(0.f)
 			.HAlign( HAlign_Center )
 			.VAlign( VAlign_Center )
+			.OnMouseButtonUp(this, &SMaterialEditorUIPreviewViewport::OnViewportClicked)
 			.BorderImage( FEditorStyle::GetBrush("BlackBrush") )
+			.Clipping(EWidgetClipping::ClipToBounds)
 			[
 				SAssignNew( PreviewZoomer, SMaterialEditorUIPreviewZoomer, PreviewMaterial )
+				.OnContextMenuRequested(this, &SMaterialEditorUIPreviewViewport::ShowContextMenu)
+				.OnZoomed(this, &SMaterialEditorUIPreviewViewport::HandleDidZoom)
+				.InitialPreviewSize(FVector2D(PreviewSize))
+				.BackgroundSettings(BackgroundSettings)
 			]
 		]
 	];
+}
+
+SMaterialEditorUIPreviewViewport::~SMaterialEditorUIPreviewViewport()
+{
+	UMaterialEditorSettings* Settings = GetMutableDefault<UMaterialEditorSettings>();
+	Settings->OnPostEditChange.RemoveAll(this);
+}
+
+FText SMaterialEditorUIPreviewViewport::GetDisplayedAtSizeText() const
+{
+	FVector2D DisplayedSize = PreviewZoomer->ComputeZoomedPreviewSize().RoundToVector();
+	return FText::Format(LOCTEXT("DisplayedAtSize", "Currently displayed at: {0}x{1}"), FText::AsNumber(DisplayedSize.X), FText::AsNumber(DisplayedSize.Y));
+}
 
-	PreviewSize = GetDefault<UMaterialEditorSettings>()->GetPreviewViewportStartingSize();
-	PreviewZoomer->SetPreviewSize( FVector2D(PreviewSize) );
+FText SMaterialEditorUIPreviewViewport::GetZoomText() const
+{
+	static FText ZoomLevelFormat = LOCTEXT("ZoomLevelFormat", "Zoom: {0}");
+	
+	FText ZoomLevelPercent = FText::AsPercent(PreviewZoomer->GetZoomLevel());
+	return FText::FormatOrdered(FTextFormat(ZoomLevelFormat), ZoomLevelPercent);
+}
+
+void SMaterialEditorUIPreviewViewport::HandleDidZoom()
+{
+	ZoomLevelFade.Play(this->AsShared());
+}
+
+void SMaterialEditorUIPreviewViewport::ExecuteZoomToActual()
+{
+	PreviewZoomer->SetZoomLevel(1.f);
+	PreviewZoomer->ScrollToCenter();
+}
+
+bool SMaterialEditorUIPreviewViewport::CanZoomToActual() const
+{
+	return !FMath::IsNearlyEqual(PreviewZoomer->GetZoomLevel(), 1.f, 0.01f) || !PreviewZoomer->IsCentered();
+}
+
+FSlateColor SMaterialEditorUIPreviewViewport::GetZoomTextColorAndOpacity() const
+{
+	return FLinearColor(1, 1, 1, 1.25f - ZoomLevelFade.GetLerp() * 0.75f);
+}
+
+void SMaterialEditorUIPreviewViewport::HandleSettingsChanged()
+{
+	const UMaterialEditorSettings& Settings = *GetDefault<UMaterialEditorSettings>();
+	
+	// Keep any global settings up to date when the user changes them in the editor prefs window
+	BackgroundSettings.Checkerboard = Settings.PreviewBackground.Checkerboard;
+	BackgroundSettings.BackgroundColor = Settings.PreviewBackground.BackgroundColor;
+	BackgroundSettings.BorderColor = Settings.PreviewBackground.BorderColor;
+
+	PreviewZoomer->SetBackgroundSettings(BackgroundSettings);
 }
 
 void SMaterialEditorUIPreviewViewport::SetPreviewMaterial(UMaterialInterface* InMaterialInterface)
@@ -984,11 +1432,134 @@ void SMaterialEditorUIPreviewViewport::SetPreviewMaterial(UMaterialInterface* In
 	PreviewZoomer->SetPreviewMaterial(InMaterialInterface);
 }
 
+FReply SMaterialEditorUIPreviewViewport::OnViewportClicked(const FGeometry& Geometry, const FPointerEvent& MouseEvent)
+{
+	if (MouseEvent.GetEffectingButton() == EKeys::RightMouseButton)
+	{
+		ShowContextMenu(Geometry, MouseEvent);
+		return FReply::Handled();
+	}
+	else
+	{
+		return FReply::Unhandled();
+	}
+}
+
+void SMaterialEditorUIPreviewViewport::ShowContextMenu(const FGeometry& Geometry, const FPointerEvent& MouseEvent)
+{
+	FWidgetPath WidgetPath = MouseEvent.GetEventPath() != nullptr ? *MouseEvent.GetEventPath() : FWidgetPath();
+	FSlateApplication::Get().PushMenu(AsShared(), WidgetPath, BuildViewOptionsMenu(/* bForContextMenu */ true).MakeWidget(), FSlateApplication::Get().GetCursorPos(), FPopupTransitionEffect::ContextMenu);
+}
+
+FReply SMaterialEditorUIPreviewViewport::OnMouseWheel(const FGeometry& MyGeometry, const FPointerEvent& MouseEvent)
+{
+	// Forward scrolls over the preview area to the zoomer so you can still scroll over the blank space around the preview
+	if (PreviewArea->IsHovered())
+	{
+		return PreviewZoomer->HandleScrollEvent(MouseEvent);
+	}
+
+	return FReply::Unhandled();
+}
+
+FMenuBuilder SMaterialEditorUIPreviewViewport::BuildViewOptionsMenu(bool bForContextMenu)
+{
+	auto GenerateBackgroundMenuContent = [this](FMenuBuilder& MenuBuilder)
+	{
+		// Not bothering to create commands for these since they'll probably be rarely changed,
+		// and would mean needing to duplicate your bindings between texture editor and material editor
+		MenuBuilder.AddMenuEntry(
+			LOCTEXT("SolidBackground", "Solid Color"),
+			LOCTEXT("SolidBackground_ToolTip", "Displays a solid background color behind the preview."),
+			FSlateIcon(),
+			FUIAction(
+				FExecuteAction::CreateSP(this, &SMaterialEditorUIPreviewViewport::SetBackgroundType, EBackgroundType::SolidColor),
+				FCanExecuteAction::CreateLambda([]() { return true; }),
+				FIsActionChecked::CreateSP(this, &SMaterialEditorUIPreviewViewport::IsBackgroundTypeChecked, EBackgroundType::SolidColor)
+			),
+			NAME_None,
+			EUserInterfaceActionType::RadioButton);
+
+		MenuBuilder.AddMenuEntry(
+			LOCTEXT("CheckeredBackground", "Checkerboard"),
+			LOCTEXT("CheckeredBackground_ToolTip", "Displays a checkerboard behind the preview."),
+			FSlateIcon(),
+			FUIAction(
+				FExecuteAction::CreateSP(this, &SMaterialEditorUIPreviewViewport::SetBackgroundType, EBackgroundType::Checkered),
+				FCanExecuteAction::CreateLambda([]() { return true; }),
+				FIsActionChecked::CreateSP(this, &SMaterialEditorUIPreviewViewport::IsBackgroundTypeChecked, EBackgroundType::Checkered)
+			),
+			NAME_None,
+			EUserInterfaceActionType::RadioButton);
+	};
+
+	FMenuBuilder MenuBuilder(true, nullptr);
+
+	MenuBuilder.BeginSection("ZoomSection", LOCTEXT("ZoomSectionHeader", "Zoom"));
+	{
+		MenuBuilder.AddMenuEntry(
+			LOCTEXT("ZoomToActual", "Zoom to 100%"),
+			LOCTEXT("ZoomToActual_Tooltip", "Resets the zoom to 100% and centers the preview."),
+			FSlateIcon(),
+			FUIAction(
+				FExecuteAction::CreateSP(this, &SMaterialEditorUIPreviewViewport::ExecuteZoomToActual),
+				FCanExecuteAction::CreateSP(this, &SMaterialEditorUIPreviewViewport::CanZoomToActual)
+			)
+		);
+	}
+	MenuBuilder.EndSection();
+
+	// view port options
+	MenuBuilder.BeginSection("ViewportSection", LOCTEXT("ViewportSectionHeader", "Viewport Options"));
+	{
+		MenuBuilder.AddSubMenu(
+			LOCTEXT("Background", "Background"),
+			LOCTEXT("BackgroundTooltip", "Configure the preview's background."),
+			FNewMenuDelegate::CreateLambda(GenerateBackgroundMenuContent)
+		);
+
+		MenuBuilder.AddMenuEntry(
+			LOCTEXT("ShowBorder", "Show Border"),
+			LOCTEXT("ShowBorder_Tooltip", "Displays a border around the preview bounds."),
+			FSlateIcon(),
+			FUIAction(
+				FExecuteAction::CreateSP(this, &SMaterialEditorUIPreviewViewport::ToggleShowBorder),
+				FCanExecuteAction::CreateLambda([]() { return true; }),
+				FIsActionChecked::CreateSP(this, &SMaterialEditorUIPreviewViewport::IsShowBorderChecked)
+			),
+			NAME_None,
+			EUserInterfaceActionType::ToggleButton);
+	}
+	MenuBuilder.EndSection();
+
+	// Don't include settings item for right-clicks
+	if (!bForContextMenu)
+	{
+		MenuBuilder.AddMenuSeparator();
+
+		MenuBuilder.AddMenuEntry(
+			LOCTEXT("Settings", "Settings"),
+			LOCTEXT("Settings_Tooltip", "Opens the material editor preferences pane."),
+			FSlateIcon(),
+			FUIAction(
+				FExecuteAction::CreateSP(this, &SMaterialEditorUIPreviewViewport::HandleSettingsActionExecute)
+			),
+			NAME_None,
+			EUserInterfaceActionType::Button);
+	}
+
+	return MenuBuilder;
+}
+
+void SMaterialEditorUIPreviewViewport::HandleSettingsActionExecute()
+{
+	FModuleManager::LoadModuleChecked<ISettingsModule>("Settings").ShowViewer("Editor", "ContentEditors", "Material Editor"); // Note: This has a space, unlike a lot of other setting sections - see MaterialEditorModuleConstants::SettingsSectionName
+}
+
 void SMaterialEditorUIPreviewViewport::OnPreviewXChanged( int32 NewValue )
 {
 	PreviewSize.X = NewValue;
 	PreviewZoomer->SetPreviewSize( FVector2D( PreviewSize ) );
-
 }
 
 void SMaterialEditorUIPreviewViewport::OnPreviewXCommitted( int32 NewValue, ETextCommit::Type )
@@ -996,12 +1567,10 @@ void SMaterialEditorUIPreviewViewport::OnPreviewXCommitted( int32 NewValue, ETex
 	OnPreviewXChanged( NewValue );
 }
 
-
 void SMaterialEditorUIPreviewViewport::OnPreviewYChanged( int32 NewValue )
 {
 	PreviewSize.Y = NewValue;
 	PreviewZoomer->SetPreviewSize( FVector2D( PreviewSize ) );
-
 }
 
 void SMaterialEditorUIPreviewViewport::OnPreviewYCommitted( int32 NewValue, ETextCommit::Type )
@@ -1009,5 +1578,32 @@ void SMaterialEditorUIPreviewViewport::OnPreviewYCommitted( int32 NewValue, ETex
 	OnPreviewYChanged( NewValue );
 }
 
+bool SMaterialEditorUIPreviewViewport::IsBackgroundTypeChecked(EBackgroundType BackgroundType) const
+{
+	return BackgroundSettings.BackgroundType == BackgroundType;
+}
+
+bool SMaterialEditorUIPreviewViewport::IsShowBorderChecked() const
+{
+	return BackgroundSettings.bShowBorder;
+}
+
+void SMaterialEditorUIPreviewViewport::SetBackgroundType(EBackgroundType NewBackgroundType)
+{
+	BackgroundSettings.BackgroundType = NewBackgroundType;
+	PreviewZoomer->SetBackgroundSettings(BackgroundSettings);
+
+	// Use this as the default of newly opened preview viewports
+	GetMutableDefault<UMaterialEditorSettings>()->PreviewBackground.BackgroundType = BackgroundSettings.BackgroundType;
+}
+
+void SMaterialEditorUIPreviewViewport::ToggleShowBorder()
+{
+	BackgroundSettings.bShowBorder = !BackgroundSettings.bShowBorder;
+	PreviewZoomer->SetBackgroundSettings(BackgroundSettings);
+
+	// Use this as the default of newly opened preview viewports
+	GetMutableDefault<UMaterialEditorSettings>()->PreviewBackground.bShowBorder = BackgroundSettings.bShowBorder;
+}
 
 #undef LOCTEXT_NAMESPACE
diff --git a/Engine/Source/Editor/MaterialEditor/Private/SMaterialEditorViewport.h b/Engine/Source/Editor/MaterialEditor/Private/SMaterialEditorViewport.h
--- a/Engine/Source/Editor/MaterialEditor/Private/SMaterialEditorViewport.h
+++ b/Engine/Source/Editor/MaterialEditor/Private/SMaterialEditorViewport.h
@@ -13,6 +13,8 @@
 #include "ThumbnailRendering/ThumbnailManager.h"
 #include "SEditorViewport.h"
 #include "Editor/UnrealEd/Public/SCommonEditorViewportToolbarBase.h"
+#include "MaterialEditorSettings.h"
+#include "Animation/CurveSequence.h"
 
 class FMaterialEditorViewportClient;
 class IMaterialEditor;
@@ -136,6 +138,8 @@ public:
 	SLATE_END_ARGS()
 
 	void Construct( const FArguments& InArgs, UMaterialInterface* PreviewMaterial );
+	virtual ~SMaterialEditorUIPreviewViewport();
+
 	void SetPreviewMaterial(UMaterialInterface* InMaterialInterface);
 
 private:
@@ -146,7 +150,31 @@ private:
 	TOptional<int32> OnGetPreviewXValue() const { return PreviewSize.X; }
 	TOptional<int32> OnGetPreviewYValue() const { return PreviewSize.Y; }
 
+	bool IsBackgroundTypeChecked(EBackgroundType BackgroundType) const;
+	bool IsShowBorderChecked() const;
+
+	void SetBackgroundType(EBackgroundType NewBackgroundType);
+	void ToggleShowBorder();
+
+	FMenuBuilder BuildViewOptionsMenu(bool bForContextMenu = false);
+	FReply OnViewportClicked(const FGeometry& Geometry, const FPointerEvent& MouseEvent);
+	void ShowContextMenu(const FGeometry& Geometry, const FPointerEvent& MouseEvent);
+
+	FReply OnMouseWheel(const FGeometry& MyGeometry, const FPointerEvent& MouseEvent);
+	void HandleSettingsActionExecute();
+	void HandleSettingsChanged();
+	
+	void HandleDidZoom();
+	void ExecuteZoomToActual();
+	bool CanZoomToActual() const;
+	FSlateColor GetZoomTextColorAndOpacity() const;
+	FText GetZoomText() const;
+	FText GetDisplayedAtSizeText() const;
+
 private:
 	FIntPoint PreviewSize;
 	TSharedPtr<class SMaterialEditorUIPreviewZoomer> PreviewZoomer;
+	TSharedPtr<SWidget> PreviewArea;
+	FPreviewBackgroundSettings BackgroundSettings;
+	FCurveSequence ZoomLevelFade;
 };
diff --git a/Engine/Source/Editor/MaterialEditor/Public/MaterialEditorSettings.h b/Engine/Source/Editor/MaterialEditor/Public/MaterialEditorSettings.h
--- a/Engine/Source/Editor/MaterialEditor/Public/MaterialEditorSettings.h
+++ b/Engine/Source/Editor/MaterialEditor/Public/MaterialEditorSettings.h
@@ -5,6 +5,90 @@
 #include "Engine/EngineTypes.h"
 #include "MaterialEditorSettings.generated.h"
 
+/**
+ * Enumerates background for the material editor UI viewport (similar to ETextureEditorBackgrounds)
+ */
+UENUM()
+enum class EBackgroundType : uint8
+{
+	SolidColor UMETA(DisplayName = "Solid Color"),
+	Checkered UMETA(DisplayName = "Checkered")
+};
+
+USTRUCT()
+struct MATERIALEDITOR_API FCheckerboardSettings
+{
+	GENERATED_BODY()
+	
+	// The first color of the checkerboard.
+	UPROPERTY(EditAnywhere, Category = "Checkerboard")
+	FColor ColorOne;
+
+	// The second color of the checkerboard
+	UPROPERTY(EditAnywhere, Category = "Checkerboard")
+	FColor ColorTwo;
+
+	// The size of the checkered tiles
+	UPROPERTY(EditAnywhere, meta = (ClampMin = "2", ClampMax = "4096"), Category = "Checkerboard")
+	int32 Size;
+
+	FCheckerboardSettings()
+		: ColorOne(FColor(128, 128, 128))
+		, ColorTwo(FColor(64, 64, 64))
+		, Size(32)
+	{
+
+	}
+};
+
+inline bool operator==(const FCheckerboardSettings& Lhs, const FCheckerboardSettings& Rhs)
+{
+	return Lhs.ColorOne == Rhs.ColorOne
+		&& Lhs.ColorTwo == Rhs.ColorTwo
+		&& Lhs.Size == Rhs.Size;
+}
+
+inline bool operator!=(const FCheckerboardSettings& Lhs, const FCheckerboardSettings& Rhs)
+{
+	return !operator==(Lhs, Rhs);
+}
+
+USTRUCT()
+struct MATERIALEDITOR_API FPreviewBackgroundSettings
+{
+	GENERATED_BODY()
+	
+	// If true, displays a border around the texture (configured via the material editor)
+	UPROPERTY()
+	bool bShowBorder;
+
+	// Color to use for the border, if enabled
+	UPROPERTY(EditAnywhere, Category = "Background")
+	FColor BorderColor;
+
+	// The type of background to show (configured via the material editor)
+	UPROPERTY()
+	EBackgroundType BackgroundType;
+
+	// The color used as the background of the preview
+	UPROPERTY(EditAnywhere, Category = "Background")
+	FColor BackgroundColor;
+
+	UPROPERTY(EditAnywhere, Category = "Background")
+	FCheckerboardSettings Checkerboard;
+
+	// Note: For now these defaults match the historical material editor behavior, not the texture editor's defaults
+	FPreviewBackgroundSettings()
+		: bShowBorder(false)
+		, BorderColor(FColor::White)
+		, BackgroundType(EBackgroundType::SolidColor)
+		, BackgroundColor(FColor::Black)
+	{
+
+	}
+
+};
+
 UCLASS(config = EditorPerProjectUserSettings)
 class MATERIALEDITOR_API UMaterialEditorSettings : public UObject
 {
@@ -28,8 +112,22 @@ protected:
 	int32 DefaultPreviewHeight = 250;
 
 public:
+	// Configures the background shown behind the UI material preview
+	UPROPERTY(config, EditAnywhere, Category = "User Interface Domain")
+	FPreviewBackgroundSettings PreviewBackground;
+
 	FIntPoint GetPreviewViewportStartingSize() const
 	{
 		return FIntPoint(DefaultPreviewWidth, DefaultPreviewHeight);
 	}
-};
\ No newline at end of file
+
+#if WITH_EDITOR
+	// Allow listening for changes just to this settings object without having to listen to all UObjects
+	FSimpleMulticastDelegate OnPostEditChange;
+	virtual void PostEditChangeProperty(struct FPropertyChangedEvent& PropertyChangedEvent) override
+	{
+		Super::PostEditChangeProperty(PropertyChangedEvent);
+		OnPostEditChange.Broadcast();
+	}
+#endif // WITH_EDITOR
+};
