Fix issue where it was possible to shift-click to open existing tabs on top of each other in material editor.

diff --git a/Engine/Source/Editor/MaterialEditor/Private/MaterialEditor.cpp b/Engine/Source/Editor/MaterialEditor/Private/MaterialEditor.cpp
--- a/Engine/Source/Editor/MaterialEditor/Private/MaterialEditor.cpp
+++ b/Engine/Source/Editor/MaterialEditor/Private/MaterialEditor.cpp
@@ -5410,7 +5410,25 @@ void FMaterialEditor::JumpToHyperlink(const UObject* ObjectReference)
 		const bool bIsShiftPressed = FSlateApplication::Get().GetModifierKeys().IsShiftDown();
 		if (bIsShiftPressed)
 		{
-			OpenMode = FDocumentTracker::ForceOpenNewDocument;
+			auto PayloadAlreadyOpened = [&]()
+			{
+				TArray< TSharedPtr<SDockTab> > GraphEditorTabs;
+				DocumentManager->FindAllTabsForFactory(GraphEditorTabFactoryPtr, /*out*/ GraphEditorTabs);
+
+				for (TSharedPtr<SDockTab>& GraphEditorTab : GraphEditorTabs)
+				{
+					TSharedRef<SGraphEditor> Editor = StaticCastSharedRef<SGraphEditor>((GraphEditorTab)->GetContent());
+
+					if (Editor->GetCurrentGraph() == Graph)
+					{
+						return true;
+					}
+				}
+
+				return false;
+			};
+
+			OpenMode = PayloadAlreadyOpened() ? FDocumentTracker::RestorePreviousDocument : FDocumentTracker::ForceOpenNewDocument;
 		}
 
 		// Open the document
