Fix material editor toolbar not appearing

diff --git a/Engine/Source/Editor/MaterialEditor/Private/MaterialEditor.h b/Engine/Source/Editor/MaterialEditor/Private/MaterialEditor.h
--- a/Engine/Source/Editor/MaterialEditor/Private/MaterialEditor.h
+++ b/Engine/Source/Editor/MaterialEditor/Private/MaterialEditor.h
@@ -259,7 +259,7 @@ public:
 	void RegisterToolbarTab(const TSharedRef<class FTabManager>& TabManager);
 	virtual void RegisterTabSpawners(const TSharedRef<class FTabManager>& TabManager) override;
 	virtual void UnregisterTabSpawners(const TSharedRef<class FTabManager>& TabManager) override;
-
+	virtual bool UsesCustomToolbarPlacement() const { return true; }
 public:
 	/** Initializes the editor to use a material. Should be the first thing called. */
 	void InitEditorForMaterial(UMaterial* InMaterial);
