Fix material instance toolbar not appearing

diff --git a/Engine/Source/Editor/MaterialEditor/Private/MaterialEditor.h b/Engine/Source/Editor/MaterialEditor/Private/MaterialEditor.h
--- a/Engine/Source/Editor/MaterialEditor/Private/MaterialEditor.h
+++ b/Engine/Source/Editor/MaterialEditor/Private/MaterialEditor.h
@@ -259,7 +259,6 @@ public:
 	void RegisterToolbarTab(const TSharedRef<class FTabManager>& TabManager);
 	virtual void RegisterTabSpawners(const TSharedRef<class FTabManager>& TabManager) override;
 	virtual void UnregisterTabSpawners(const TSharedRef<class FTabManager>& TabManager) override;
-	virtual bool UsesCustomToolbarPlacement() const { return true; }
 public:
 	/** Initializes the editor to use a material. Should be the first thing called. */
 	void InitEditorForMaterial(UMaterial* InMaterial);
diff --git a/Engine/Source/Editor/MaterialEditor/Public/IMaterialEditor.h b/Engine/Source/Editor/MaterialEditor/Public/IMaterialEditor.h
--- a/Engine/Source/Editor/MaterialEditor/Public/IMaterialEditor.h
+++ b/Engine/Source/Editor/MaterialEditor/Public/IMaterialEditor.h
@@ -155,6 +155,7 @@ public:
 	/** Fills inheritance drop down menu */
 	virtual void GenerateInheritanceMenu(class UToolMenu* Menu) {};
 
+	virtual bool UsesCustomToolbarPlacement() const { return true; }
 private:
 	FMaterialEditorClosedEvent MaterialEditorClosedEvent;
 	FRegisterTabSpawnersEvent RegisterTabSpawnersEvent;
