Material editor widgets now use the non-PropertyHandle overload of ResetToDefaultOverride. This version does not defer the action to the next frame, which was causing the property handle to be invalid.

diff --git a/Engine/Source/Editor/MaterialEditor/Private/MaterialEditorInstanceDetailCustomization.cpp b/Engine/Source/Editor/MaterialEditor/Private/MaterialEditorInstanceDetailCustomization.cpp
--- a/Engine/Source/Editor/MaterialEditor/Private/MaterialEditorInstanceDetailCustomization.cpp
+++ b/Engine/Source/Editor/MaterialEditor/Private/MaterialEditorInstanceDetailCustomization.cpp
@@ -501,8 +501,8 @@ void FMaterialInstanceParameterDetails::CreateParameterValueWidget(UDEditorParam
 
 		IDetailPropertyRow& PropertyRow = DetailGroup.AddPropertyRow(ParameterValueProperty.ToSharedRef());
 
-		FIsResetToDefaultVisible IsResetVisible = FIsResetToDefaultVisible::CreateStatic(&FMaterialPropertyHelpers::ShouldShowResetToDefault, Parameter, MaterialEditorInstance);
-		FResetToDefaultHandler ResetHandler = FResetToDefaultHandler::CreateStatic(&FMaterialPropertyHelpers::ResetToDefault, Parameter, MaterialEditorInstance);
+		TAttribute<bool> IsResetVisible = TAttribute<bool>::Create(TAttribute<bool>::FGetter::CreateStatic(&FMaterialPropertyHelpers::ShouldShowResetToDefault, Parameter, MaterialEditorInstance));
+		FSimpleDelegate ResetHandler = FSimpleDelegate::CreateStatic(&FMaterialPropertyHelpers::ResetToDefault, Parameter, MaterialEditorInstance);
 		FResetToDefaultOverride ResetOverride = FResetToDefaultOverride::Create(IsResetVisible, ResetHandler);
 
 		PropertyRow
@@ -524,8 +524,6 @@ void FMaterialInstanceParameterDetails::CreateParameterValueWidget(UDEditorParam
 				{
 					TWeakObjectPtr<UMaterialExpressionTextureSampleParameter> SamplerExpression = Expression;
 
-					PropertyRow.OverrideResetToDefault(ResetOverride);
-
 					PropertyRow.CustomWidget()
 					.NameContent()
 					[
@@ -576,7 +574,7 @@ void FMaterialInstanceParameterDetails::CreateMaskParameterValueWidget(UDEditorP
 		IDetailPropertyRow& PropertyRow = DetailGroup.AddPropertyRow(ParameterValueProperty.ToSharedRef());
 		PropertyRow.EditCondition(IsParamEnabled, FOnBooleanValueChanged::CreateStatic(&FMaterialPropertyHelpers::OnOverrideParameter, Parameter, MaterialEditorInstance));
 		// Handle reset to default manually
-		PropertyRow.OverrideResetToDefault(FResetToDefaultOverride::Create(FResetToDefaultHandler::CreateStatic(&FMaterialPropertyHelpers::ResetToDefault, Parameter, MaterialEditorInstance)));
+		PropertyRow.OverrideResetToDefault(FResetToDefaultOverride::Create(FSimpleDelegate::CreateStatic(&FMaterialPropertyHelpers::ResetToDefault, Parameter, MaterialEditorInstance)));
 		PropertyRow.Visibility(TAttribute<EVisibility>::Create(TAttribute<EVisibility>::FGetter::CreateStatic(&FMaterialPropertyHelpers::ShouldShowExpression, Parameter, MaterialEditorInstance, ShowHiddenDelegate)));
 
 		const FText ParameterName = FText::FromName(Parameter->ParameterInfo.Name);
@@ -666,7 +664,7 @@ void FMaterialInstanceParameterDetails::CreateVectorChannelMaskParameterValueWid
 		IDetailPropertyRow& PropertyRow = DetailGroup.AddPropertyRow(ParameterValueProperty.ToSharedRef());
 		PropertyRow.EditCondition(IsParamEnabled, FOnBooleanValueChanged::CreateStatic(&FMaterialPropertyHelpers::OnOverrideParameter, Parameter, MaterialEditorInstance));
 		// Handle reset to default manually
-		PropertyRow.OverrideResetToDefault(FResetToDefaultOverride::Create(FResetToDefaultHandler::CreateStatic(&FMaterialPropertyHelpers::ResetToDefault, Parameter, MaterialEditorInstance)));
+		PropertyRow.OverrideResetToDefault(FResetToDefaultOverride::Create(FSimpleDelegate::CreateStatic(&FMaterialPropertyHelpers::ResetToDefault, Parameter, MaterialEditorInstance)));
 		PropertyRow.Visibility(TAttribute<EVisibility>::Create(TAttribute<EVisibility>::FGetter::CreateStatic(&FMaterialPropertyHelpers::ShouldShowExpression, Parameter, MaterialEditorInstance, ShowHiddenDelegate)));
 
 		const FText ParameterName = FText::FromName(Parameter->ParameterInfo.Name);
@@ -722,8 +720,8 @@ void FMaterialInstanceParameterDetails::CreateScalarAtlasPositionParameterValueW
 		const FText ParameterName = FText::FromName(Parameter->ParameterInfo.Name);
 		UDEditorScalarParameterValue* AtlasParameter = Cast<UDEditorScalarParameterValue>(Parameter);
 
-		FIsResetToDefaultVisible IsResetVisible = FIsResetToDefaultVisible::CreateStatic(&FMaterialPropertyHelpers::ShouldShowResetToDefault, Parameter, MaterialEditorInstance);
-		FResetToDefaultHandler ResetHandler = FResetToDefaultHandler::CreateStatic(&FMaterialPropertyHelpers::ResetCurveToDefault, Parameter, MaterialEditorInstance);
+		TAttribute<bool> IsResetVisible = TAttribute<bool>::Create(TAttribute<bool>::FGetter::CreateStatic(&FMaterialPropertyHelpers::ShouldShowResetToDefault, Parameter, MaterialEditorInstance));
+		FSimpleDelegate ResetHandler = FSimpleDelegate::CreateStatic(&FMaterialPropertyHelpers::ResetCurveToDefault, Parameter, MaterialEditorInstance);
 		FResetToDefaultOverride ResetOverride = FResetToDefaultOverride::Create(IsResetVisible, ResetHandler);
 
 		PropertyRow.OverrideResetToDefault(ResetOverride);
@@ -776,8 +774,8 @@ void FMaterialInstanceParameterDetails::CreateLabeledTextureParameterValueWidget
 
 					IDetailPropertyRow& PropertyRow = DetailGroup.AddPropertyRow(ParameterValueProperty.ToSharedRef());
 
-					FIsResetToDefaultVisible IsResetVisible = FIsResetToDefaultVisible::CreateStatic(&FMaterialPropertyHelpers::ShouldShowResetToDefault, Parameter, MaterialEditorInstance);
-					FResetToDefaultHandler ResetHandler = FResetToDefaultHandler::CreateStatic(&FMaterialPropertyHelpers::ResetToDefault, Parameter, MaterialEditorInstance);
+					TAttribute<bool> IsResetVisible = TAttribute<bool>::Create(TAttribute<bool>::FGetter::CreateStatic(&FMaterialPropertyHelpers::ShouldShowResetToDefault, Parameter, MaterialEditorInstance));
+					FSimpleDelegate ResetHandler = FSimpleDelegate::CreateStatic(&FMaterialPropertyHelpers::ResetToDefault, Parameter, MaterialEditorInstance);
 					FResetToDefaultOverride ResetOverride = FResetToDefaultOverride::Create(IsResetVisible, ResetHandler);
 
 					PropertyRow
diff --git a/Engine/Source/Editor/MaterialEditor/Private/MaterialPropertyHelpers.cpp b/Engine/Source/Editor/MaterialEditor/Private/MaterialPropertyHelpers.cpp
--- a/Engine/Source/Editor/MaterialEditor/Private/MaterialPropertyHelpers.cpp
+++ b/Engine/Source/Editor/MaterialEditor/Private/MaterialPropertyHelpers.cpp
@@ -781,7 +781,7 @@ FText FMaterialPropertyHelpers::GetParameterTooltip(UDEditorParameterValue* Para
 }
 
 
-void FMaterialPropertyHelpers::ResetToDefault(TSharedPtr<IPropertyHandle> PropertyHandle, class UDEditorParameterValue* Parameter, UMaterialEditorInstanceConstant* MaterialEditorInstance)
+void FMaterialPropertyHelpers::ResetToDefault(class UDEditorParameterValue* Parameter, UMaterialEditorInstanceConstant* MaterialEditorInstance)
 {
 	const FScopedTransaction Transaction( LOCTEXT( "ResetToDefault", "Reset To Default" ) );
 	Parameter->Modify();
@@ -869,9 +869,8 @@ void FMaterialPropertyHelpers::ResetToDefault(TSharedPtr<IPropertyHandle> Proper
 	}
 }
 
-void FMaterialPropertyHelpers::ResetLayerAssetToDefault(TSharedPtr<IPropertyHandle> InPropertyHandle, UDEditorParameterValue* InParameter, TEnumAsByte<EMaterialParameterAssociation> InAssociation, int32 Index, UMaterialEditorInstanceConstant* MaterialEditorInstance)
+void FMaterialPropertyHelpers::ResetLayerAssetToDefault(UDEditorParameterValue* InParameter, TEnumAsByte<EMaterialParameterAssociation> InAssociation, int32 Index, UMaterialEditorInstanceConstant* MaterialEditorInstance)
 {
-	
 	const FScopedTransaction Transaction(LOCTEXT("ResetToDefault", "Reset To Default"));
 	InParameter->Modify();
 	
@@ -920,7 +919,7 @@ void FMaterialPropertyHelpers::ResetLayerAssetToDefault(TSharedPtr<IPropertyHand
 	
 }
 
-bool FMaterialPropertyHelpers::ShouldLayerAssetShowResetToDefault(TSharedPtr<IPropertyHandle> InPropertyHandle, TSharedPtr<FSortedParamData> InParameterData, UMaterialInterface* InMaterial)
+bool FMaterialPropertyHelpers::ShouldLayerAssetShowResetToDefault(TSharedPtr<FSortedParamData> InParameterData, UMaterialInterface* InMaterial)
 {
 	if (!InParameterData->Parameter)
 	{
@@ -973,7 +972,7 @@ bool FMaterialPropertyHelpers::ShouldLayerAssetShowResetToDefault(TSharedPtr<IPr
 	return false;
 }
 
-bool FMaterialPropertyHelpers::ShouldShowResetToDefault(TSharedPtr<IPropertyHandle> InPropertyHandle, UDEditorParameterValue* InParameter, UMaterialEditorInstanceConstant* MaterialEditorInstance)
+bool FMaterialPropertyHelpers::ShouldShowResetToDefault(UDEditorParameterValue* InParameter, UMaterialEditorInstanceConstant* MaterialEditorInstance)
 {
 	const FMaterialParameterInfo& ParameterInfo = InParameter->ParameterInfo;
 
@@ -1298,7 +1297,7 @@ void FMaterialPropertyHelpers::SetPositionFromCurveAsset(const FAssetData& Asset
 	}
 }
 
-void FMaterialPropertyHelpers::ResetCurveToDefault(TSharedPtr<IPropertyHandle> PropertyHandle, UDEditorParameterValue* Parameter, UMaterialEditorInstanceConstant* MaterialEditorInstance)
+void FMaterialPropertyHelpers::ResetCurveToDefault(UDEditorParameterValue* Parameter, UMaterialEditorInstanceConstant* MaterialEditorInstance)
 {
 	const FScopedTransaction Transaction(LOCTEXT("ResetToDefault", "Reset To Default"));
 	Parameter->Modify();
diff --git a/Engine/Source/Editor/MaterialEditor/Private/SMaterialLayersFunctionsTree.cpp b/Engine/Source/Editor/MaterialEditor/Private/SMaterialLayersFunctionsTree.cpp
--- a/Engine/Source/Editor/MaterialEditor/Private/SMaterialLayersFunctionsTree.cpp
+++ b/Engine/Source/Editor/MaterialEditor/Private/SMaterialLayersFunctionsTree.cpp
@@ -432,8 +432,8 @@ void SMaterialLayersFunctionsInstanceTreeItem::Construct(const FArguments& InArg
 		const int32 LayerStateIndex = StackParameterData->ParameterInfo.Association == EMaterialParameterAssociation::BlendParameter ? StackParameterData->ParameterInfo.Index + 1 : StackParameterData->ParameterInfo.Index;
 
 		TAttribute<bool> IsParamEnabled = TAttribute<bool>::Create(TAttribute<bool>::FGetter::CreateSP(Tree, &SMaterialLayersFunctionsInstanceTree::IsOverriddenExpression, StackParameterData->Parameter, LayerStateIndex));
-		FIsResetToDefaultVisible IsAssetResetVisible = FIsResetToDefaultVisible::CreateStatic(&FMaterialPropertyHelpers::ShouldLayerAssetShowResetToDefault, StackParameterData, MaterialEditorInstance->Parent);
-		FResetToDefaultHandler ResetAssetHandler = FResetToDefaultHandler::CreateSP(InArgs._InTree, &SMaterialLayersFunctionsInstanceTree::ResetAssetToDefault, StackParameterData);
+		TAttribute<bool> IsAssetResetVisible = TAttribute<bool>::Create(TAttribute<bool>::FGetter::CreateStatic(&FMaterialPropertyHelpers::ShouldLayerAssetShowResetToDefault, StackParameterData, MaterialEditorInstance->Parent));
+		FSimpleDelegate ResetAssetHandler = FSimpleDelegate::CreateSP(Tree, &SMaterialLayersFunctionsInstanceTree::ResetAssetToDefault, StackParameterData);
 		FResetToDefaultOverride ResetAssetOverride = FResetToDefaultOverride::Create(IsAssetResetVisible, ResetAssetHandler);
 
 		IDetailTreeNode& Node = *StackParameterData->ParameterNode;
@@ -585,14 +585,14 @@ void SMaterialLayersFunctionsInstanceTreeItem::Construct(const FArguments& InArg
 
 		TAttribute<bool> IsParamEnabled = TAttribute<bool>::Create(TAttribute<bool>::FGetter::CreateSP(Tree, &SMaterialLayersFunctionsInstanceTree::IsOverriddenExpression, StackParameterData->Parameter, LayerStateIndex));
 		NameOverride = FText::FromName(StackParameterData->Parameter->ParameterInfo.Name);
-		FIsResetToDefaultVisible IsResetVisible = FIsResetToDefaultVisible::CreateStatic(&FMaterialPropertyHelpers::ShouldShowResetToDefault, StackParameterData->Parameter, MaterialEditorInstance);
-		FResetToDefaultHandler ResetHandler = FResetToDefaultHandler::CreateStatic(&FMaterialPropertyHelpers::ResetToDefault, StackParameterData->Parameter, MaterialEditorInstance);
+		TAttribute<bool> IsResetVisible = TAttribute<bool>::Create(TAttribute<bool>::FGetter::CreateStatic(&FMaterialPropertyHelpers::ShouldShowResetToDefault, StackParameterData->Parameter, MaterialEditorInstance));
+		FSimpleDelegate ResetHandler = FSimpleDelegate::CreateStatic(&FMaterialPropertyHelpers::ResetToDefault, StackParameterData->Parameter, MaterialEditorInstance);
 		FResetToDefaultOverride ResetOverride = FResetToDefaultOverride::Create(IsResetVisible, ResetHandler);
 		
 		if (ScalarParam && ScalarParam->AtlasData.bIsUsedAsAtlasPosition)
 		{
-			IsResetVisible = FIsResetToDefaultVisible::CreateStatic(&FMaterialPropertyHelpers::ShouldShowResetToDefault, StackParameterData->Parameter, MaterialEditorInstance);
-			ResetHandler = FResetToDefaultHandler::CreateStatic(&FMaterialPropertyHelpers::ResetCurveToDefault, StackParameterData->Parameter, MaterialEditorInstance);
+			IsResetVisible = TAttribute<bool>::Create(TAttribute<bool>::FGetter::CreateStatic(&FMaterialPropertyHelpers::ShouldShowResetToDefault, StackParameterData->Parameter, MaterialEditorInstance));
+			ResetHandler = FSimpleDelegate::CreateStatic(&FMaterialPropertyHelpers::ResetCurveToDefault, StackParameterData->Parameter, MaterialEditorInstance);
 			ResetOverride = FResetToDefaultOverride::Create(IsResetVisible, ResetHandler);
 		}
 
@@ -1222,9 +1222,9 @@ void SMaterialLayersFunctionsInstanceTree::RefreshOnAssetChange(const struct FAs
 	RequestTreeRefresh();
 }
 
-void SMaterialLayersFunctionsInstanceTree::ResetAssetToDefault(TSharedPtr<IPropertyHandle> InHandle, TSharedPtr<FSortedParamData> InData)
+void SMaterialLayersFunctionsInstanceTree::ResetAssetToDefault(TSharedPtr<FSortedParamData> InData)
 {
-	FMaterialPropertyHelpers::ResetLayerAssetToDefault(FunctionInstanceHandle.ToSharedRef(), InData->Parameter, InData->ParameterInfo.Association, InData->ParameterInfo.Index, MaterialEditorInstance);
+	FMaterialPropertyHelpers::ResetLayerAssetToDefault(InData->Parameter, InData->ParameterInfo.Association, InData->ParameterInfo.Index, MaterialEditorInstance);
 	UpdateThumbnailMaterial(InData->ParameterInfo.Association, InData->ParameterInfo.Index, false);
 	CreateGroupsWidget();
 	RequestTreeRefresh();
diff --git a/Engine/Source/Editor/MaterialEditor/Private/SMaterialLayersFunctionsTree.h b/Engine/Source/Editor/MaterialEditor/Private/SMaterialLayersFunctionsTree.h
--- a/Engine/Source/Editor/MaterialEditor/Private/SMaterialLayersFunctionsTree.h
+++ b/Engine/Source/Editor/MaterialEditor/Private/SMaterialLayersFunctionsTree.h
@@ -142,7 +142,7 @@ public:
 	struct FMaterialLayersFunctions* FunctionInstance;
 	TSharedPtr<IPropertyHandle> FunctionInstanceHandle;
 	void RefreshOnAssetChange(const struct FAssetData& InAssetData, int32 Index, EMaterialParameterAssociation MaterialType);
-	void ResetAssetToDefault(TSharedPtr<IPropertyHandle> InHandle, TSharedPtr<FSortedParamData> InData);
+	void ResetAssetToDefault(TSharedPtr<FSortedParamData> InData);
 	void AddLayer();
 	void RemoveLayer(int32 Index);
 	FReply UnlinkLayer(int32 Index);
diff --git a/Engine/Source/Editor/MaterialEditor/Public/MaterialPropertyHelpers.h b/Engine/Source/Editor/MaterialEditor/Public/MaterialPropertyHelpers.h
--- a/Engine/Source/Editor/MaterialEditor/Public/MaterialPropertyHelpers.h
+++ b/Engine/Source/Editor/MaterialEditor/Public/MaterialPropertyHelpers.h
@@ -16,9 +16,9 @@
 struct FAssetData;
 class IDetailGroup;
 class IDetailLayoutBuilder;
+class IDetailTreeNode;
 class IPropertyHandle;
 class UDEditorParameterValue;
-class UMaterialEditorInstanceConstant;
 enum class ECheckBoxState : uint8;
 class UMaterialInterface;
 class SMaterialLayersFunctionsInstanceTreeItem;
@@ -43,7 +43,7 @@ public:
 	EStackDataType StackDataType;
 
 	UPROPERTY(Transient)
-	TObjectPtr<class UDEditorParameterValue> Parameter = nullptr;
+	TObjectPtr<UDEditorParameterValue> Parameter = nullptr;
 
 	FName PropertyName;
 
@@ -51,9 +51,9 @@ public:
 
 	FMaterialParameterInfo ParameterInfo;
 
-	TSharedPtr<class IDetailTreeNode> ParameterNode;
+	TSharedPtr<IDetailTreeNode> ParameterNode;
 
-	TSharedPtr<class IPropertyHandle> ParameterHandle;
+	TSharedPtr<IPropertyHandle> ParameterHandle;
 
 	TArray<TSharedPtr<struct FSortedParamData>> Children;
 
@@ -65,11 +65,11 @@ struct FUnsortedParamData
 {
 	GENERATED_USTRUCT_BODY()
 	UPROPERTY(Transient)
-	TObjectPtr<class UDEditorParameterValue> Parameter = nullptr;
+	TObjectPtr<UDEditorParameterValue> Parameter = nullptr;
 	FEditorParameterGroup ParameterGroup;
-	TSharedPtr<class IDetailTreeNode> ParameterNode;
+	TSharedPtr<IDetailTreeNode> ParameterNode;
 	FName UnsortedName;
-	TSharedPtr<class IPropertyHandle> ParameterHandle;
+	TSharedPtr<IPropertyHandle> ParameterHandle;
 };
 
 class SLayerHandle : public SCompoundWidget
@@ -140,33 +140,33 @@ class MATERIALEDITOR_API FMaterialPropertyHelpers
 {
 public:
 	/** Returns true if the parameter is being overridden */
-	static bool IsOverriddenExpression(class UDEditorParameterValue* Parameter);
+	static bool IsOverriddenExpression(UDEditorParameterValue* Parameter);
 	static bool IsOverriddenExpression(TObjectPtr<UDEditorParameterValue> Parameter) { return IsOverriddenExpression(Parameter.Get()); }
 	static ECheckBoxState IsOverriddenExpressionCheckbox(UDEditorParameterValue* Parameter);
 
 	/** Gets the expression description of this parameter from the base material */
-	static	FText GetParameterExpressionDescription(class UDEditorParameterValue* Parameter, UObject* MaterialEditorInstance);
+	static	FText GetParameterExpressionDescription(UDEditorParameterValue* Parameter, UObject* MaterialEditorInstance);
 	
-	static FText GetParameterTooltip(class UDEditorParameterValue* Parameter, UObject* MaterialEditorInstance);
+	static FText GetParameterTooltip(UDEditorParameterValue* Parameter, UObject* MaterialEditorInstance);
 	/**
 	 * Called when a parameter is overridden;
 	 */
-	static void OnOverrideParameter(bool NewValue, class UDEditorParameterValue* Parameter, UMaterialEditorInstanceConstant* MaterialEditorInstance);
+	static void OnOverrideParameter(bool NewValue, UDEditorParameterValue* Parameter, UMaterialEditorInstanceConstant* MaterialEditorInstance);
 
 	static EVisibility ShouldShowExpression(UDEditorParameterValue* Parameter, UMaterialEditorInstanceConstant* MaterialEditorInstance, FGetShowHiddenParameters ShowHiddenDelegate);
 	static EVisibility ShouldShowExpression(TObjectPtr<UDEditorParameterValue> Parameter, UMaterialEditorInstanceConstant* MaterialEditorInstance, FGetShowHiddenParameters ShowHiddenDelegate) { return ShouldShowExpression(Parameter.Get(), MaterialEditorInstance, ShowHiddenDelegate); }
 
 	/** Generic material property reset to default implementation.  Resets Parameter to default */
-	static void ResetToDefault(TSharedPtr<IPropertyHandle> PropertyHandle, class UDEditorParameterValue* Parameter, UMaterialEditorInstanceConstant* MaterialEditorInstance);
-	static void ResetToDefault(TSharedPtr<IPropertyHandle> PropertyHandle, TObjectPtr<UDEditorParameterValue> Parameter, UMaterialEditorInstanceConstant* MaterialEditorInstance) { ResetToDefault(PropertyHandle, Parameter.Get(), MaterialEditorInstance); }
-	static bool ShouldShowResetToDefault(TSharedPtr<IPropertyHandle> PropertyHandle, class UDEditorParameterValue* Parameter, UMaterialEditorInstanceConstant* MaterialEditorInstance);
-	static bool ShouldShowResetToDefault(TSharedPtr<IPropertyHandle> PropertyHandle, TObjectPtr<UDEditorParameterValue> Parameter, UMaterialEditorInstanceConstant* MaterialEditorInstance) { return ShouldShowResetToDefault(PropertyHandle, Parameter.Get(), MaterialEditorInstance); }
+	static void ResetToDefault(UDEditorParameterValue* Parameter, UMaterialEditorInstanceConstant* MaterialEditorInstance);
+	static void ResetToDefault(TObjectPtr<UDEditorParameterValue> Parameter, UMaterialEditorInstanceConstant* MaterialEditorInstance) { ResetToDefault(Parameter.Get(), MaterialEditorInstance); }
+	static bool ShouldShowResetToDefault(UDEditorParameterValue* Parameter, UMaterialEditorInstanceConstant* MaterialEditorInstance);
+	static bool ShouldShowResetToDefault(TObjectPtr<UDEditorParameterValue> Parameter, UMaterialEditorInstanceConstant* MaterialEditorInstance) { return ShouldShowResetToDefault(Parameter.Get(), MaterialEditorInstance); }
 	
 	/** Specific resets for layer and blend asses */
-	static void ResetLayerAssetToDefault(TSharedPtr<IPropertyHandle> PropertyHandle, class UDEditorParameterValue* InParameter, TEnumAsByte<EMaterialParameterAssociation> InAssociation, int32 Index, UMaterialEditorInstanceConstant* MaterialEditorInstance);
+	static void ResetLayerAssetToDefault(UDEditorParameterValue* InParameter, TEnumAsByte<EMaterialParameterAssociation> InAssociation, int32 Index, UMaterialEditorInstanceConstant* MaterialEditorInstance);
 	/** If reset to default button should show for a layer or blend asset*/
-	static bool ShouldLayerAssetShowResetToDefault(TSharedPtr<IPropertyHandle> PropertyHandle, TSharedPtr<struct FSortedParamData> InParameterData, UMaterialInterface* InMaterial);
-	static bool ShouldLayerAssetShowResetToDefault(TSharedPtr<IPropertyHandle> PropertyHandle, TSharedPtr<struct FSortedParamData> InParameterData, TObjectPtr<UMaterialInterface> InMaterial) { return ShouldLayerAssetShowResetToDefault(PropertyHandle, InParameterData, InMaterial.Get()); }
+	static bool ShouldLayerAssetShowResetToDefault(TSharedPtr<FSortedParamData> InParameterData, UMaterialInterface* InMaterial);
+	static bool ShouldLayerAssetShowResetToDefault(TSharedPtr<FSortedParamData> InParameterData, TObjectPtr<UMaterialInterface> InMaterial) { return ShouldLayerAssetShowResetToDefault(InParameterData, InMaterial.Get()); }
 
 	static void OnMaterialLayerAssetChanged(const struct FAssetData& InAssetData, int32 Index, EMaterialParameterAssociation MaterialType, TSharedPtr<class IPropertyHandle> InHandle, FMaterialLayersFunctions* InMaterialFunction);
 
@@ -175,15 +175,15 @@ public:
 	static FReply OnClickedSaveNewMaterialInstance(class UMaterialInterface* Object, UObject* EditorObject);
 	static FReply OnClickedSaveNewMaterialInstance(TObjectPtr<class UMaterialInterface> Object, UObject* EditorObject) { return OnClickedSaveNewMaterialInstance(Object.Get(), EditorObject); }
 
-	static void CopyMaterialToInstance(class UMaterialInstanceConstant* ChildInstance, TArray<struct FEditorParameterGroup> &ParameterGroups);
-	static void TransitionAndCopyParameters(class UMaterialInstanceConstant* ChildInstance, TArray<struct FEditorParameterGroup> &ParameterGroups, bool bForceCopy = false);
+	static void CopyMaterialToInstance(class UMaterialInstanceConstant* ChildInstance, TArray<FEditorParameterGroup> &ParameterGroups);
+	static void TransitionAndCopyParameters(class UMaterialInstanceConstant* ChildInstance, TArray<FEditorParameterGroup> &ParameterGroups, bool bForceCopy = false);
 	static FReply OnClickedSaveNewFunctionInstance(class UMaterialFunctionInterface* Object, class UMaterialInterface* PreviewMaterial, UObject* EditorObject);
 	static FReply OnClickedSaveNewLayerInstance(class UMaterialFunctionInterface* Object, TSharedPtr<FSortedParamData> InSortedData);
 
 	static void GetVectorChannelMaskComboBoxStrings(TArray<TSharedPtr<FString>>& OutComboBoxStrings, TArray<TSharedPtr<class SToolTip>>& OutToolTips, TArray<bool>& OutRestrictedItems);
-	static FString GetVectorChannelMaskValue(class UDEditorParameterValue* InParameter);
+	static FString GetVectorChannelMaskValue(UDEditorParameterValue* InParameter);
 	static FString GetVectorChannelMaskValue(TObjectPtr<UDEditorParameterValue> InParameter) { return GetVectorChannelMaskValue(InParameter.Get()); }
-	static void SetVectorChannelMaskValue(const FString& StringValue, TSharedPtr<IPropertyHandle> PropertyHandle, class UDEditorParameterValue* InParameter, UObject* MaterialEditorInstance);
+	static void SetVectorChannelMaskValue(const FString& StringValue, TSharedPtr<IPropertyHandle> PropertyHandle, UDEditorParameterValue* InParameter, UObject* MaterialEditorInstance);
 	static void SetVectorChannelMaskValue(const FString& StringValue, TSharedPtr<IPropertyHandle> PropertyHandle, TObjectPtr<UDEditorParameterValue> InParameter, UObject* MaterialEditorInstance) { SetVectorChannelMaskValue(StringValue, PropertyHandle, InParameter.Get(), MaterialEditorInstance); }
 
 	static TArray<class UFactory*> GetAssetFactories(EMaterialParameterAssociation AssetType);
@@ -192,16 +192,16 @@ public:
 	*
 	* @param ParameterGroup		Name to be looked for.
 	*/
-	static FEditorParameterGroup&  GetParameterGroup(class UMaterial* InMaterial, FName& ParameterGroup, TArray<struct FEditorParameterGroup>& ParameterGroups);
+	static FEditorParameterGroup&  GetParameterGroup(class UMaterial* InMaterial, FName& ParameterGroup, TArray<FEditorParameterGroup>& ParameterGroups);
 
 	static TSharedRef<SWidget> MakeStackReorderHandle(TSharedPtr<SMaterialLayersFunctionsInstanceTreeItem> InOwningStack);
 
-	static bool OnShouldSetCurveAsset(const FAssetData& AssetData, TSoftObjectPtr<class UCurveLinearColorAtlas> InAtlas);
-	static bool OnShouldFilterCurveAsset(const FAssetData& AssetData, TSoftObjectPtr<class UCurveLinearColorAtlas> InAtlas);
-	static void SetPositionFromCurveAsset(const FAssetData& AssetData, TSoftObjectPtr<class UCurveLinearColorAtlas> InAtlas, class UDEditorScalarParameterValue* InParameter, TSharedPtr<IPropertyHandle> PropertyHandle, UObject* MaterialEditorInstance);
+	static bool OnShouldSetCurveAsset(const FAssetData& AssetData, TSoftObjectPtr<UCurveLinearColorAtlas> InAtlas);
+	static bool OnShouldFilterCurveAsset(const FAssetData& AssetData, TSoftObjectPtr<UCurveLinearColorAtlas> InAtlas);
+	static void SetPositionFromCurveAsset(const FAssetData& AssetData, TSoftObjectPtr<UCurveLinearColorAtlas> InAtlas, class UDEditorScalarParameterValue* InParameter, TSharedPtr<IPropertyHandle> PropertyHandle, UObject* MaterialEditorInstance);
 
-	static void ResetCurveToDefault(TSharedPtr<IPropertyHandle> PropertyHandle, class UDEditorParameterValue* Parameter, UMaterialEditorInstanceConstant* MaterialEditorInstance);
-	static void ResetCurveToDefault(TSharedPtr<IPropertyHandle> PropertyHandle, TObjectPtr<UDEditorParameterValue> Parameter, UMaterialEditorInstanceConstant* MaterialEditorInstance) { ResetCurveToDefault(PropertyHandle, Parameter.Get(), MaterialEditorInstance); }
+	static void ResetCurveToDefault(UDEditorParameterValue* Parameter, UMaterialEditorInstanceConstant* MaterialEditorInstance);
+	static void ResetCurveToDefault(TObjectPtr<UDEditorParameterValue> Parameter, UMaterialEditorInstanceConstant* MaterialEditorInstance) { ResetCurveToDefault(Parameter.Get(), MaterialEditorInstance); }
 
 	static FText LayerID;
 	static FText BlendID;
diff --git a/Engine/Source/Editor/PropertyEditor/Public/IDetailPropertyRow.h b/Engine/Source/Editor/PropertyEditor/Public/IDetailPropertyRow.h
--- a/Engine/Source/Editor/PropertyEditor/Public/IDetailPropertyRow.h
+++ b/Engine/Source/Editor/PropertyEditor/Public/IDetailPropertyRow.h
@@ -20,7 +20,7 @@ public:
 	
 	FDetailColumnSizeData()
 	{
-		ValueColumnWidthValue = 0.5f;
+		ValueColumnWidthValue = 0.7f;
 		RightColumnWidthValue = 0.05f;
 		HoveredSplitterIndexValue = INDEX_NONE;
 
@@ -54,7 +54,7 @@ private:
 	float RightColumnWidthValue;
 	int HoveredSplitterIndexValue;
 
-	float GetNameColumnWidth() const { return 1.0f - ValueColumnWidthValue; }
+	float GetNameColumnWidth() const { return 1.0f - (ValueColumnWidthValue + RightColumnWidthValue); }
 	float GetValueColumnWidth() const { return ValueColumnWidthValue; }
 	float GetRightColumnWidth() const { return RightColumnWidthValue; }
 	float GetPropertyColumnWidth() const { return 1.0f - RightColumnWidthValue; }
@@ -133,6 +133,17 @@ public:
 		return ResetToDefault;
 	}
 
+	/** Create a FResetToDefaultOverride from a simple delegate. */
+	static FResetToDefaultOverride Create(FSimpleDelegate InResetToDefaultClicked, const bool InPropagateToChildren = false)
+	{
+		FResetToDefaultOverride ResetToDefault;
+		ResetToDefault.bForceShow = true;
+		ResetToDefault.OnClickedDelegate = InResetToDefaultClicked;
+		ResetToDefault.bPropagateToChildren = InPropagateToChildren;
+		ResetToDefault.bForceHide = false;
+		return ResetToDefault;
+	}
+
 	/** Creates a FResetToDefaultOverride in which reset to default is never visible */
 	static FResetToDefaultOverride Hide(const bool InPropagateToChildren = false)
 	{
