Fixing material editor slowness with many parameters (UE-108596).

- Edigrating CL 15410279. The same change as 11660435 from a while ago.

diff --git a/Engine/Source/Editor/MaterialEditor/Private/SMaterialLayersFunctionsTree.cpp b/Engine/Source/Editor/MaterialEditor/Private/SMaterialLayersFunctionsTree.cpp
--- a/Engine/Source/Editor/MaterialEditor/Private/SMaterialLayersFunctionsTree.cpp
+++ b/Engine/Source/Editor/MaterialEditor/Private/SMaterialLayersFunctionsTree.cpp
@@ -1373,6 +1373,8 @@ void SMaterialLayersFunctionsInstanceTree::CreateGroupsWidget()
 
 	Children.Empty();
 	ParameterGroups->GetChildren(Children);
+	// the order of DeferredSearches should correspond to NonLayerProperties exactly
+	TArray<TSharedPtr<IPropertyHandle>> DeferredSearches;
 	for (int32 GroupIdx = 0; GroupIdx < Children.Num(); ++GroupIdx)
 	{
 		TArray<void*> GroupPtrs;
@@ -1526,14 +1528,29 @@ void SMaterialLayersFunctionsInstanceTree::CreateGroupsWidget()
 					
 				NonLayerProperty.Parameter = Parameter;
 				NonLayerProperty.ParameterGroup = ParameterGroup;
-				NonLayerProperty.ParameterNode = Generator->FindTreeNode(ParameterValueProperty);
-				NonLayerProperty.ParameterHandle = NonLayerProperty.ParameterNode->CreatePropertyHandle();
+
+				DeferredSearches.Add(ParameterValueProperty);
 				NonLayerProperty.UnsortedName = Parameter->ParameterInfo.Name;
 
 				NonLayerProperties.Add(NonLayerProperty);
 			}
 		}
 	}
+
+	checkf(NonLayerProperties.Num() == DeferredSearches.Num(), TEXT("Internal inconsistency: number of node searches does not match the number of properties"));
+	TArray<TSharedPtr<IDetailTreeNode>> DeferredResults = Generator->FindTreeNodes(DeferredSearches);
+	checkf(NonLayerProperties.Num() == DeferredResults.Num(), TEXT("Internal inconsistency: number of node search results does not match the number of properties"));
+
+	for (int Idx = 0, NumUnsorted = NonLayerProperties.Num(); Idx < NumUnsorted; ++Idx)
+	{
+		FUnsortedParamData& NonLayerProperty = NonLayerProperties[Idx];
+		NonLayerProperty.ParameterNode = DeferredResults[Idx];
+		NonLayerProperty.ParameterHandle = NonLayerProperty.ParameterNode->CreatePropertyHandle();
+	}
+
+	DeferredResults.Empty();
+	DeferredSearches.Empty();
+
 	for (int32 LayerIdx = 0; LayerIdx < LayerProperties.Num(); LayerIdx++)
 	{
 		for (int32 ChildIdx = 0; ChildIdx < LayerProperties[LayerIdx]->Children.Num(); ChildIdx++)
@@ -2719,6 +2736,8 @@ void SMaterialLayersFunctionsMaterialTree::CreateGroupsWidget()
 
 	Children.Empty();
 	ParameterGroups->GetChildren(Children);
+	// the order should correspond to NonLayerProperty exactly
+	TArray<TSharedPtr<IPropertyHandle>> DeferredSearches;
 	for (int32 GroupIdx = 0; GroupIdx < Children.Num(); ++GroupIdx)
 	{
 		TArray<void*> GroupPtrs;
@@ -2872,14 +2891,29 @@ void SMaterialLayersFunctionsMaterialTree::CreateGroupsWidget()
 
 				NonLayerProperty.Parameter = Parameter;
 				NonLayerProperty.ParameterGroup = ParameterGroup;
-				NonLayerProperty.ParameterNode = Wrapper->GetGenerator()->FindTreeNode(ParameterValueProperty);
-				NonLayerProperty.ParameterHandle = NonLayerProperty.ParameterNode->CreatePropertyHandle();
+
+				DeferredSearches.Add(ParameterValueProperty);
 				NonLayerProperty.UnsortedName = Parameter->ParameterInfo.Name;
 
 				NonLayerProperties.Add(NonLayerProperty);
 			}
 		}
 	}
+
+	checkf(NonLayerProperties.Num() == DeferredSearches.Num(), TEXT("Internal inconsistency: number of node searches does not match the number of properties"));
+	TArray<TSharedPtr<IDetailTreeNode>> DeferredResults = Wrapper->GetGenerator()->FindTreeNodes(DeferredSearches);
+	checkf(NonLayerProperties.Num() == DeferredResults.Num(), TEXT("Internal inconsistency: number of node search results does not match the number of properties"));
+
+	for (int Idx = 0, NumUnsorted = NonLayerProperties.Num(); Idx < NumUnsorted; ++Idx)
+	{
+		FUnsortedParamData& NonLayerProperty = NonLayerProperties[Idx];
+		NonLayerProperty.ParameterNode = DeferredResults[Idx];
+		NonLayerProperty.ParameterHandle = NonLayerProperty.ParameterNode->CreatePropertyHandle();
+	}
+
+	DeferredResults.Empty();
+	DeferredSearches.Empty();
+
 	for (int32 LayerIdx = 0; LayerIdx < LayerProperties.Num(); LayerIdx++)
 	{
 		for (int32 ChildIdx = 0; ChildIdx < LayerProperties[LayerIdx]->Children.Num(); ChildIdx++)
