Material Editor Cleanup Pass

diff --git a/Engine/Source/Editor/MaterialEditor/Private/MaterialEditor.cpp b/Engine/Source/Editor/MaterialEditor/Private/MaterialEditor.cpp
--- a/Engine/Source/Editor/MaterialEditor/Private/MaterialEditor.cpp
+++ b/Engine/Source/Editor/MaterialEditor/Private/MaterialEditor.cpp
@@ -294,7 +294,7 @@ void FMaterialEditor::RegisterToolbarTab(const TSharedRef<class FTabManager>& In
 		.SetIcon(FSlateIcon(FEditorStyle::GetStyleSetName(), "LevelEditor.Tabs.Details"));
 
 	InTabManager->RegisterTabSpawner(FMaterialEditorTabs::ParameterDefaultsTabId, FOnSpawnTab::CreateSP(this, &FMaterialEditor::SpawnTab_ParameterDefaults))
-		.SetDisplayName(LOCTEXT("ParameterDefaultsTab", "Parameter Defaults"))
+		.SetDisplayName(LOCTEXT("ParametersTab", "Parameters"))
 		.SetGroup(WorkspaceMenuCategoryRef)
 		.SetIcon(FSlateIcon(FEditorStyle::GetStyleSetName(), "LevelEditor.Tabs.Details"));
 
@@ -1747,91 +1747,25 @@ void FMaterialEditor::RegisterToolBar()
 void FMaterialEditor::AddInheritanceMenuEntry(FToolMenuSection& Section, const FAssetData& AssetData, bool bIsFunctionPreviewMaterial)
 {
 	FExecuteAction OpenAction;
-	FExecuteAction FindInContentBrowserAction;
 	if (bIsFunctionPreviewMaterial)
 	{
 		OpenAction.BindStatic(&FMaterialEditorUtilities::OnOpenFunction, AssetData);
-		FindInContentBrowserAction.BindStatic(&FMaterialEditorUtilities::OnShowFunctionInContentBrowser, AssetData);
 	}
 	else
 	{
 		OpenAction.BindStatic(&FMaterialEditorUtilities::OnOpenMaterial, AssetData);
-		FindInContentBrowserAction.BindStatic(&FMaterialEditorUtilities::OnShowMaterialInContentBrowser, AssetData);
 	}
 
 	FFormatNamedArguments Args;
 	Args.Add(TEXT("ParentName"), FText::FromName(AssetData.AssetName));
 	FText Label = FText::Format(LOCTEXT("InstanceParentName", "{ParentName}"), Args);
 
-	FSlateIcon OpenIcon(FEditorStyle::GetStyleSetName(), "ContentBrowser.AssetActions.OpenInExternalEditor");
-	FSlateIcon FindInContentBrowserIcon(FEditorStyle::GetStyleSetName(), "SystemWideCommands.FindInContentBrowser");
-
-	const float MenuIconSize = FCoreStyle::Get().GetFloat("Menu.MenuIconSize", nullptr, 16.f);
-
-	TSharedRef<SWidget> EntryWidget =
-		SNew(SHorizontalBox)
-		.ToolTipText(LOCTEXT("OpenInEditor", "Open In Editor"))
-
-		+ SHorizontalBox::Slot()
-		.AutoWidth()
-		.Padding(FMargin(2, 0, 2, 0))
-		[
-			SNew( SBox )
-			.WidthOverride(MenuIconSize + 2 )
-			.HeightOverride(MenuIconSize)
-			.HAlign(HAlign_Center)
-			.VAlign(VAlign_Center)
-			[
-				SNew( SBox )
-				.WidthOverride(MenuIconSize)
-				.HeightOverride(MenuIconSize)
-				[
-					SNew(SImage)
-					.Image(OpenIcon.GetIcon())
-				]
-			]
-		]
-
-		+ SHorizontalBox::Slot()
-		.FillWidth( 1.0f )
-		.Padding(FMargin(2, 0, 6, 0))
-		.VAlign( VAlign_Center )
-		[
-			SNew(STextBlock)
-			.Text(Label)
-		]
-
-		+ SHorizontalBox::Slot()
-		.AutoWidth()
-		.VAlign( VAlign_Center )
-		.HAlign( HAlign_Right )
-		[
-			SNew(SButton)
-			.ButtonStyle(FEditorStyle::Get(), "HoverHintOnly")
-			.ToolTipText(LOCTEXT("FindInContentBrowser", "Find In Content Browser"))
-			.OnClicked_Lambda([FindInContentBrowserAction]() { FindInContentBrowserAction.ExecuteIfBound(); return FReply::Handled(); })
-			[
-				SNew( SBox )
-				.WidthOverride(MenuIconSize + 2 )
-				.HeightOverride(MenuIconSize)
-				.HAlign(HAlign_Center)
-				.VAlign(VAlign_Center)
-				[
-					SNew( SBox )
-					.WidthOverride(MenuIconSize)
-					.HeightOverride(MenuIconSize)
-					[
-						SNew(SImage)
-						.Image(FindInContentBrowserIcon.GetIcon())
-					]
-				]
-			]
-		];
-
 	Section.AddEntry(FToolMenuEntry::InitMenuEntry(
 		NAME_None,
-		FUIAction(OpenAction),
-		EntryWidget
+		Label,
+		LOCTEXT("OpenInEditor", "Open In Editor"),
+		FSlateIcon(),
+		FUIAction(OpenAction)
 	));
 }
 
@@ -4727,7 +4661,7 @@ TSharedRef<SDockTab> FMaterialEditor::SpawnTab_ParameterDefaults(const FSpawnTab
 {
 	TSharedRef<SDockTab> SpawnedTab = SNew(SDockTab)
 		.Icon(FEditorStyle::GetBrush("LevelEditor.Tabs.Details"))
-		.Label(LOCTEXT("ParameterDefaults", "Parameter Defaults"))
+		.Label(LOCTEXT("Parameters", "Parameters"))
 		[
 			SNew(SBox)
 			[
diff --git a/Engine/Source/Editor/MaterialEditor/Private/MaterialEditorInstanceDetailCustomization.cpp b/Engine/Source/Editor/MaterialEditor/Private/MaterialEditorInstanceDetailCustomization.cpp
--- a/Engine/Source/Editor/MaterialEditor/Private/MaterialEditorInstanceDetailCustomization.cpp
+++ b/Engine/Source/Editor/MaterialEditor/Private/MaterialEditorInstanceDetailCustomization.cpp
@@ -344,58 +344,20 @@ void FMaterialInstanceParameterDetails::CreateGroupsWidget(TSharedRef<IPropertyH
 				.Padding(2.0f)
 				[
 					SNew(SButton)
-					.ButtonStyle(FEditorStyle::Get(), "FlatButton.Dark")
+					.Text(LOCTEXT("SaveSibling", "Save Sibling"))
 					.HAlign(HAlign_Center)
 					.OnClicked(OnSiblingButtonClicked)
-					.ToolTipText(LOCTEXT("SaveToSiblingInstance", "Save To Sibling Instance"))
-					.Content()
-					[
-						SNew(SHorizontalBox)
-						+ SHorizontalBox::Slot()
-						.AutoWidth()
-						[
-							SNew(STextBlock)
-							.Font(FEditorStyle::Get().GetFontStyle("FontAwesome.10"))
-							.TextStyle(FEditorStyle::Get(), "NormalText.Important")
-							.Text(FText::FromString(FString(TEXT("\xf0c7 \xf178"))) /*fa-filter*/)
-						]
-						+ SHorizontalBox::Slot()
-						.AutoWidth()
-						[
-							SNew(STextBlock)
-							.TextStyle(FEditorStyle::Get(), "NormalText.Important")
-							.Text(FText::FromString(FString(TEXT(" Save Sibling"))) /*fa-filter*/)
-						]
-					]
+					.ToolTipText(LOCTEXT("SaveToSiblingInstance", "Save to Sibling Instance"))
 				]
 				+ SHorizontalBox::Slot()
 				.AutoWidth()
 				.Padding(2.0f)
 				[
 					SNew(SButton)
-					.ButtonStyle(FEditorStyle::Get(), "FlatButton.Dark")
+					.Text(LOCTEXT("SaveChild", "Save Child"))
 					.HAlign(HAlign_Center)
 					.OnClicked(OnChildButtonClicked)
-					.ToolTipText(LOCTEXT("SaveToChildInstance", "Save To Child Instance"))
-					.Content()
-					[
-						SNew(SHorizontalBox)
-						+ SHorizontalBox::Slot()
-						.AutoWidth()
-						[
-							SNew(STextBlock)
-							.Font(FEditorStyle::Get().GetFontStyle("FontAwesome.10"))
-							.TextStyle(FEditorStyle::Get(), "NormalText.Important")
-							.Text(FText::FromString(FString(TEXT("\xf0c7 \xf149"))) /*fa-filter*/)
-						]
-						+ SHorizontalBox::Slot()
-						.AutoWidth()
-						[
-							SNew(STextBlock)
-							.TextStyle(FEditorStyle::Get(), "NormalText.Important")
-							.Text(FText::FromString(FString(TEXT(" Save Child"))) /*fa-filter*/)
-						]
-					]
+					.ToolTipText(LOCTEXT("SaveToChildInstance", "Save to Child Instance"))
 				]
 			];
 	}
diff --git a/Engine/Source/Editor/MaterialEditor/Private/MaterialEditorModes.cpp b/Engine/Source/Editor/MaterialEditor/Private/MaterialEditorModes.cpp
--- a/Engine/Source/Editor/MaterialEditor/Private/MaterialEditorModes.cpp
+++ b/Engine/Source/Editor/MaterialEditor/Private/MaterialEditorModes.cpp
@@ -19,7 +19,7 @@ const FName FMaterialEditorApplicationModes::StandardMaterialEditorMode(TEXT("Gr
 
 TSharedPtr<FTabManager::FLayout> FMaterialEditorApplicationModes::GetDefaultEditorLayout(TSharedPtr<FMaterialEditor> InMaterialEditor)
 {
-	const TSharedRef<FTabManager::FLayout> StandaloneDefaultLayout = FTabManager::NewLayout("Standalone_MaterialEditor_Layout_v12")
+	const TSharedRef<FTabManager::FLayout> StandaloneDefaultLayout = FTabManager::NewLayout("Standalone_MaterialEditor_Layout_v13")
 	->AddArea
 	(
 		FTabManager::NewPrimaryArea() ->SetOrientation(Orient_Vertical)
@@ -72,7 +72,7 @@ TSharedPtr<FTabManager::FLayout> FMaterialEditorApplicationModes::GetDefaultEdit
 				->Split
 				(
 					FTabManager::NewStack()
-					->AddTab( FMaterialEditorTabs::PaletteTabId, ETabState::OpenedTab )
+					->AddTab( FMaterialEditorTabs::PaletteTabId, ETabState::SidebarTab, ESidebarLocation::Right, 0.14f )
 				)
 			)
 		)
diff --git a/Engine/Source/Editor/MaterialEditor/Private/SMaterialLayersFunctionsTree.cpp b/Engine/Source/Editor/MaterialEditor/Private/SMaterialLayersFunctionsTree.cpp
--- a/Engine/Source/Editor/MaterialEditor/Private/SMaterialLayersFunctionsTree.cpp
+++ b/Engine/Source/Editor/MaterialEditor/Private/SMaterialLayersFunctionsTree.cpp
@@ -42,6 +42,8 @@
 #include "EditorFontGlyphs.h"
 #include "SResetToDefaultPropertyEditor.h"
 #include "ThumbnailRendering/ThumbnailManager.h"
+#include "Styling/AppStyle.h"
+#include "Styling/StyleColors.h"
 
 #define LOCTEXT_NAMESPACE "MaterialLayerCustomization"
 
@@ -53,32 +55,32 @@ FString SMaterialLayersFunctionsInstanceTreeItem::GetCurvePath(UDEditorScalarPar
 
 const FSlateBrush* SMaterialLayersFunctionsInstanceTreeItem::GetBorderImage() const
 {
-	if (StackParameterData->StackDataType == EStackDataType::Stack)
+	return FAppStyle::Get().GetBrush("DetailsView.CategoryMiddle");
+}
+
+FSlateColor SMaterialLayersFunctionsInstanceTreeItem::GetOuterBackgroundColor(TSharedPtr<FSortedParamData> InParamData) const
+{
+	if (InParamData->StackDataType == EStackDataType::Stack)
 	{
 		if (bIsBeingDragged)
 		{
-			return FEditorStyle::GetBrush("MaterialInstanceEditor.StackBodyDragged");
+			return FAppStyle::Get().GetSlateColor("Colors.Recessed");
 		}
 		else if (bIsHoveredDragTarget)
 		{
-			return FEditorStyle::GetBrush("MaterialInstanceEditor.StackBody_Highlighted");
+			return FAppStyle::Get().GetSlateColor("Colors.Highlight");
 		}
 		else
 		{
-			return FEditorStyle::GetBrush("MaterialInstanceEditor.StackHeader");
+			return FAppStyle::Get().GetSlateColor("Colors.Header");
 		}
 	}
-	else
+	else if (IsHovered() || InParamData->StackDataType == EStackDataType::Group)
 	{
-		if (bIsHoveredDragTarget)
-		{
-			return FEditorStyle::GetBrush("MaterialInstanceEditor.StackBody_Highlighted");
-		}
-		else
-		{
-			return FEditorStyle::GetBrush("MaterialInstanceEditor.StackBody");
-		}
+		return FAppStyle::Get().GetSlateColor("Colors.Header");
 	}
+
+	return FAppStyle::Get().GetSlateColor("Colors.Panel");
 }
 
 void SMaterialLayersFunctionsInstanceTreeItem::RefreshOnRowChange(const FAssetData& AssetData, SMaterialLayersFunctionsInstanceTree* InTree)
@@ -242,12 +244,12 @@ void SMaterialLayersFunctionsInstanceTreeItem::Construct(const FArguments& InArg
 	TSharedRef<SWidget> ResetWidget = SNullWidget::NullWidget;
 	FText NameOverride;
 	TSharedRef<SVerticalBox> WrapperWidget = SNew(SVerticalBox);
-
+	EHorizontalAlignment ValueAlignment = HAlign_Left;
 // STACK --------------------------------------------------
 	if (StackParameterData->StackDataType == EStackDataType::Stack)
 	{
 		WrapperWidget->AddSlot()
-			.Padding(3.0f)
+			.Padding(0.0f)
 			.AutoHeight()
 			[
 				SNullWidget::NullWidget
@@ -278,7 +280,7 @@ void SMaterialLayersFunctionsInstanceTreeItem::Construct(const FArguments& InArg
 			HeaderRowWidget->AddSlot()
 				.HAlign(HAlign_Center)
 				.VAlign(VAlign_Center)
-				.Padding(2.5f, 0)
+				.Padding(0.0f)
 				.AutoWidth()
 				[
 					SNullWidget::NullWidget
@@ -337,11 +339,8 @@ void SMaterialLayersFunctionsInstanceTreeItem::Construct(const FArguments& InArg
 				.Padding(5.0f)
 				[
 					SNew(SEditableTextBox)
-					.BackgroundColor(FLinearColor(0.045f, 0.045f, 0.045f, 1.0f))
 					.Text(TAttribute<FText>::Create(TAttribute<FText>::FGetter::CreateSP(this, &SMaterialLayersFunctionsInstanceTreeItem::GetLayerName, InArgs._InTree, StackParameterData->ParameterInfo.Index)))
 					.OnTextCommitted(FOnTextCommitted::CreateSP(this, &SMaterialLayersFunctionsInstanceTreeItem::OnNameChanged, InArgs._InTree, StackParameterData->ParameterInfo.Index))
-					.Font(FEditorStyle::GetFontStyle(TEXT("MaterialEditor.Layers.EditableFontImportant")))
-					.ForegroundColor(FLinearColor::White)
 				];
 		}
 		else
@@ -353,7 +352,6 @@ void SMaterialLayersFunctionsInstanceTreeItem::Construct(const FArguments& InArg
 				[
 					SNew(STextBlock)
 					.Text(NameOverride)
-					.TextStyle(FEditorStyle::Get(), "NormalText.Important")
 				];
 		}
 
@@ -370,17 +368,11 @@ void SMaterialLayersFunctionsInstanceTreeItem::Construct(const FArguments& InArg
 			.Padding(0.0f, 0.0f, 0.0f, 0.0f)
 			[
 				SNew(SButton)
-				.ButtonStyle(FEditorStyle::Get(), "HoverHintOnly")
+				.Text(LOCTEXT("Unlink", "Unlink"))
+				.HAlign(HAlign_Center)
 				.OnClicked(Tree, &SMaterialLayersFunctionsInstanceTree::UnlinkLayer, StackParameterData->ParameterInfo.Index)
-				.Visibility(Tree, &SMaterialLayersFunctionsInstanceTree::GetUnlinkLayerVisibility, StackParameterData->ParameterInfo.Index)
 				.ToolTipText(LOCTEXT("UnlinkLayer", "Whether or not to unlink this layer/blend combination from the parent."))
-				.Content()
-				[
-					SNew(STextBlock)
-					.TextStyle(FEditorStyle::Get(), "ContentBrowser.TopBar.Font")
-					.Font(FEditorStyle::Get().GetFontStyle("FontAwesome.10"))
-					.Text(FEditorFontGlyphs::Chain_Broken) /*fa-filter*/
-				]
+				.Visibility(Tree, &SMaterialLayersFunctionsInstanceTree::GetUnlinkLayerVisibility, StackParameterData->ParameterInfo.Index)
 			];
 
 		// Can only remove layers that aren't the base layer.
@@ -404,7 +396,9 @@ void SMaterialLayersFunctionsInstanceTreeItem::Construct(const FArguments& InArg
 		NameOverride = FText::FromName(StackParameterData->Group.GroupName);
 		LeftSideWidget = SNew(STextBlock)
 			.Text(NameOverride)
-			.TextStyle(FEditorStyle::Get(), "TinyText");
+			.Font(FAppStyle::Get().GetFontStyle("PropertyWindow.BoldFont"))
+			.TextStyle(FAppStyle::Get(), "DetailsView.CategoryTextStyle")
+			.TransformPolicy(ETextTransformPolicy::ToUpper);
 		const int32 LayerStateIndex = StackParameterData->ParameterInfo.Association == EMaterialParameterAssociation::BlendParameter ? StackParameterData->ParameterInfo.Index + 1 : StackParameterData->ParameterInfo.Index;
 		LeftSideWidget->SetEnabled(InArgs._InTree->FunctionInstance->LayerStates[LayerStateIndex]);
 		RightSideWidget->SetEnabled(InArgs._InTree->FunctionInstance->LayerStates[LayerStateIndex]);
@@ -414,6 +408,7 @@ void SMaterialLayersFunctionsInstanceTreeItem::Construct(const FArguments& InArg
 // ASSET --------------------------------------------------
 	if (StackParameterData->StackDataType == EStackDataType::Asset)
 	{
+		ValueAlignment = HAlign_Fill;
 		FOnSetObject ObjectChanged = FOnSetObject::CreateSP(this, &SMaterialLayersFunctionsInstanceTreeItem::RefreshOnRowChange, Tree);
 		StackParameterData->ParameterHandle->GetProperty()->SetMetaData(FName(TEXT("DisplayThumbnail")), TEXT("true"));
 		FIntPoint ThumbnailOverride;
@@ -501,22 +496,21 @@ void SMaterialLayersFunctionsInstanceTreeItem::Construct(const FArguments& InArg
 					.NewAssetFactories(FMaterialPropertyHelpers::GetAssetFactories(InAssociation))
 				]
 				+ SHorizontalBox::Slot()
-				.Padding(0.0f, 2.0f, 0.0f, 0.0f)
 				.AutoWidth()
 				.VAlign(VAlign_Center)
+				.Padding(FMargin(0.0f, 0.0f, 2.0f, 0.0f))
 				[
 					SNew(SCheckBox)
 					.Type(ESlateCheckBoxType::ToggleButton)
-					.Style(&FCoreStyle::Get().GetWidgetStyle< FCheckBoxStyle >("ToggleButtonCheckbox"))
+					.Style(&FAppStyle::Get().GetWidgetStyle<FCheckBoxStyle>("ToggleButtonCheckboxAlt"))
 					.OnCheckStateChanged(this, &SMaterialLayersFunctionsInstanceTreeItem::FilterClicked, InArgs._InTree, StackParameterData)
 					.IsChecked(this, &SMaterialLayersFunctionsInstanceTreeItem::GetFilterChecked, InArgs._InTree, StackParameterData)
 					.ToolTipText(LOCTEXT("FilterLayerAssets", "Filter asset picker to only show related layers or blends. \nStaying within the inheritance hierarchy can improve instruction count."))
+					.Padding(2.0f)
 					.Content()
 					[
-						SNew(STextBlock)
-						.TextStyle(FEditorStyle::Get(), "ContentBrowser.TopBar.Font")
-						.Font(FEditorStyle::Get().GetFontStyle("FontAwesome.10"))
-						.Text(FText::FromString(FString(TEXT("\xf0b0"))) /*fa-filter*/)
+						SNew(SImage)
+						.Image(FAppStyle::Get().GetBrush("Icons.Filter"))
 					]
 				]
 			]
@@ -542,29 +536,10 @@ void SMaterialLayersFunctionsInstanceTreeItem::Construct(const FArguments& InArg
 			.Padding(2.0f)
 			[
 				SNew(SButton)
-				.ButtonStyle(FEditorStyle::Get(), "FlatButton.Dark")
+				.Text(LOCTEXT("SaveChild", "Save Child"))
 				.HAlign(HAlign_Center)
 				.OnClicked(OnChildButtonClicked)
 				.ToolTipText(LOCTEXT("SaveToChildInstance", "Save To Child Instance"))
-				.Content()
-				[
-					SNew(SHorizontalBox)
-					+ SHorizontalBox::Slot()
-					.AutoWidth()
-					[
-						SNew(STextBlock)
-						.Font(FEditorStyle::Get().GetFontStyle("FontAwesome.10"))
-						.TextStyle(FEditorStyle::Get(), "NormalText.Important")
-						.Text(FText::FromString(FString(TEXT("\xf0c7 \xf149"))) /*fa-filter*/)
-					]
-					+ SHorizontalBox::Slot()
-					.AutoWidth()
-					[
-						SNew(STextBlock)
-						.TextStyle(FEditorStyle::Get(), "NormalText.Important")
-						.Text(FText::FromString(FString(TEXT(" Save Child"))) /*fa-filter*/)
-					]
-				]
 			];
 			
 		LeftSideWidget->SetEnabled(InArgs._InTree->FunctionInstance->LayerStates[LayerStateIndex]);
@@ -631,7 +606,6 @@ void SMaterialLayersFunctionsInstanceTreeItem::Construct(const FArguments& InArg
 				[
 					SNew(SHorizontalBox)
 					+ SHorizontalBox::Slot()
-					.HAlign(HAlign_Left)
 					.AutoWidth()
 					[
 						PropertyCustomizationHelpers::MakePropertyComboBox(StackParameterData->ParameterNode->CreatePropertyHandle(), GetMaskStrings, GetMaskValue, SetMaskValue)
@@ -934,6 +908,10 @@ void SMaterialLayersFunctionsInstanceTreeItem::Construct(const FArguments& InArg
 		LeftSideWidget = NodeWidgets.NameWidget.ToSharedRef();
 		RightSideWidget = NodeWidgets.ValueWidget.ToSharedRef();
 
+		ResetWidget = SNew(SResetToDefaultPropertyEditor, Node.CreatePropertyHandle())
+			.IsEnabled(IsParamEnabled)
+			.CustomResetToDefault(ResetOverride);
+
 		StackParameterData->ParameterNode->CreatePropertyHandle()->SetOnPropertyValueChanged(FSimpleDelegate::CreateSP(Tree, &SMaterialLayersFunctionsInstanceTree::UpdateThumbnailMaterial, StackParameterData->ParameterInfo.Association, StackParameterData->ParameterInfo.Index, false));
 		StackParameterData->ParameterNode->CreatePropertyHandle()->SetOnChildPropertyValueChanged(FSimpleDelegate::CreateSP(Tree, &SMaterialLayersFunctionsInstanceTree::UpdateThumbnailMaterial, StackParameterData->ParameterInfo.Association, StackParameterData->ParameterInfo.Index, false));
 
@@ -956,7 +934,6 @@ void SMaterialLayersFunctionsInstanceTreeItem::Construct(const FArguments& InArg
 // END PROPERTY CHILD
 
 // FINAL WRAPPER
-	float ValuePadding = bisPaddedProperty ? 20.0f : 0.0f;
 	if (StackParameterData->StackDataType == EStackDataType::Stack)
 	{
 		TSharedPtr<SHorizontalBox> FinalStack;
@@ -964,10 +941,16 @@ void SMaterialLayersFunctionsInstanceTreeItem::Construct(const FArguments& InArg
 			.AutoHeight()
 			[
 				SNew(SBorder)
-				.BorderImage(this, &SMaterialLayersFunctionsInstanceTreeItem::GetBorderImage)
-				.Padding(0.0f)
+				.BorderImage(FAppStyle::Get().GetBrush("DetailsView.GridLine"))
+				.Padding(FMargin(0, 0, 0, 1))
 				[
-					SAssignNew(FinalStack, SHorizontalBox)
+					SNew(SBorder)
+					.BorderImage(this, &SMaterialLayersFunctionsInstanceTreeItem::GetBorderImage)
+					.BorderBackgroundColor(this, &SMaterialLayersFunctionsInstanceTreeItem::GetOuterBackgroundColor, StackParameterData)
+					.Padding(0.0f)
+					[
+						SAssignNew(FinalStack, SHorizontalBox)
+					]
 				]
 			];
 		if (StackParameterData->ParameterInfo.Index != 0)
@@ -995,48 +978,19 @@ void SMaterialLayersFunctionsInstanceTreeItem::Construct(const FArguments& InArg
 				LeftSideWidget
 			];
 	}
-	else
+	else if (StackParameterData->StackDataType == EStackDataType::Group)
 	{
-		FSlateBrush* StackBrush;
-		switch (StackParameterData->ParameterInfo.Association)
-		{
-		case EMaterialParameterAssociation::LayerParameter:
-		{
-			StackBrush = const_cast<FSlateBrush*>(FEditorStyle::GetBrush("MaterialInstanceEditor.StackBody"));
-			break;
-		}
-		case EMaterialParameterAssociation::BlendParameter:
-		{
-			StackBrush = const_cast<FSlateBrush*>(FEditorStyle::GetBrush("MaterialInstanceEditor.StackBodyBlend"));
-			break;
-		}
-		default:
-			break;
-		}
-
-		if (ResetWidget == SNullWidget::NullWidget)
-		{
-			const FSlateBrush* DiffersFromDefaultBrush = FEditorStyle::GetBrush("PropertyWindow.DiffersFromDefault");
-			ResetWidget = SNew(SSpacer)
-				.Size(DiffersFromDefaultBrush != nullptr ? DiffersFromDefaultBrush->ImageSize : FVector2D(8.0f, 8.0f));
-		}
-
 		WrapperWidget->AddSlot()
 			.AutoHeight()
 			[
 				SNew(SBorder)
-				.BorderImage(this, &SMaterialLayersFunctionsInstanceTreeItem::GetBorderImage)
-				.Padding(0.0f)
+				.BorderImage(FAppStyle::Get().GetBrush("DetailsView.GridLine"))
+				.Padding(FMargin(0, 0, 0, 1))
 				[
-					SNew(SSplitter)
-					.Style(FEditorStyle::Get(), "DetailsView.Splitter")
-					.PhysicalSplitterHandleSize(1.0f)
-					.HitDetectionSplitterHandleSize(5.0f)
-					.HighlightedHandleIndex(Tree->ColumnSizeData.HoveredSplitterIndex)
-					+ SSplitter::Slot()
-					.Value(Tree->ColumnSizeData.NameColumnWidth)
-					.OnSlotResized(Tree->ColumnSizeData.OnNameColumnResized)
-					.Value(0.25f)
+					SNew(SBorder)
+					.Padding(0.0f)
+					.BorderImage(this, &SMaterialLayersFunctionsInstanceTreeItem::GetBorderImage)
+					.BorderBackgroundColor(this, &SMaterialLayersFunctionsInstanceTreeItem::GetOuterBackgroundColor, StackParameterData)
 					[
 						SNew(SHorizontalBox)
 						+ SHorizontalBox::Slot()
@@ -1053,28 +1007,79 @@ void SMaterialLayersFunctionsInstanceTreeItem::Construct(const FArguments& InArg
 							LeftSideWidget
 						]
 					]
-					+ SSplitter::Slot()
-					.Value(Tree->ColumnSizeData.ValueColumnWidth)
-					.OnSlotResized(Tree->ColumnSizeData.OnValueColumnResized)
+				]
+			];
+	}
+	else
+	{
+		if (ResetWidget == SNullWidget::NullWidget)
+		{
+			const FSlateBrush* DiffersFromDefaultBrush = FEditorStyle::GetBrush("PropertyWindow.DiffersFromDefault");
+			ResetWidget = SNew(SSpacer)
+				.Size(DiffersFromDefaultBrush != nullptr ? DiffersFromDefaultBrush->ImageSize : FVector2D(8.0f, 8.0f));
+		}
+
+		WrapperWidget->AddSlot()
+			.AutoHeight()
+			[
+				SNew(SBorder)
+				.BorderImage(FAppStyle::Get().GetBrush("DetailsView.GridLine"))
+				.Padding(FMargin(0, 0, 0, 1))
+				[
+					SNew(SBorder)
+					.Padding(0.0f)
+					.BorderImage(this, &SMaterialLayersFunctionsInstanceTreeItem::GetBorderImage)
+					.BorderBackgroundColor(this, &SMaterialLayersFunctionsInstanceTreeItem::GetOuterBackgroundColor, StackParameterData)
 					[
-						SNew(SHorizontalBox)
-						+ SHorizontalBox::Slot()
-						.MaxWidth(350.0f - ValuePadding)
-						.Padding(FMargin(5.0f, 2.0f, ValuePadding, 2.0f))
+						SNew(SSplitter)
+						.Style(FEditorStyle::Get(), "DetailsView.Splitter")
+						.PhysicalSplitterHandleSize(1.0f)
+						.HitDetectionSplitterHandleSize(5.0f)
+						.HighlightedHandleIndex(Tree->ColumnSizeData.HoveredSplitterIndex)
+						+ SSplitter::Slot()
+						.Value(Tree->ColumnSizeData.NameColumnWidth)
+						.OnSlotResized(Tree->ColumnSizeData.OnNameColumnResized)
+						.Value(0.25f)
+						[
+							SNew(SHorizontalBox)
+							+ SHorizontalBox::Slot()
+							.AutoWidth()
+							.VAlign(VAlign_Center)
+							.Padding(FMargin(3.0f))
+							[
+								SNew(SExpanderArrow, SharedThis(this))
+							]
+							+ SHorizontalBox::Slot()
+							.Padding(FMargin(2.0f))
+							.VAlign(VAlign_Center)
+							[
+								LeftSideWidget
+							]
+						]
+						+ SSplitter::Slot()
+						.Value(Tree->ColumnSizeData.ValueColumnWidth)
+						.OnSlotResized(Tree->ColumnSizeData.OnValueColumnResized)
 						[
-							RightSideWidget
+							SNew(SHorizontalBox)
+							.Clipping(EWidgetClipping::OnDemand)
+							+ SHorizontalBox::Slot()
+							.HAlign(ValueAlignment)
+							.VAlign(VAlign_Center)
+							[
+								RightSideWidget
+							]
 						]
-					]
-					+ SSplitter::Slot()
-					.SizeRule(SSplitter::ESizeRule::SizeToContent)
-					[
-						SNew(SHorizontalBox)
-						+ SHorizontalBox::Slot()
-						.VAlign(VAlign_Center)
-						.HAlign(HAlign_Center)
-						.Padding(5.0f)
+						+ SSplitter::Slot()
+						.SizeRule(SSplitter::ESizeRule::SizeToContent)
 						[
-							ResetWidget
+							SNew(SHorizontalBox)
+							+ SHorizontalBox::Slot()
+							.VAlign(VAlign_Center)
+							.HAlign(HAlign_Center)
+							.Padding(5.0f)
+							[
+								ResetWidget
+							]
 						]
 					]
 				]
@@ -1737,32 +1742,27 @@ void SMaterialLayersFunctionsInstanceWrapper::Refresh()
 
 		this->ChildSlot
 			[
-				SNew(SBorder)
-				.BorderImage(FEditorStyle::GetBrush("MaterialInstanceEditor.LayersBorder"))
-				.Padding(FMargin(4.0f))
+				SNew(SVerticalBox)
+				+ SVerticalBox::Slot()
+				.Padding(0.0f)
+				.AutoHeight()
 				[
-					SNew(SVerticalBox)
-					+ SVerticalBox::Slot()
-					.AutoHeight()
-					[
-						SAssignNew(HeaderBox, SHorizontalBox)
-						+ SHorizontalBox::Slot()
-						.Padding(FMargin(3.0f, 1.0f))
-						.HAlign(HAlign_Left)
-						.AutoWidth()
-						.VAlign(VAlign_Center)
-						[
-							SNew(STextBlock)
-							.Text(FText::FromName(NestedTree->LayersFunctionsParameterName))
-							.TextStyle(FEditorStyle::Get(), "LargeText")
-						]
-					]
-					+ SVerticalBox::Slot()
-					.Padding(FMargin(3.0f, 0.0f))
+					SAssignNew(HeaderBox, SHorizontalBox)
+					+ SHorizontalBox::Slot()
+					.Padding(FMargin(4.0f, 0.0f))
+					.HAlign(HAlign_Left)
+					.AutoWidth()
+					.VAlign(VAlign_Center)
 					[
-						NestedTree.ToSharedRef()
+						SNew(STextBlock)
+						.Text(FText::FromName(NestedTree->LayersFunctionsParameterName))
 					]
-			]
+				]
+				+ SVerticalBox::Slot()
+				.Padding(FMargin(0.0f))
+				[
+					NestedTree.ToSharedRef()
+				]
 		];
 		if (NestedTree->FunctionParameter != nullptr && FMaterialPropertyHelpers::IsOverriddenExpression(NestedTree->FunctionParameter))
 		{
@@ -1784,97 +1784,40 @@ void SMaterialLayersFunctionsInstanceWrapper::Refresh()
 			.Padding(2.0f)
 			[
 				SNew(SButton)
-				.ButtonStyle(FEditorStyle::Get(), "FlatButton.DarkGrey")
+				.Text(LOCTEXT("Relink", "Relink"))
 				.HAlign(HAlign_Center)
 				.OnClicked(OnRelinkToParent)
 				.ToolTipText(LOCTEXT("RelinkToParentLayers", "Relink to Parent Layers and Blends"))
 				.Visibility(NestedTree.Get(), &SMaterialLayersFunctionsInstanceTree::GetRelinkLayersToParentVisibility)
-				.Content()
-				[
-					SNew(SHorizontalBox)
-					+ SHorizontalBox::Slot()
-					.AutoWidth()
-					[
-						SNew(STextBlock)
-						.Font(FEditorStyle::Get().GetFontStyle("FontAwesome.10"))
-						.TextStyle(FEditorStyle::Get(), "NormalText.Important")
-						.Text(FEditorFontGlyphs::Link)
-					]
-					+ SHorizontalBox::Slot()
-					.AutoWidth()
-					[
-						SNew(STextBlock)
-						.TextStyle(FEditorStyle::Get(), "NormalText.Important")
-						.Text(FText::FromString(FString(TEXT(" Relink"))) /*fa-filter*/)
-					]
-				]
 			];
 		HeaderBox->AddSlot()
 				.AutoWidth()
 				.Padding(2.0f)
 				[
 					SNew(SButton)
-					.ButtonStyle(FEditorStyle::Get(), "FlatButton.DarkGrey")
+					.Text(LOCTEXT("SaveSibling", "Save Sibling"))
 					.HAlign(HAlign_Center)
 					.OnClicked(OnSiblingButtonClicked)
-					.ToolTipText(LOCTEXT("SaveToSiblingInstance", "Save To Sibling Instance"))
-					.Content()
-					[
-						SNew(SHorizontalBox)
-						+ SHorizontalBox::Slot()
-						.AutoWidth()
-						[
-							SNew(STextBlock)
-							.Font(FEditorStyle::Get().GetFontStyle("FontAwesome.10"))
-							.TextStyle(FEditorStyle::Get(), "NormalText.Important")
-							.Text(FText::FromString(FString(TEXT("\xf0c7 \xf178"))) /*fa-filter*/)
-						]
-						+ SHorizontalBox::Slot()
-						.AutoWidth()
-						[
-							SNew(STextBlock)
-							.TextStyle(FEditorStyle::Get(), "NormalText.Important")
-						.Text(FText::FromString(FString(TEXT(" Save Sibling"))) /*fa-filter*/)
-						]
-					]
+					.ToolTipText(LOCTEXT("SaveToSiblingInstance", "Save to Sibling Instance"))
 				];
 		HeaderBox->AddSlot()
 			.AutoWidth()
 			.Padding(2.0f)
 			[
 				SNew(SButton)
-				.ButtonStyle(FEditorStyle::Get(), "FlatButton.DarkGrey")
+				.Text(LOCTEXT("SaveChild", "Save Child"))
 				.HAlign(HAlign_Center)
 				.OnClicked(OnChildButtonClicked)
 				.ToolTipText(LOCTEXT("SaveToChildInstance", "Save To Child Instance"))
-				.Content()
-				[
-					SNew(SHorizontalBox)
-					+ SHorizontalBox::Slot()
-					.AutoWidth()
-					[
-						SNew(STextBlock)
-						.Font(FEditorStyle::Get().GetFontStyle("FontAwesome.10"))
-						.TextStyle(FEditorStyle::Get(), "NormalText.Important")
-						.Text(FText::FromString(FString(TEXT("\xf0c7 \xf149"))) /*fa-filter*/)
-					]
-					+ SHorizontalBox::Slot()
-					.AutoWidth()
-					[
-						SNew(STextBlock)
-						.TextStyle(FEditorStyle::Get(), "NormalText.Important")
-						.Text(FText::FromString(FString(TEXT(" Save Child"))) /*fa-filter*/)
-					]
-				]
 			];
 	}
 	else
 	{
 		this->ChildSlot
 			[
-				SNew(SBorder)
-				.BorderImage(FEditorStyle::GetBrush("MaterialInstanceEditor.StackBody"))
-				.Padding(FMargin(4.0f))
+				SNew(SBox)
+				.Padding(FMargin(10.0f))
+				.HAlign(HAlign_Center)
 				[
 					SNew(STextBlock)
 					.Text(LOCTEXT("AddLayerParameterPrompt", "Add a Material Attribute Layers parameter to see it here."))
@@ -1926,14 +1869,21 @@ FString SMaterialLayersFunctionsMaterialTreeItem::GetCurvePath(UDEditorScalarPar
 
 const FSlateBrush* SMaterialLayersFunctionsMaterialTreeItem::GetBorderImage() const
 {
-	if (StackParameterData->StackDataType == EStackDataType::Stack)
+	return FAppStyle::Get().GetBrush("DetailsView.CategoryMiddle");
+}
+
+FSlateColor SMaterialLayersFunctionsMaterialTreeItem::GetOuterBackgroundColor(TSharedPtr<FSortedParamData> InParamData) const
+{
+	if (InParamData->StackDataType == EStackDataType::Stack)
 	{
-		return FEditorStyle::GetBrush("MaterialInstanceEditor.StackHeader");
+		return FAppStyle::Get().GetSlateColor("Colors.Header");
 	}
-	else
+	else if (IsHovered() || InParamData->StackDataType == EStackDataType::Group)
 	{
-		return FEditorStyle::GetBrush("MaterialInstanceEditor.StackBody");
+		return FAppStyle::Get().GetSlateColor("Colors.Header");
 	}
+
+	return FAppStyle::Get().GetSlateColor("Colors.Panel");
 }
 
 void SMaterialLayersFunctionsMaterialTreeItem::RefreshOnRowChange(const FAssetData& AssetData, SMaterialLayersFunctionsMaterialTree* InTree)
@@ -1959,12 +1909,12 @@ void SMaterialLayersFunctionsMaterialTreeItem::Construct(const FArguments& InArg
 	TSharedRef<SWidget> RightSideWidget = SNullWidget::NullWidget;
 	FText NameOverride;
 	TSharedRef<SVerticalBox> WrapperWidget = SNew(SVerticalBox);
-
+	EHorizontalAlignment ValueAlignment = HAlign_Left;
 	// STACK --------------------------------------------------
 	if (StackParameterData->StackDataType == EStackDataType::Stack)
 	{
 		WrapperWidget->AddSlot()
-			.Padding(3.0f)
+			.Padding(0.0f)
 			.AutoHeight()
 			[
 				SNullWidget::NullWidget
@@ -1981,7 +1931,7 @@ void SMaterialLayersFunctionsMaterialTreeItem::Construct(const FArguments& InArg
 			HeaderRowWidget->AddSlot()
 				.HAlign(HAlign_Center)
 				.VAlign(VAlign_Center)
-				.Padding(2.5f, 0)
+				.Padding(0.0f)
 				.AutoWidth()
 				[
 					SNullWidget::NullWidget
@@ -2039,7 +1989,6 @@ void SMaterialLayersFunctionsMaterialTreeItem::Construct(const FArguments& InArg
 				[
 					SNew(STextBlock)
 					.Text(this, &SMaterialLayersFunctionsMaterialTreeItem::GetLayerName, InArgs._InTree, StackParameterData->ParameterInfo.Index)
-					.Font(FEditorStyle::GetFontStyle(TEXT("MaterialEditor.Layers.EditableFontImportant")))
 				];
 				HeaderRowWidget->AddSlot()
 				.FillWidth(1.0f)
@@ -2057,7 +2006,6 @@ void SMaterialLayersFunctionsMaterialTreeItem::Construct(const FArguments& InArg
 				[
 					SNew(STextBlock)
 					.Text(NameOverride)
-					.TextStyle(FEditorStyle::Get(), "NormalText.Important")
 				];
 		}
 		LeftSideWidget = HeaderRowWidget;
@@ -2070,13 +2018,16 @@ void SMaterialLayersFunctionsMaterialTreeItem::Construct(const FArguments& InArg
 		NameOverride = FText::FromName(StackParameterData->Group.GroupName);
 		LeftSideWidget = SNew(STextBlock)
 			.Text(NameOverride)
-			.TextStyle(FEditorStyle::Get(), "TinyText");
+			.Font(FAppStyle::Get().GetFontStyle("PropertyWindow.BoldFont"))
+			.TextStyle(FAppStyle::Get(), "DetailsView.CategoryTextStyle")
+			.TransformPolicy(ETextTransformPolicy::ToUpper);
 	}
 	// END GROUP
 
 	// ASSET --------------------------------------------------
 	if (StackParameterData->StackDataType == EStackDataType::Asset)
 	{
+		ValueAlignment = HAlign_Fill;
 		StackParameterData->ParameterHandle->GetProperty()->SetMetaData(FName(TEXT("DisplayThumbnail")), TEXT("true"));
 		FIntPoint ThumbnailOverride;
 		if (StackParameterData->ParameterInfo.Association == EMaterialParameterAssociation::LayerParameter)
@@ -2113,29 +2064,29 @@ void SMaterialLayersFunctionsMaterialTreeItem::Construct(const FArguments& InArg
 
 		RightSideWidget = SNew(SVerticalBox)
 			+ SVerticalBox::Slot()
+			.Padding(0.0f)
 			[
 				SNew(SHorizontalBox)
 				+ SHorizontalBox::Slot()
-			.AutoWidth()
-			.HAlign(HAlign_Center)
-			.VAlign(VAlign_Center)
-			.Padding(4.0f)
-			.MaxWidth(ThumbnailOverride.X)
-			[
-				SAssignNew(ThumbnailBox, SBox)
+				.AutoWidth()
+				.HAlign(HAlign_Center)
+				.VAlign(VAlign_Center)
+				.Padding(0.0f)
+				.MaxWidth(ThumbnailOverride.X)
 				[
-					Tree->CreateThumbnailWidget(StackParameterData->ParameterInfo.Association, StackParameterData->ParameterInfo.Index, ThumbnailOverride.X)
+					SAssignNew(ThumbnailBox, SBox)
+					[
+						Tree->CreateThumbnailWidget(StackParameterData->ParameterInfo.Association, StackParameterData->ParameterInfo.Index, ThumbnailOverride.X)
+					]
+				]
+				+ SHorizontalBox::Slot()
+				.FillWidth(1.0)
+				[
+					SNew(SObjectPropertyEntryBox)
+					.AllowedClass(UMaterialFunctionInterface::StaticClass())
+					.ObjectPath(this, &SMaterialLayersFunctionsMaterialTreeItem::GetInstancePath, Tree)
+					.DisplayCompactSize(true)
 				]
-			]
-		+ SHorizontalBox::Slot()
-			.FillWidth(1.0)
-			[
-				SNew(SObjectPropertyEntryBox)
-				.AllowedClass(UMaterialFunctionInterface::StaticClass())
-				.ObjectPath(this, &SMaterialLayersFunctionsMaterialTreeItem::GetInstancePath, Tree)
-				.DisplayCompactSize(true)
-			]
-	
 			];
 		ThumbnailBox->SetMaxDesiredHeight(ThumbnailOverride.Y);
 		ThumbnailBox->SetMinDesiredHeight(ThumbnailOverride.Y);
@@ -2145,7 +2096,6 @@ void SMaterialLayersFunctionsMaterialTreeItem::Construct(const FArguments& InArg
 	// END ASSET
 
 	// PROPERTY ----------------------------------------------
-	bool bisPaddedProperty = false;
 	if (StackParameterData->StackDataType == EStackDataType::Property)
 	{
 
@@ -2205,8 +2155,8 @@ void SMaterialLayersFunctionsMaterialTreeItem::Construct(const FArguments& InArg
 				[
 					SNew(STextBlock)
 					.Text(ParameterName)
-				.ToolTipText(FMaterialPropertyHelpers::GetParameterExpressionDescription(StackParameterData->Parameter, MaterialEditorInstance))
-				.Font(FEditorStyle::GetFontStyle(TEXT("PropertyWindow.NormalFont")))
+					.ToolTipText(FMaterialPropertyHelpers::GetParameterExpressionDescription(StackParameterData->Parameter, MaterialEditorInstance))
+					.Font(FEditorStyle::GetFontStyle(TEXT("PropertyWindow.NormalFont")))
 				]
 			.ValueContent()
 				.HAlign(HAlign_Fill)
@@ -2242,13 +2192,13 @@ void SMaterialLayersFunctionsMaterialTreeItem::Construct(const FArguments& InArg
 						[
 							SAssignNew(NameVerticalBox, SVerticalBox)
 							+ SVerticalBox::Slot()
-						.AutoHeight()
-						[
-							SNew(STextBlock)
-							.Text(ParameterName)
-						.ToolTipText(FMaterialPropertyHelpers::GetParameterExpressionDescription(StackParameterData->Parameter, MaterialEditorInstance))
-						.Font(FEditorStyle::GetFontStyle(TEXT("PropertyWindow.NormalFont")))
-						]
+							.AutoHeight()
+							[
+								SNew(STextBlock)
+								.Text(ParameterName)
+								.ToolTipText(FMaterialPropertyHelpers::GetParameterExpressionDescription(StackParameterData->Parameter, MaterialEditorInstance))
+								.Font(FEditorStyle::GetFontStyle(TEXT("PropertyWindow.NormalFont")))
+							]
 						];
 					CustomWidget.ValueContent()
 						[
@@ -2269,21 +2219,21 @@ void SMaterialLayersFunctionsMaterialTreeItem::Construct(const FArguments& InArg
 							[
 								SNew(SHorizontalBox)
 								+ SHorizontalBox::Slot()
-							.AutoWidth()
-							.Padding(20.0, 2.0, 4.0, 2.0)
-							[
-								SNew(STextBlock)
-								.Text(FText::FromName(Red))
-							.Font(FEditorStyle::GetFontStyle(TEXT("PropertyWindow.BoldFont")))
-							]
-						+ SHorizontalBox::Slot()
-							.HAlign(HAlign_Left)
-							.Padding(4.0, 2.0)
-							[
-								SNew(STextBlock)
-								.Text(TextureParam->ChannelNames.R)
-							.Font(FEditorStyle::GetFontStyle(TEXT("PropertyWindow.NormalFont")))
-							]
+								.AutoWidth()
+								.Padding(20.0, 2.0, 4.0, 2.0)
+								[
+									SNew(STextBlock)
+									.Text(FText::FromName(Red))
+									.Font(FEditorStyle::GetFontStyle(TEXT("PropertyWindow.BoldFont")))
+								]
+								+ SHorizontalBox::Slot()
+								.HAlign(HAlign_Left)
+								.Padding(4.0, 2.0)
+								[
+									SNew(STextBlock)
+									.Text(TextureParam->ChannelNames.R)
+									.Font(FEditorStyle::GetFontStyle(TEXT("PropertyWindow.NormalFont")))
+								]
 							];
 					}
 					if (!TextureParam->ChannelNames.G.IsEmpty())
@@ -2292,21 +2242,21 @@ void SMaterialLayersFunctionsMaterialTreeItem::Construct(const FArguments& InArg
 							[
 								SNew(SHorizontalBox)
 								+ SHorizontalBox::Slot()
-							.Padding(20.0, 2.0, 4.0, 2.0)
-							.AutoWidth()
-							[
-								SNew(STextBlock)
-								.Text(FText::FromName(Green))
-							.Font(FEditorStyle::GetFontStyle(TEXT("PropertyWindow.BoldFont")))
-							]
-						+ SHorizontalBox::Slot()
-							.HAlign(HAlign_Left)
-							.Padding(4.0, 2.0)
-							[
-								SNew(STextBlock)
-								.Text(TextureParam->ChannelNames.G)
-							.Font(FEditorStyle::GetFontStyle(TEXT("PropertyWindow.NormalFont")))
-							]
+								.Padding(20.0, 2.0, 4.0, 2.0)
+								.AutoWidth()
+								[
+									SNew(STextBlock)
+									.Text(FText::FromName(Green))
+									.Font(FEditorStyle::GetFontStyle(TEXT("PropertyWindow.BoldFont")))
+								]
+								+ SHorizontalBox::Slot()
+								.HAlign(HAlign_Left)
+								.Padding(4.0, 2.0)
+								[
+									SNew(STextBlock)
+									.Text(TextureParam->ChannelNames.G)
+									.Font(FEditorStyle::GetFontStyle(TEXT("PropertyWindow.NormalFont")))
+								]
 							];
 					}
 					if (!TextureParam->ChannelNames.B.IsEmpty())
@@ -2315,21 +2265,21 @@ void SMaterialLayersFunctionsMaterialTreeItem::Construct(const FArguments& InArg
 							[
 								SNew(SHorizontalBox)
 								+ SHorizontalBox::Slot()
-							.Padding(20.0, 2.0, 4.0, 2.0)
-							.AutoWidth()
-							[
-								SNew(STextBlock)
-								.Text(FText::FromName(Blue))
-							.Font(FEditorStyle::GetFontStyle(TEXT("PropertyWindow.BoldFont")))
-							]
-						+ SHorizontalBox::Slot()
-							.HAlign(HAlign_Left)
-							.Padding(4.0, 2.0)
-							[
-								SNew(STextBlock)
-								.Text(TextureParam->ChannelNames.B)
-							.Font(FEditorStyle::GetFontStyle(TEXT("PropertyWindow.NormalFont")))
-							]
+								.Padding(20.0, 2.0, 4.0, 2.0)
+								.AutoWidth()
+								[
+									SNew(STextBlock)
+									.Text(FText::FromName(Blue))
+									.Font(FEditorStyle::GetFontStyle(TEXT("PropertyWindow.BoldFont")))
+								]
+								+ SHorizontalBox::Slot()
+								.HAlign(HAlign_Left)
+								.Padding(4.0, 2.0)
+								[
+									SNew(STextBlock)
+									.Text(TextureParam->ChannelNames.B)
+									.Font(FEditorStyle::GetFontStyle(TEXT("PropertyWindow.NormalFont")))
+								]
 							];
 					}
 					if (!TextureParam->ChannelNames.A.IsEmpty())
@@ -2338,21 +2288,21 @@ void SMaterialLayersFunctionsMaterialTreeItem::Construct(const FArguments& InArg
 							[
 								SNew(SHorizontalBox)
 								+ SHorizontalBox::Slot()
-							.Padding(20.0, 2.0, 4.0, 2.0)
-							.AutoWidth()
-							[
-								SNew(STextBlock)
-								.Text(FText::FromName(Alpha))
-							.Font(FEditorStyle::GetFontStyle(TEXT("PropertyWindow.BoldFont")))
-							]
-						+ SHorizontalBox::Slot()
-							.HAlign(HAlign_Left)
-							.Padding(4.0, 2.0)
-							[
-								SNew(STextBlock)
-								.Text(TextureParam->ChannelNames.A)
-							.Font(FEditorStyle::GetFontStyle(TEXT("PropertyWindow.NormalFont")))
-							]
+								.Padding(20.0, 2.0, 4.0, 2.0)
+								.AutoWidth()
+								[
+									SNew(STextBlock)
+									.Text(FText::FromName(Alpha))
+									.Font(FEditorStyle::GetFontStyle(TEXT("PropertyWindow.BoldFont")))
+								]
+								+ SHorizontalBox::Slot()
+								.HAlign(HAlign_Left)
+								.Padding(4.0, 2.0)
+								[
+									SNew(STextBlock)
+									.Text(TextureParam->ChannelNames.A)
+									.Font(FEditorStyle::GetFontStyle(TEXT("PropertyWindow.NormalFont")))
+								]
 							];
 					}
 				}
@@ -2371,28 +2321,28 @@ void SMaterialLayersFunctionsMaterialTreeItem::Construct(const FArguments& InArg
 				[
 					SNew(SHorizontalBox)
 					+ SHorizontalBox::Slot()
-				.VAlign(VAlign_Center)
-				[
-					SNew(STextBlock)
-					.Text(NameOverride)
-				.ToolTipText(FMaterialPropertyHelpers::GetParameterExpressionDescription(StackParameterData->Parameter, MaterialEditorInstance))
-				.Font(FEditorStyle::GetFontStyle(TEXT("PropertyWindow.NormalFont")))
-				]
+					.VAlign(VAlign_Center)
+					[
+						SNew(STextBlock)
+						.Text(NameOverride)
+						.ToolTipText(FMaterialPropertyHelpers::GetParameterExpressionDescription(StackParameterData->Parameter, MaterialEditorInstance))
+						.Font(FEditorStyle::GetFontStyle(TEXT("PropertyWindow.NormalFont")))
+					]
 				]
-			.ValueContent()
+				.ValueContent()
 				.MaxDesiredWidth(200.0f)
 				[
 					SNew(SHorizontalBox)
 					+ SHorizontalBox::Slot()
-				.FillWidth(1.0f)
+					.FillWidth(1.0f)
 				[
 					SNew(SHorizontalBox)
 					+ SHorizontalBox::Slot()
-				.HAlign(HAlign_Left)
-				.AutoWidth()
-				[
-					RMaskProperty->CreatePropertyNameWidget(FText::GetEmpty(), FText::GetEmpty(), false)
-				]
+					.HAlign(HAlign_Left)
+					.AutoWidth()
+					[
+						RMaskProperty->CreatePropertyNameWidget(FText::GetEmpty(), FText::GetEmpty(), false)
+					]
 			+ SHorizontalBox::Slot()
 				.HAlign(HAlign_Left)
 				.AutoWidth()
@@ -2459,8 +2409,6 @@ void SMaterialLayersFunctionsMaterialTreeItem::Construct(const FArguments& InArg
 					]
 				];
 			}
-
-			bisPaddedProperty = true;
 		}
 
 		FNodeWidgets NodeWidgets = Node.CreateNodeWidgets();
@@ -2480,7 +2428,6 @@ void SMaterialLayersFunctionsMaterialTreeItem::Construct(const FArguments& InArg
 
 	// FINAL WRAPPER
 
-	float ValuePadding = bisPaddedProperty ? 20.0f : 0.0f;
 	LeftSideWidget->SetEnabled(false);
 	RightSideWidget->SetEnabled(false);
 	if (StackParameterData->StackDataType == EStackDataType::Stack)
@@ -2491,6 +2438,7 @@ void SMaterialLayersFunctionsMaterialTreeItem::Construct(const FArguments& InArg
 			[
 				SNew(SBorder)
 				.BorderImage(this, &SMaterialLayersFunctionsMaterialTreeItem::GetBorderImage)
+				.BorderBackgroundColor(this, &SMaterialLayersFunctionsMaterialTreeItem::GetOuterBackgroundColor, StackParameterData)
 				.Padding(0.0f)
 				[
 					SAssignNew(FinalStack, SHorizontalBox)
@@ -2510,39 +2458,19 @@ void SMaterialLayersFunctionsMaterialTreeItem::Construct(const FArguments& InArg
 				LeftSideWidget
 			];
 	}
-	else
+	else if (StackParameterData->StackDataType == EStackDataType::Group)
 	{
-		FSlateBrush* StackBrush;
-		switch (StackParameterData->ParameterInfo.Association)
-		{
-		case EMaterialParameterAssociation::LayerParameter:
-		{
-			StackBrush = const_cast<FSlateBrush*>(FEditorStyle::GetBrush("MaterialInstanceEditor.StackBody"));
-			break;
-		}
-		case EMaterialParameterAssociation::BlendParameter:
-		{
-			StackBrush = const_cast<FSlateBrush*>(FEditorStyle::GetBrush("MaterialInstanceEditor.StackBodyBlend"));
-			break;
-		}
-		default:
-			break;
-		}
 		WrapperWidget->AddSlot()
 			.AutoHeight()
 			[
 				SNew(SBorder)
-				.BorderImage(this, &SMaterialLayersFunctionsMaterialTreeItem::GetBorderImage)
-				.Padding(0.0f)
+				.BorderImage(FAppStyle::Get().GetBrush("DetailsView.GridLine"))
+				.Padding(FMargin(0, 0, 0, 1))
 				[
-					SNew(SSplitter)
-					.Style(FEditorStyle::Get(), "DetailsView.Splitter")
-					.PhysicalSplitterHandleSize(1.0f)
-					.HitDetectionSplitterHandleSize(5.0f)
-					+ SSplitter::Slot()
-					.Value(Tree->ColumnSizeData.NameColumnWidth)
-					.OnSlotResized(Tree->ColumnSizeData.OnNameColumnResized)
-					.Value(0.25f)
+					SNew(SBorder)
+					.Padding(0.0f)
+					.BorderImage(this, &SMaterialLayersFunctionsMaterialTreeItem::GetBorderImage)
+					.BorderBackgroundColor(this, &SMaterialLayersFunctionsMaterialTreeItem::GetOuterBackgroundColor, StackParameterData)
 					[
 						SNew(SHorizontalBox)
 						+ SHorizontalBox::Slot()
@@ -2559,16 +2487,59 @@ void SMaterialLayersFunctionsMaterialTreeItem::Construct(const FArguments& InArg
 							LeftSideWidget
 						]
 					]
-					+ SSplitter::Slot()
-					.Value(Tree->ColumnSizeData.ValueColumnWidth)
-					.OnSlotResized(Tree->ColumnSizeData.OnValueColumnResized)
+				]
+			];
+	}
+	else
+	{
+		WrapperWidget->AddSlot()
+			.AutoHeight()
+			[
+				SNew(SBorder)
+				.BorderImage(FAppStyle::Get().GetBrush("DetailsView.GridLine"))
+				.Padding(FMargin(0, 0, 0, 1))
+				[
+					SNew(SBorder)
+					.Padding(0.0f)
+					.BorderImage(this, &SMaterialLayersFunctionsMaterialTreeItem::GetBorderImage)
+					.BorderBackgroundColor(this, &SMaterialLayersFunctionsMaterialTreeItem::GetOuterBackgroundColor, StackParameterData)
 					[
-						SNew(SHorizontalBox)
-						+ SHorizontalBox::Slot()
-						.MaxWidth(350.0f - ValuePadding)
-						.Padding(FMargin(5.0f, 2.0f, ValuePadding, 2.0f))
+						SNew(SSplitter)
+						.Style(FEditorStyle::Get(), "DetailsView.Splitter")
+						.PhysicalSplitterHandleSize(1.0f)
+						.HitDetectionSplitterHandleSize(5.0f)
+						+ SSplitter::Slot()
+						.Value(Tree->ColumnSizeData.NameColumnWidth)
+						.OnSlotResized(Tree->ColumnSizeData.OnNameColumnResized)
+						.Value(0.25f)
 						[
-							RightSideWidget
+							SNew(SHorizontalBox)
+							+ SHorizontalBox::Slot()
+							.AutoWidth()
+							.VAlign(VAlign_Center)
+							.Padding(FMargin(3.0f))
+							[
+								SNew(SExpanderArrow, SharedThis(this))
+							]
+							+ SHorizontalBox::Slot()
+							.Padding(FMargin(2.0f))
+							.VAlign(VAlign_Center)
+							[
+								LeftSideWidget
+							]
+						]
+						+ SSplitter::Slot()
+						.Value(Tree->ColumnSizeData.ValueColumnWidth)
+						.OnSlotResized(Tree->ColumnSizeData.OnValueColumnResized)
+						[
+							SNew(SHorizontalBox)
+							.Clipping(EWidgetClipping::OnDemand)
+							+ SHorizontalBox::Slot()
+							.HAlign(ValueAlignment)
+							.VAlign(VAlign_Center)
+							[
+								RightSideWidget
+							]
 						]
 					]
 				]
@@ -3072,32 +3043,27 @@ void SMaterialLayersFunctionsMaterialWrapper::Refresh()
 	{
 		this->ChildSlot
 			[
-				SNew(SBorder)
-				.BorderImage(FEditorStyle::GetBrush("MaterialInstanceEditor.LayersBorder"))
-				.Padding(FMargin(4.0f))
+				
+				SNew(SVerticalBox)
+				+ SVerticalBox::Slot()
+				.AutoHeight()
 				[
-					SNew(SVerticalBox)
-					+ SVerticalBox::Slot()
-					.AutoHeight()
-					[
-						SAssignNew(HeaderBox, SHorizontalBox)
-						+ SHorizontalBox::Slot()
-						.Padding(FMargin(3.0f, 1.0f))
-						.HAlign(HAlign_Left)
-						.AutoWidth()
-						.VAlign(VAlign_Center)
-						[
-							SNew(STextBlock)
-							.Text(FText::FromName(NestedTree->LayersFunctionsParameterName))
-							.TextStyle(FEditorStyle::Get(), "LargeText")
-						]
-					]
-					+ SVerticalBox::Slot()
-					.Padding(FMargin(3.0f, 0.0f))
+					SAssignNew(HeaderBox, SHorizontalBox)
+					+ SHorizontalBox::Slot()
+					.Padding(FMargin(3.0f, 1.0f))
+					.HAlign(HAlign_Left)
+					.AutoWidth()
+					.VAlign(VAlign_Center)
 					[
-						NestedTree.ToSharedRef()
+						SNew(STextBlock)
+						.Text(FText::FromName(NestedTree->LayersFunctionsParameterName))
 					]
 				]
+				+ SVerticalBox::Slot()
+				.Padding(FMargin(0.0f))
+				[
+					NestedTree.ToSharedRef()
+				]
 			];
 		HeaderBox->AddSlot()
 			.FillWidth(1.0f)
@@ -3110,8 +3076,8 @@ void SMaterialLayersFunctionsMaterialWrapper::Refresh()
 		this->ChildSlot
 			[
 				SNew(SBorder)
-				.BorderImage(FEditorStyle::GetBrush("MaterialInstanceEditor.StackBody"))
-				.Padding(FMargin(4.0f))
+				.Padding(FMargin(10.0f))
+				.HAlign(HAlign_Center)
 				[
 					SNew(STextBlock)
 					.Text(LOCTEXT("AddLayerParameterPrompt", "Add a Material Attribute Layers parameter to see it here."))
diff --git a/Engine/Source/Editor/MaterialEditor/Private/SMaterialLayersFunctionsTree.h b/Engine/Source/Editor/MaterialEditor/Private/SMaterialLayersFunctionsTree.h
--- a/Engine/Source/Editor/MaterialEditor/Private/SMaterialLayersFunctionsTree.h
+++ b/Engine/Source/Editor/MaterialEditor/Private/SMaterialLayersFunctionsTree.h
@@ -45,6 +45,7 @@ private:
 	FString GetCurvePath(class UDEditorScalarParameterValue* Parameter) const;
 	const FSlateBrush* GetBorderImage() const;
 
+	FSlateColor GetOuterBackgroundColor(TSharedPtr<FSortedParamData> InParamData) const;
 public:
 
 	void RefreshOnRowChange(const FAssetData& AssetData, SMaterialLayersFunctionsInstanceTree* InTree);
@@ -220,6 +221,7 @@ private:
 	FString GetCurvePath(class UDEditorScalarParameterValue* Parameter) const;
 	const FSlateBrush* GetBorderImage() const;
 
+	FSlateColor GetOuterBackgroundColor(TSharedPtr<FSortedParamData> InParamData) const;
 public:
 
 	void RefreshOnRowChange(const FAssetData& AssetData, SMaterialLayersFunctionsMaterialTree* InTree);
diff --git a/Engine/Source/Editor/MaterialEditor/Private/SMaterialParametersOverviewWidget.cpp b/Engine/Source/Editor/MaterialEditor/Private/SMaterialParametersOverviewWidget.cpp
--- a/Engine/Source/Editor/MaterialEditor/Private/SMaterialParametersOverviewWidget.cpp
+++ b/Engine/Source/Editor/MaterialEditor/Private/SMaterialParametersOverviewWidget.cpp
@@ -41,6 +41,7 @@
 #include "Widgets/Layout/SWidgetSwitcher.h"
 #include "MaterialEditor/DEditorRuntimeVirtualTextureParameterValue.h"
 #include "ThumbnailRendering/ThumbnailManager.h"
+#include "Styling/StyleColors.h"
 
 
 
@@ -54,14 +55,17 @@ FString SMaterialParametersOverviewTreeItem::GetCurvePath(UDEditorScalarParamete
 
 const FSlateBrush* SMaterialParametersOverviewTreeItem::GetBorderImage() const
 {
-	if (IsHovered())
-	{
-		return FEditorStyle::GetBrush("DetailsView.CategoryMiddle_Hovered");
-	}
-	else
+	return FAppStyle::Get().GetBrush("DetailsView.CategoryMiddle");
+}
+
+FSlateColor SMaterialParametersOverviewTreeItem::GetOuterBackgroundColor(TSharedPtr<FSortedParamData> InParamData) const
+{
+	if (IsHovered() || InParamData->StackDataType == EStackDataType::Group)
 	{
-		return FEditorStyle::GetBrush("DetailsView.CategoryMiddle");
+		return FAppStyle::Get().GetSlateColor("Colors.Header");
 	}
+
+	return FAppStyle::Get().GetSlateColor("Colors.Panel");
 }
 
 void SMaterialParametersOverviewTreeItem::RefreshOnRowChange(const FAssetData& AssetData, TSharedPtr<SMaterialParametersOverviewTree> InTree)
@@ -88,8 +92,10 @@ void SMaterialParametersOverviewTreeItem::Construct(const FArguments& InArgs, co
 	{
 		NameOverride = FText::FromName(StackParameterData->Group.GroupName);
 		LeftSideWidget = SNew(STextBlock)
+			.TransformPolicy(ETextTransformPolicy::ToUpper)
 			.Text(NameOverride)
-			.TextStyle(FEditorStyle::Get(), "TinyText");
+			.Font(FAppStyle::Get().GetFontStyle("PropertyWindow.BoldFont"))
+			.TextStyle(FAppStyle::Get(), "DetailsView.CategoryTextStyle");
 	}
 // END GROUP
 
@@ -239,21 +245,21 @@ void SMaterialParametersOverviewTreeItem::Construct(const FArguments& InArgs, co
 							[
 								SNew(SHorizontalBox)
 								+ SHorizontalBox::Slot()
-							.AutoWidth()
-							.Padding(20.0, 2.0, 4.0, 2.0)
-							[
-								SNew(STextBlock)
-								.Text(FText::FromName(Red))
-							.Font(FEditorStyle::GetFontStyle(TEXT("PropertyWindow.BoldFont")))
-							]
-						+ SHorizontalBox::Slot()
-							.HAlign(HAlign_Left)
-							.Padding(4.0, 2.0)
-							[
-								SNew(STextBlock)
-								.Text(TextureParam->ChannelNames.R)
-							.Font(FEditorStyle::GetFontStyle(TEXT("PropertyWindow.NormalFont")))
-							]
+								.AutoWidth()
+								.Padding(20.0, 2.0, 4.0, 2.0)
+								[
+									SNew(STextBlock)
+									.Text(FText::FromName(Red))
+									.Font(FEditorStyle::GetFontStyle(TEXT("PropertyWindow.BoldFont")))
+								]
+								+ SHorizontalBox::Slot()
+								.HAlign(HAlign_Left)
+								.Padding(4.0, 2.0)
+								[
+									SNew(STextBlock)
+									.Text(TextureParam->ChannelNames.R)
+									.Font(FEditorStyle::GetFontStyle(TEXT("PropertyWindow.NormalFont")))
+								]
 							];
 					}
 					if (!TextureParam->ChannelNames.G.IsEmpty())
@@ -468,43 +474,50 @@ void SMaterialParametersOverviewTreeItem::Construct(const FArguments& InArgs, co
 			.AutoHeight()
 			[
 				SNew(SBorder)
-				.Padding(0.0f)
-				.BorderImage(this, &SMaterialParametersOverviewTreeItem::GetBorderImage)
+				.BorderImage(FAppStyle::Get().GetBrush("DetailsView.GridLine"))
+				.Padding(FMargin(0, 0, 0, 1))
 				[
-					SNew(SSplitter)
-					.Style(FEditorStyle::Get(), "DetailsView.Splitter")
-					.PhysicalSplitterHandleSize(1.0f)
-					.HitDetectionSplitterHandleSize(5.0f)
-					+ SSplitter::Slot()
-					.Value(ColumnSizeData.NameColumnWidth)
-					.OnSlotResized(ColumnSizeData.OnNameColumnResized)
-					.Value(0.25f)
+					SNew(SBorder)
+					.Padding(3.0f)
+					.BorderImage(this, &SMaterialParametersOverviewTreeItem::GetBorderImage)
+					.BorderBackgroundColor(this, &SMaterialParametersOverviewTreeItem::GetOuterBackgroundColor, StackParameterData)
 					[
-						SNew(SHorizontalBox)
-						+ SHorizontalBox::Slot()
-						.AutoWidth()
-						.VAlign(VAlign_Center)
-						.Padding(FMargin(3.0f))
+						SNew(SSplitter)
+						.Style(FEditorStyle::Get(), "DetailsView.Splitter")
+						.PhysicalSplitterHandleSize(1.0f)
+						.HitDetectionSplitterHandleSize(5.0f)
+						+ SSplitter::Slot()
+						.Value(ColumnSizeData.NameColumnWidth)
+						.OnSlotResized(ColumnSizeData.OnNameColumnResized)
+						.Value(0.25f)
 						[
-							SNew(SExpanderArrow, SharedThis(this))
-						]
-						+ SHorizontalBox::Slot()
-						.Padding(FMargin(2.0f))
-						.VAlign(VAlign_Center)
-						[
-							LeftSideWidget
+							SNew(SHorizontalBox)
+							+ SHorizontalBox::Slot()
+							.AutoWidth()
+							.VAlign(VAlign_Center)
+							.Padding(FMargin(3.0f))
+							[
+								SNew(SExpanderArrow, SharedThis(this))
+							]
+							+ SHorizontalBox::Slot()
+							.Padding(FMargin(2.0f))
+							.VAlign(VAlign_Center)
+							[
+								LeftSideWidget
+							]
 						]
-					]
-					+ SSplitter::Slot()
-					.Value(ColumnSizeData.ValueColumnWidth)
-					.OnSlotResized(ColumnSizeData.OnValueColumnResized) 
-					[
-						SNew(SHorizontalBox)
-						+ SHorizontalBox::Slot()
-						.MaxWidth(350.0f - ValuePadding)
-						.Padding(FMargin(5.0f, 2.0f, ValuePadding, 2.0f))
+						+ SSplitter::Slot()
+						.Value(ColumnSizeData.ValueColumnWidth)
+						.OnSlotResized(ColumnSizeData.OnValueColumnResized) 
 						[
-							RightSideWidget
+							SNew(SHorizontalBox)
+							+ SHorizontalBox::Slot()
+							.Padding(FMargin(5.0f, 2.0f, 5.0, 2.0f))
+							.HAlign(HAlign_Left)
+							.VAlign(VAlign_Center)
+							[
+								RightSideWidget
+							]
 						]
 					]
 				]
@@ -805,7 +818,7 @@ void SMaterialParametersOverviewTree::ShowSubParameters()
 
 const FSlateBrush* SMaterialParametersOverviewPanel::GetBackgroundImage() const
 {
-	return FEditorStyle::GetBrush("DetailsView.CategoryTop_Hovered");
+	return FEditorStyle::GetBrush("DetailsView.CategoryTop");
 }
 
 int32 SMaterialParametersOverviewPanel::GetPanelIndex() const
@@ -834,63 +847,37 @@ void SMaterialParametersOverviewPanel::Refresh()
 		[
 			SNew(SVerticalBox)
 			+ SVerticalBox::Slot()
+			.Padding(0.0f)
 			[
-				SNew(SBorder)
-				.BorderImage(this, &SMaterialParametersOverviewPanel::GetBackgroundImage)
-				.Padding(FMargin(4.0f))
+				
+				SNew(SVerticalBox)
+				+ SVerticalBox::Slot()
+				.AutoHeight()
+				[
+					SAssignNew(HeaderBox, SHorizontalBox)
+				]
+				+ SVerticalBox::Slot()
+				.Padding(0.0f)
 				[
-					SNew(SVerticalBox)
-					+ SVerticalBox::Slot()
-					.AutoHeight()
+					SNew(SHorizontalBox)
+					+ SHorizontalBox::Slot()
+					.HAlign(HAlign_Fill)
 					[
-						SAssignNew(HeaderBox, SHorizontalBox)
-						+ SHorizontalBox::Slot()
-						.Padding(FMargin(3.0f, 1.0f))
-						.HAlign(HAlign_Left)
-						.AutoWidth()
-						.VAlign(VAlign_Center)
-						[
-							SNew(STextBlock)
-							.Text(LOCTEXT("ParameterDefaults", "Parameter Defaults"))
-							.Font(FEditorStyle::GetFontStyle("DetailsView.CategoryFontStyle"))
-							.ShadowOffset(FVector2D(1.0f, 1.0f))
-						]
+						NestedTree.ToSharedRef()
 					]
-					+ SVerticalBox::Slot()
-					.Padding(FMargin(3.0f, 2.0f, 3.0f, 3.0f))
+					+ SHorizontalBox::Slot()
+					.HAlign(HAlign_Right)
+					.AutoWidth()
 					[
-						SNew(SBorder)
-						.BorderImage(FEditorStyle::GetBrush("DetailsView.CategoryTop"))
+						SNew(SBox)
+						.WidthOverride(16.0f)
 						[
-							SNew(SHorizontalBox)
-							+ SHorizontalBox::Slot()
-							.HAlign(HAlign_Fill)
-							[
-								SNew(SWidgetSwitcher)
-								.WidgetIndex(this, &SMaterialParametersOverviewPanel::GetPanelIndex)
-								+ SWidgetSwitcher::Slot()
-								[
-									SNew(STextBlock)
-									.Text(LOCTEXT("AddParams", "Add parameters to see them here."))
-								]
-								+ SWidgetSwitcher::Slot()
-								[
-									NestedTree.ToSharedRef()
-								]
-							]
-							+ SHorizontalBox::Slot()
-							.HAlign(HAlign_Right)
-							.AutoWidth()
-							[
-								SNew(SBox)
-								.WidthOverride(16.0f)
-								[
-									ExternalScrollbar.ToSharedRef()
-								]
-							]
+							ExternalScrollbar.ToSharedRef()
 						]
 					]
+					
 				]
+				
 			]
 		];
 
@@ -907,29 +894,10 @@ void SMaterialParametersOverviewPanel::Refresh()
 				.Padding(2.0f)
 				[
 					SNew(SButton)
-					.ButtonStyle(FEditorStyle::Get(), "FlatButton.DarkGrey")
+					.Text(LOCTEXT("SaveChild", "Save Child"))
 					.HAlign(HAlign_Center)
 					.OnClicked(OnChildButtonClicked)
 					.ToolTipText(LOCTEXT("SaveToChildInstance", "Save To Child Instance"))
-					.Content()
-					[
-						SNew(SHorizontalBox)
-						+ SHorizontalBox::Slot()
-						.AutoWidth()
-						[
-							SNew(STextBlock)
-							.Font(FEditorStyle::Get().GetFontStyle("FontAwesome.10"))
-							.TextStyle(FEditorStyle::Get(), "NormalText.Important")
-							.Text(FText::FromString(FString(TEXT("\xf0c7 \xf149"))) /*fa-filter*/)
-						]
-						+ SHorizontalBox::Slot()
-						.AutoWidth()
-						[
-							SNew(STextBlock)
-							.TextStyle(FEditorStyle::Get(), "NormalText.Important")
-							.Text(FText::FromString(FString(TEXT(" Save Child"))) /*fa-filter*/)
-						]
-					]
 				];
 		}
 	}
@@ -937,9 +905,9 @@ void SMaterialParametersOverviewPanel::Refresh()
 	{
 		this->ChildSlot
 			[
-				SNew(SBorder)
-				.BorderImage(FEditorStyle::GetBrush("MaterialInstanceEditor.StackBody"))
-				.Padding(FMargin(4.0f))
+				SNew(SBox)
+				.Padding(10.0f)
+				.HAlign(HAlign_Center)
 				[
 					SNew(STextBlock)
 					.Text(LOCTEXT("ConnectMaterialParametersToFillList", "Connect a parameter to see it here."))
diff --git a/Engine/Source/Editor/MaterialEditor/Private/SMaterialParametersOverviewWidget.h b/Engine/Source/Editor/MaterialEditor/Private/SMaterialParametersOverviewWidget.h
--- a/Engine/Source/Editor/MaterialEditor/Private/SMaterialParametersOverviewWidget.h
+++ b/Engine/Source/Editor/MaterialEditor/Private/SMaterialParametersOverviewWidget.h
@@ -48,6 +48,7 @@ private:
 
 	const FSlateBrush* GetBorderImage() const;
 
+	FSlateColor GetOuterBackgroundColor(TSharedPtr<FSortedParamData> InParamData) const;
 private:
 
 	/** The node info to build the tree view row from. */
