Material Stats Update Pass

diff --git a/Engine/Source/Editor/MaterialEditor/Private/MaterialStatsCommon.cpp b/Engine/Source/Editor/MaterialEditor/Private/MaterialStatsCommon.cpp
--- a/Engine/Source/Editor/MaterialEditor/Private/MaterialStatsCommon.cpp
+++ b/Engine/Source/Editor/MaterialEditor/Private/MaterialStatsCommon.cpp
@@ -8,6 +8,7 @@
 #include "MaterialEditorSettings.h"
 #include "RHIShaderFormatDefinitions.inl"
 #include "ShaderCompilerCore.h"
+#include "Styling/StyleColors.h"
 
 /***********************************************************************************************************************/
 /*begin FMaterialResourceStats functions*/
@@ -23,11 +24,6 @@ void FMaterialResourceStats::SetupExtaCompilationSettings(const EShaderPlatform
 
 /***********************************************************************************************************************/
 /*begin FMaterialStatsUtils */
-const FLinearColor FMaterialStatsUtils::BlueColor(0.1851f, 1.0f, 0.940258f);
-const FLinearColor FMaterialStatsUtils::YellowColor(1.0f, 0.934216f, 0.199542f);
-const FLinearColor FMaterialStatsUtils::GreenColor(0.540805f, 1.0f, 0.321716f);
-const FLinearColor FMaterialStatsUtils::OrangeColor(1.0f, 0.316738f, 0.095488f);
-const FLinearColor FMaterialStatsUtils::DefaultGridTextColor(0.244819f, 0.301351f, 0.390625f);
 
 TSharedPtr<FMaterialStats> FMaterialStatsUtils::CreateMaterialStats(class IMaterialEditor* MaterialEditor)
 {
@@ -288,23 +284,23 @@ FString FMaterialStatsUtils::RepresentativeShaderTypeToString(const ERepresentat
 	}
 }
 
-FLinearColor FMaterialStatsUtils::PlatformTypeColor(EPlatformCategoryType PlatformType)
+FSlateColor FMaterialStatsUtils::PlatformTypeColor(EPlatformCategoryType PlatformType)
 {
-	FLinearColor Color(FLinearColor::Blue);
+	FSlateColor Color(FStyleColors::Foreground);
 
 	switch (PlatformType)
 	{
 		case EPlatformCategoryType::Desktop:
-			Color = BlueColor;
+			Color = FStyleColors::AccentBlue;
 		break;
 		case EPlatformCategoryType::Android:
-			Color = GreenColor;
+			Color = FStyleColors::AccentGreen;
 		break;
 		case EPlatformCategoryType::IOS:
-			Color = FLinearColor::Gray;
+			Color = FStyleColors::AccentYellow;
 		break;
 		case EPlatformCategoryType::Console:
-			Color = FLinearColor::Red;
+			Color = FStyleColors::AccentPurple;
 		break;
 
 		default:
@@ -315,28 +311,28 @@ FLinearColor FMaterialStatsUtils::PlatformTypeColor(EPlatformCategoryType Platfo
 	return Color;
 }
 
-FLinearColor FMaterialStatsUtils::QualitySettingColor(const EMaterialQualityLevel::Type QualityType)
+FSlateColor FMaterialStatsUtils::QualitySettingColor(const EMaterialQualityLevel::Type QualityType)
 {
 	switch (QualityType)
 	{
 		case EMaterialQualityLevel::Low:
-			return GreenColor;
+			return FStyleColors::AccentGreen;
 		break;
 		case EMaterialQualityLevel::High:
-			return OrangeColor;
+			return FStyleColors::AccentOrange;
 		break;
 		case EMaterialQualityLevel::Medium:
-			return YellowColor;
+			return FStyleColors::Warning;
 		break;
 		case EMaterialQualityLevel::Epic:
-			return FLinearColor::Red;
+			return FStyleColors::Error;
 		break;
 		default:
-			return FLinearColor::Black;
+			return FStyleColors::Foreground;
 		break;
 	}
 
-	return FLinearColor::Black;
+	return  FStyleColors::Foreground;
 }
 
 void FMaterialStatsUtils::GetRepresentativeShaderTypesAndDescriptions(TMap<FName, TArray<FRepresentativeShaderInfo>>& ShaderTypeNamesAndDescriptions, const FMaterial* TargetMaterial)
diff --git a/Engine/Source/Editor/MaterialEditor/Private/MaterialStatsGrid.cpp b/Engine/Source/Editor/MaterialEditor/Private/MaterialStatsGrid.cpp
--- a/Engine/Source/Editor/MaterialEditor/Private/MaterialStatsGrid.cpp
+++ b/Engine/Source/Editor/MaterialEditor/Private/MaterialStatsGrid.cpp
@@ -124,7 +124,7 @@ void FStatsGridRow_Errors::CreateRow(TSharedPtr<FMaterialStats> StatsManager)
 {
 	// add an "Error" string in the descriptor column
 	TSharedPtr<FGridCell> HeaderCell = MakeShareable(new FGridCell_StaticString(TEXT("Errors"), TEXT("Errors")));
-	HeaderCell->SetColor(FMaterialStatsUtils::OrangeColor);
+	HeaderCell->SetColor(FStyleColors::AccentOrange);
 	HeaderCell->SetContentBold(true);
 	AddCell(FMaterialStatsGrid::DescriptorColumnName, HeaderCell);
 
@@ -138,7 +138,7 @@ void FStatsGridRow_Errors::AddPlatform(TSharedPtr<FMaterialStats> StatsManager,
 	// add a cell that will query any available errors for this platform
 	const FString CellContent = FMaterialStatsUtils::MaterialQualityToShortString(QualityLevel);
 	TSharedPtr<FGridCell_ShaderValue> Cell = MakeShareable(new FGridCell_ShaderValue(StatsManager, EShaderInfoType::Errors, ERepresentativeShader::Num, QualityLevel, Platform->GetPlatformShaderType()));
-	Cell->SetColor(FMaterialStatsUtils::OrangeColor);
+	Cell->SetColor(FStyleColors::AccentOrange);
 	Cell->SetHorizontalAlignment(HAlign_Fill);
 
 	const FName ColumnName = FMaterialStatsGrid::MakePlatformColumnName(Platform, QualityLevel);
@@ -193,7 +193,7 @@ void FStatsGridRow_Shaders::CreateRow(TSharedPtr<FMaterialStats> StatsManager)
 
 		HeaderCell = MakeShareable(new FGridCell_StaticString(HeaderContent, HeaderContent));
 		HeaderCell->SetContentBold(true);
-		HeaderCell->SetColor(FLinearColor::Gray);
+		HeaderCell->SetColor(FStyleColors::Foreground);
 	}
 	else
 	{
@@ -207,7 +207,7 @@ void FStatsGridRow_Shaders::CreateRow(TSharedPtr<FMaterialStats> StatsManager)
 	TSharedPtr<FGridCell> ShaderNameCell = MakeShareable(new FGridCell_StaticString(ShaderColumnContent, ShaderColumnContent));
 	ShaderNameCell->SetHorizontalAlignment(HAlign_Fill);
 	ShaderNameCell->SetContentBold(true);
-	ShaderNameCell->SetColor(FLinearColor::Gray);
+	ShaderNameCell->SetColor(FStyleColors::Foreground);
 	AddCell(FMaterialStatsGrid::ShaderColumnName, ShaderNameCell);
 
 	FillPlatformCellsHelper(StatsManager);
@@ -262,7 +262,7 @@ void FStatsGridRow_Interpolators::CreateRow(TSharedPtr<FMaterialStats> StatsMana
 {
 	// static string in the descriptor column
 	TSharedPtr<FGridCell> HeaderCell = MakeShareable(new FGridCell_StaticString(TEXT("Interpolators"), TEXT("Interpolators")));
-	HeaderCell->SetColor(FLinearColor::Gray);
+	HeaderCell->SetColor(FStyleColors::Foreground);
 	HeaderCell->SetContentBold(true);
 	AddCell(FMaterialStatsGrid::DescriptorColumnName, HeaderCell);
 
@@ -291,7 +291,7 @@ void FStatsGridRow_NumTextureSamples::CreateRow(TSharedPtr<FMaterialStats> Stats
 {
 	// static string in the descriptor column
 	TSharedPtr<FGridCell> HeaderCell = MakeShareable(new FGridCell_StaticString(TEXT("Texture Lookups (Est.)"), TEXT("Texture Lookups (Est.)")));
-	HeaderCell->SetColor(FLinearColor::Gray);
+	HeaderCell->SetColor(FStyleColors::Foreground);
 	HeaderCell->SetContentBold(true);
 	AddCell(FMaterialStatsGrid::DescriptorColumnName, HeaderCell);
 
@@ -320,7 +320,7 @@ void FStatsGridRow_NumVirtualTextureLookups::CreateRow(TSharedPtr<FMaterialStats
 {
 	// static string in the descriptor column
 	TSharedPtr<FGridCell> HeaderCell = MakeShareable(new FGridCell_StaticString(TEXT("VT Lookups (Est.)"), TEXT("Virtual Texture Lookups (Est.)")));
-	HeaderCell->SetColor(FLinearColor::Gray);
+	HeaderCell->SetColor(FStyleColors::Foreground);
 	HeaderCell->SetContentBold(true);
 	AddCell(FMaterialStatsGrid::DescriptorColumnName, HeaderCell);
 
@@ -444,7 +444,7 @@ FString FMaterialStatsGrid::GetColumnContentLong(const FName ColumnName) const
 	return Ptr != nullptr ? Ptr->ContentLong : TEXT("");
 }
 
-FLinearColor FMaterialStatsGrid::GetColumnColor(const FName ColumnName) const
+FSlateColor FMaterialStatsGrid::GetColumnColor(const FName ColumnName) const
 {
 	auto* Ptr = GridColumnContent.Find(ColumnName);
 
@@ -811,7 +811,7 @@ FName FMaterialStatsGrid::MakePlatformColumnName(const TSharedPtr<FShaderPlatfor
 
 FGridCell::FGridCell()
 {
-	CellColor = FMaterialStatsUtils::DefaultGridTextColor;
+	CellColor = FStyleColors::Foreground;
 }
 
 FString FGridCell_Empty::GetCellContent()
diff --git a/Engine/Source/Editor/MaterialEditor/Private/MaterialStatsGrid.h b/Engine/Source/Editor/MaterialEditor/Private/MaterialStatsGrid.h
--- a/Engine/Source/Editor/MaterialEditor/Private/MaterialStatsGrid.h
+++ b/Engine/Source/Editor/MaterialEditor/Private/MaterialStatsGrid.h
@@ -9,13 +9,14 @@
 #include "MaterialStatsCommon.h"
 #include "Containers/StaticArray.h"
 #include "MaterialStats.h"
+#include "Styling/StyleColors.h"
 
 /** class to represent a single cell inside the material stats grid */
 class FGridCell
 {
 protected:
 	/** attributes used at display time to configure widgets */
-	FLinearColor CellColor;
+	FSlateColor CellColor;
 	bool bBoldContent = false;
 
 	EHorizontalAlignment HAlignment = EHorizontalAlignment::HAlign_Center;
@@ -30,8 +31,8 @@ public:
 	/** this can be used for tool tips or other detailed descriptions */
 	virtual FString GetCellContentLong() = 0;
 
-	FORCEINLINE FLinearColor GetColor() const;
-	FORCEINLINE void SetColor(const FLinearColor& Color);
+	FORCEINLINE FSlateColor GetColor() const;
+	FORCEINLINE void SetColor(const FSlateColor& Color);
 
 	FORCEINLINE bool IsContentBold() const;
 	FORCEINLINE void SetContentBold(bool bValue);
@@ -254,7 +255,7 @@ class FMaterialStatsGrid
 	{
 		FString Content = TEXT("");
 		FString ContentLong = TEXT("");
-		FLinearColor Color = FLinearColor::Gray;
+		FSlateColor Color = FStyleColors::Foreground;
 	};
 
 	/** collection of column information sorted by their names */
@@ -328,7 +329,7 @@ public:
 
 	FString GetColumnContent(const FName ColumnName) const;
 	FString GetColumnContentLong(const FName ColumnName) const;
-	FLinearColor GetColumnColor(const FName ColumnName) const;
+	FSlateColor GetColumnColor(const FName ColumnName) const;
 
 	/** helper function that will assemble a column name from the given arguments */
 	static FName MakePlatformColumnName(const TSharedPtr<FShaderPlatformSettings>& Platform, const EMaterialQualityLevel::Type Quality);
@@ -362,12 +363,12 @@ FORCEINLINE TArray<FName> FMaterialStatsGrid::GetVisibleColumnNames() const
 }
 //////////////////////////////////////////////////////////////////////////////////////////////////
 
-FORCEINLINE FLinearColor FGridCell::GetColor() const
+FORCEINLINE FSlateColor FGridCell::GetColor() const
 {
 	return CellColor;
 }
 
-FORCEINLINE void FGridCell::SetColor(const FLinearColor& Color)
+FORCEINLINE void FGridCell::SetColor(const FSlateColor& Color)
 {
 	CellColor = Color;
 }
diff --git a/Engine/Source/Editor/MaterialEditor/Private/SMaterialEditorStatsWidget.cpp b/Engine/Source/Editor/MaterialEditor/Private/SMaterialEditorStatsWidget.cpp
--- a/Engine/Source/Editor/MaterialEditor/Private/SMaterialEditorStatsWidget.cpp
+++ b/Engine/Source/Editor/MaterialEditor/Private/SMaterialEditorStatsWidget.cpp
@@ -14,6 +14,7 @@
 #include "EditorStyleSet.h"
 #include "MaterialStatsGrid.h"
 #include "Framework/Application/SlateApplication.h"
+#include "Styling/StyleColors.h"
 
 #define LOCTEXT_NAMESPACE "MaterialEditorStatsWidget"
 
@@ -42,7 +43,7 @@ public:
 
 		SMultiColumnTableRow<TSharedPtr<int32>>::Construct(
 			FSuperRowType::FArguments()
-			.Style(FEditorStyle::Get(), "DataTableEditor.CellListViewRow"),
+			.Style(FAppStyle::Get(), "TableView.Row"),
 			InOwnerTableView
 		);
 	}
@@ -88,7 +89,6 @@ TSharedRef<SWidget> SMaterialStatsViewRow::GenerateWidgetForColumn(const FName&
 		[
 			SNew(STextBlock)
 			.TextStyle(FEditorStyle::Get(), UsedFontStyle)
-			.ColorAndOpacity(this, &SMaterialStatsViewRow::GetColorForCell, ColumnName)
 			.Text(this, &SMaterialStatsViewRow::GetTextForCell, ColumnName, false)
 			.ToolTipText(this, &SMaterialStatsViewRow::GetTextForCell, ColumnName, true)
 			.AutoWrapText(true)
@@ -117,19 +117,7 @@ FText SMaterialStatsViewRow::GetTextForCell(const FName Name, const bool bToolTi
 
 FSlateColor SMaterialStatsViewRow::GetColorForCell(const FName Name) const
 {
-	FSlateColor Color(FColorList::White);
-
-	auto StatsPtr = MaterialStatsWPtr.Pin();
-	if (StatsPtr.IsValid() && PtrRowID.IsValid())
-	{
-		const int32 RowID = *PtrRowID;
-
-		const auto Cell = StatsPtr->GetStatsGrid()->GetCell(RowID, Name);
-
-		return Cell->GetColor();
-	}
-
-	return Color;
+	return FStyleColors::Foreground;
 }
 
 EHorizontalAlignment SMaterialStatsViewRow::GetHAlignForCell(const FName Name) const
@@ -209,7 +197,7 @@ float SMaterialEditorStatsWidget::GetColumnSize(const FName ColumnName) const
 
 SHeaderRow::FColumn::FArguments SMaterialEditorStatsWidget::CreateColumnArgs(const FName ColumnName)
 {
-	FLinearColor Color = FLinearColor::Gray;
+	FSlateColor Color = FStyleColors::AccentGreen;
 	FString Content;
 	FString ContentLong;
 	auto StatsPtr = MaterialStatsWPtr.Pin();
@@ -227,6 +215,7 @@ SHeaderRow::FColumn::FArguments SMaterialEditorStatsWidget::CreateColumnArgs(con
 		.HeaderContent()
 		[
 			SNew(STextBlock)
+			.Margin(4.0f)
 			.ColorAndOpacity(Color)
 			.Text(FText::FromString(Content))
 			.ToolTipText(FText::FromString(ContentLong))
@@ -291,6 +280,40 @@ void SMaterialEditorStatsWidget::RebuildColumns()
 	}
 }
 
+uint32 SMaterialEditorStatsWidget::CountSubPlatforms(EPlatformCategoryType Category) const
+{
+	uint32 PlatformCount = 0;
+	const auto StatsPtr = MaterialStatsWPtr.Pin();
+
+	if (!StatsPtr.IsValid())
+	{
+		return 0;
+	}
+
+	auto& PlatformsDB = StatsPtr->GetPlatformsTypeDB();
+	auto* ArrPlaformsPtr = PlatformsDB.Find(Category);
+	if (ArrPlaformsPtr == nullptr)
+	{
+		return 0;
+	}
+
+	auto& ArrPlatforms = *ArrPlaformsPtr;
+
+	for (int32 i = 0; i < ArrPlatforms.Num(); ++i)
+	{
+		auto PlatformPtr = ArrPlatforms[i];
+		if (!PlatformPtr.IsValid() || !PlatformPtr->IsStatsGridPresenceAllowed())
+		{
+			continue;
+		}
+
+		PlatformCount++;
+	}
+
+	return PlatformCount;
+
+}
+
 void SMaterialEditorStatsWidget::CreatePlatformMenus(FMenuBuilder& Builder, EPlatformCategoryType Category)
 {
 	const auto StatsPtr = MaterialStatsWPtr.Pin();
@@ -360,15 +383,12 @@ void SMaterialEditorStatsWidget::CreatePlatformMenus(FMenuBuilder& Builder, EPla
 			WidgetPtr->RequestRefresh();
 		};
 
-//		const FLinearColor PlatformColor = FMaterialStatsUtils::PlatformTypeColor(PlatformPtr->GetCategoryType());
 		auto PlatformWidget = SNew(SCheckBox)			
 			.OnCheckStateChanged_Lambda(Lamda_PlatformFlipState)
 			.IsChecked_Lambda(Lamda_PlatformCheckState)
 			.Content()
 			[
 				SNew(STextBlock)
-//				.ColorAndOpacity(PlatformColor)
-				.TextStyle(FEditorStyle::Get(), "RichTextBlock.Bold")
 				.Text(FText::FromName(PlatformName))
 				.Margin(FMargin(2.0f, 2.0f, 4.0f, 2.0f))
 			];
@@ -383,14 +403,13 @@ void SMaterialEditorStatsWidget::CreatePlatformCategoryMenus(FMenuBuilder& Build
 	{
 		const EPlatformCategoryType PlatformType = static_cast<EPlatformCategoryType>(i);
 
-//		const FLinearColor PlatformColor = FMaterialStatsUtils::PlatformTypeColor(PlatformType);
-
-		auto PlatformTypeWidget = SNew(STextBlock)
-			.TextStyle(FEditorStyle::Get(), "RichTextBlock.Bold")
-			.Text(FText::FromString(FMaterialStatsUtils::GetPlatformTypeName(PlatformType)));
-//			.ColorAndOpacity(PlatformColor);
+		if (CountSubPlatforms(PlatformType) > 0)
+		{
+			auto PlatformTypeWidget = SNew(STextBlock)
+				.Text(FText::FromString(FMaterialStatsUtils::GetPlatformTypeName(PlatformType)));
 
-		Builder.AddSubMenu(PlatformTypeWidget, FNewMenuDelegate::CreateSP(this, &SMaterialEditorStatsWidget::CreatePlatformMenus, PlatformType));
+			Builder.AddSubMenu(PlatformTypeWidget, FNewMenuDelegate::CreateSP(this, &SMaterialEditorStatsWidget::CreatePlatformMenus, PlatformType));
+		}
 	}
 }
 
@@ -463,7 +482,6 @@ void SMaterialEditorStatsWidget::CreateQualityMenus(FMenuBuilder& Builder)
 			return ECheckBoxState::Unchecked;
 		};
 
-//		const FLinearColor QualityColor = FMaterialStatsUtils::QualitySettingColor(QualityLevel);
 		const FText QualitySettingName = FText::FromString(FMaterialStatsUtils::MaterialQualityToString(QualityLevel));
 
 		auto QualityWidget = SNew(SCheckBox)
@@ -472,8 +490,6 @@ void SMaterialEditorStatsWidget::CreateQualityMenus(FMenuBuilder& Builder)
 			.Content()
 			[
 				SNew(STextBlock)
-//				.ColorAndOpacity(QualityColor)
-				.TextStyle(FEditorStyle::Get(), "RichTextBlock.Bold")
 				.Text(QualitySettingName)
 			];
 
@@ -483,11 +499,7 @@ void SMaterialEditorStatsWidget::CreateQualityMenus(FMenuBuilder& Builder)
 
 void SMaterialEditorStatsWidget::CreateGlobalQualityMenu(class FMenuBuilder& Builder)
 {
-//	FColor GreenShade(119, 184, 42);
-
 	auto GlobalQualityWidget = SNew(STextBlock)
-//		.ColorAndOpacity(FLinearColor(GreenShade))
-		.TextStyle(FEditorStyle::Get(), "RichTextBlock.Bold")
 		.Text(LOCTEXT("GlobalQualitySettings", "Global Quality Settings"));
 
 	Builder.AddSubMenu(GlobalQualityWidget, FNewMenuDelegate::CreateSP(this, &SMaterialEditorStatsWidget::CreateQualityMenus));
@@ -536,7 +548,7 @@ void SMaterialEditorStatsWidget::AddWarningMessage(const FString& Message)
 				[
 					SNew(STextBlock)
 					.TextStyle(FEditorStyle::Get(), TEXT("RichTextBlock.Bold"))
-					.ColorAndOpacity(FMaterialStatsUtils::YellowColor)
+					.ColorAndOpacity(FStyleColors::Warning)
 					.Text(FText::FromString(Message))
 					.ToolTipText(FText::FromString(Message))
 				]
@@ -580,41 +592,43 @@ void SMaterialEditorStatsWidget::Construct(const FArguments& InArgs)
 	[
 		SNew(SVerticalBox)
 		+SVerticalBox::Slot() // this will contain the tool bar
+		.Padding(0.0f)
 		.AutoHeight()
 		[
-			SNew(SBorder)
-			.BorderImage(FEditorStyle::GetBrush(TEXT("ToolPanel.GroupBorder")))
+			SNew(SHorizontalBox)
+			+SHorizontalBox::Slot()
+			.FillWidth(1.0)
 			[
-				SNew(SHorizontalBox)
-				+SHorizontalBox::Slot()
-				.AutoWidth()
+				SNullWidget::NullWidget
+			]
+			+SHorizontalBox::Slot()
+			.AutoWidth()
+			[
+				SNew(SComboButton)
+				.ComboButtonStyle(FAppStyle::Get(), "SimpleComboButton")
+				.ForegroundColor(FSlateColor::UseStyle())
+				.ContentPadding(0)
+				.OnGetMenuContent(this, &SMaterialEditorStatsWidget::GetSettingsButtonContent)
+				.HasDownArrow(true)
+				.ButtonContent()
 				[
-					SNew(SComboButton)
-					.ComboButtonStyle(FEditorStyle::Get(), "ToolbarComboButton")
-					.ForegroundColor(FLinearColor::White)
-					.ContentPadding(0)
-					.OnGetMenuContent(this, &SMaterialEditorStatsWidget::GetSettingsButtonContent)
-					.ButtonContent()
+					SNew(SHorizontalBox)
+					// Icon
+					+ SHorizontalBox::Slot()
+					.AutoWidth()
+					.VAlign(VAlign_Center)
 					[
-						SNew(SHorizontalBox)
-						// Icon
-						+ SHorizontalBox::Slot()
-						.AutoWidth()
-						.VAlign(VAlign_Center)
-						[
-							SNew(SImage)
-							.Image(FEditorStyle::GetBrush("LevelEditor.Tabs.StatsViewer"))
-						]
-						// Text
-						+ SHorizontalBox::Slot()
-						.AutoWidth()
-						.VAlign(VAlign_Center)
-						.Padding(0, 0, 2, 0)
-						[
-							SNew(STextBlock)
-							.TextStyle(FEditorStyle::Get(), "ContentBrowser.TopBar.Font")
-							.Text(LOCTEXT("SettingsButton", "Settings"))
-						]
+						SNew(SImage)
+						.Image(FAppStyle::Get().GetBrush("LevelEditor.GameSettings"))
+					]
+					// Text
+					+ SHorizontalBox::Slot()
+					.AutoWidth()
+					.VAlign(VAlign_Center)
+					.Padding(2, 0, 2, 0)
+					[
+						SNew(STextBlock)
+						.Text(LOCTEXT("SettingsButton", "Settings"))
 					]
 				]
 			]
@@ -623,7 +637,9 @@ void SMaterialEditorStatsWidget::Construct(const FArguments& InArgs)
 		.AutoHeight()
 		[
 			SNew(SBorder)
-			.BorderImage(FEditorStyle::GetBrush(TEXT("Graph.TitleBackground")))
+			.Padding(0.0f)
+			.HAlign(HAlign_Fill)
+			.BorderImage(FAppStyle::Get().GetBrush(TEXT("Brushes.Recessed")))
 			[
 				MessageArea.ToSharedRef()
 			]
@@ -631,7 +647,8 @@ void SMaterialEditorStatsWidget::Construct(const FArguments& InArgs)
 		+ SVerticalBox::Slot() // this will contain the stats grid
 		[
 			SNew(SBorder)
-			.BorderImage(FEditorStyle::GetBrush(TEXT("Graph.TitleBackground")))
+			.BorderImage(FAppStyle::Get().GetBrush(TEXT("Brushes.Recessed")))
+			.Padding(0.0f)
 			.HAlign(HAlign_Fill)
 			[
 				SNew(SHorizontalBox)
@@ -652,7 +669,6 @@ void SMaterialEditorStatsWidget::Construct(const FArguments& InArgs)
 							+ SHorizontalBox::Slot()
 							.HAlign(HAlign_Fill)
 							.FillWidth(1.f)
-							.Padding(5, 0)
 							[
 								SNew(SScrollBox)
 								.Orientation(Orient_Horizontal)
diff --git a/Engine/Source/Editor/MaterialEditor/Private/SMaterialEditorStatsWidget.h b/Engine/Source/Editor/MaterialEditor/Private/SMaterialEditorStatsWidget.h
--- a/Engine/Source/Editor/MaterialEditor/Private/SMaterialEditorStatsWidget.h
+++ b/Engine/Source/Editor/MaterialEditor/Private/SMaterialEditorStatsWidget.h
@@ -41,6 +41,7 @@ public:
 
 private:
 	void RebuildColumns();
+	uint32 CountSubPlatforms(EPlatformCategoryType Category) const;
 	void AddColumn(const FName ColumnName);
 	void InsertColumnAfter(const FName ColumnName, const FName PreviousColumn);
 	void RemoveColumn(const FName ColumnName);
diff --git a/Engine/Source/Editor/MaterialEditor/Public/MaterialStatsCommon.h b/Engine/Source/Editor/MaterialEditor/Public/MaterialStatsCommon.h
--- a/Engine/Source/Editor/MaterialEditor/Public/MaterialStatsCommon.h
+++ b/Engine/Source/Editor/MaterialEditor/Public/MaterialStatsCommon.h
@@ -77,12 +77,6 @@ public:
 		{}
 	};
 
-	static const FLinearColor BlueColor;
-	static const FLinearColor YellowColor;
-	static const FLinearColor GreenColor;
-	static const FLinearColor OrangeColor;
-	static const FLinearColor DefaultGridTextColor;
-
 public:
 	/** call this to create an instance to FMaterialStats */
 	static TSharedPtr<class FMaterialStats> CreateMaterialStats(class IMaterialEditor* MaterialEditor);
@@ -106,8 +100,8 @@ public:
 
 	static FString RepresentativeShaderTypeToString(const ERepresentativeShader ShaderType);
 
-	static FLinearColor QualitySettingColor(const EMaterialQualityLevel::Type QualityType);
-	static FLinearColor PlatformTypeColor(EPlatformCategoryType PlatformType);
+	static FSlateColor QualitySettingColor(const EMaterialQualityLevel::Type QualityType);
+	static FSlateColor PlatformTypeColor(EPlatformCategoryType PlatformType);
 
 	MATERIALEDITOR_API static bool IsPlatformOfflineCompilerAvailable(const EShaderPlatform ShaderPlatform);
 	MATERIALEDITOR_API static FString GetPlatformOfflineCompilerPath(const EShaderPlatform ShaderPlatform);
