Adding Edit Condition Checkbox to Material Layer Instances

diff --git a/Engine/Source/Editor/MaterialEditor/Private/SMaterialLayersFunctionsTree.cpp b/Engine/Source/Editor/MaterialEditor/Private/SMaterialLayersFunctionsTree.cpp
--- a/Engine/Source/Editor/MaterialEditor/Private/SMaterialLayersFunctionsTree.cpp
+++ b/Engine/Source/Editor/MaterialEditor/Private/SMaterialLayersFunctionsTree.cpp
@@ -905,7 +905,19 @@ void SMaterialLayersFunctionsInstanceTreeItem::Construct(const FArguments& InArg
 		}
 
 		FNodeWidgets NodeWidgets = Node.CreateNodeWidgets();
-		LeftSideWidget = NodeWidgets.NameWidget.ToSharedRef();
+		LeftSideWidget = SNew(SHorizontalBox)
+			+SHorizontalBox::Slot()
+			.Padding(6.0f, 0.0f, 6.0f, 0.0f)
+			.HAlign(HAlign_Center)
+			.VAlign(VAlign_Center)
+			.AutoWidth()
+			[
+				NodeWidgets.EditConditionWidget.ToSharedRef()
+			]
+			+SHorizontalBox::Slot()
+			[
+				NodeWidgets.NameWidget.ToSharedRef()
+			];
 		RightSideWidget = NodeWidgets.ValueWidget.ToSharedRef();
 
 		ResetWidget = SNew(SResetToDefaultPropertyEditor, Node.CreatePropertyHandle())
