Fixing non-unity errors

diff --git a/Engine/Source/Editor/MaterialEditor/Private/MaterialStatsGrid.cpp b/Engine/Source/Editor/MaterialEditor/Private/MaterialStatsGrid.cpp
--- a/Engine/Source/Editor/MaterialEditor/Private/MaterialStatsGrid.cpp
+++ b/Engine/Source/Editor/MaterialEditor/Private/MaterialStatsGrid.cpp
@@ -107,7 +107,7 @@ void FStatsGridRow_Quality::AddPlatform(TSharedPtr<FMaterialStats> StatsManager,
 	TSharedPtr<FGridCell_StaticString> Cell = MakeShareable(new FGridCell_StaticString(CellContent, CellContent));
 
 	Cell->SetContentBold(true);
-	auto CellColor = FMaterialStatsUtils::QualitySettingColor(QualityLevel);
+	FSlateColor CellColor = FMaterialStatsUtils::QualitySettingColor(QualityLevel);
 	Cell->SetColor(CellColor);
 
 	const FName ColumnName = FMaterialStatsGrid::MakePlatformColumnName(Platform, QualityLevel);
diff --git a/Engine/Source/Editor/MaterialEditor/Public/MaterialStatsCommon.h b/Engine/Source/Editor/MaterialEditor/Public/MaterialStatsCommon.h
--- a/Engine/Source/Editor/MaterialEditor/Public/MaterialStatsCommon.h
+++ b/Engine/Source/Editor/MaterialEditor/Public/MaterialStatsCommon.h
@@ -4,6 +4,7 @@
 #include "SceneTypes.h"
 #include "RHIDefinitions.h"
 #include "MaterialShared.h"
+#include "Styling/SlateColor.h"
 
 /** custom resource material class used to mark the resource as used for shader stats extraction */
 class MATERIALEDITOR_API FMaterialResourceStats : public FMaterialResource
