Material and Material Instance editors update v1.0

diff --git a/Engine/Source/Editor/EditorStyle/Private/StarshipStyle.cpp b/Engine/Source/Editor/EditorStyle/Private/StarshipStyle.cpp
--- a/Engine/Source/Editor/EditorStyle/Private/StarshipStyle.cpp
+++ b/Engine/Source/Editor/EditorStyle/Private/StarshipStyle.cpp
@@ -3976,7 +3976,7 @@ void FStarshipEditorStyle::FStyle::SetupGraphEditorStyles()
 		Set("ClassIcon.K2Node_DynamicCast",		new IMAGE_BRUSH("Icons/icon_Blueprint_Cast_16x", Icon16x16));
 		Set("ClassIcon.EdGraphNode_Comment",	new IMAGE_BRUSH("Icons/icon_Blueprint_Comment_16x", Icon16x16));
 
-
+		Set("GraphEditor.Clean",				new IMAGE_BRUSH_SVG("Starship/GraphEditors/CleanUp", Icon20x20));
 		Set("GraphEditor.OverrideFunction_16x", new IMAGE_BRUSH_SVG("Starship/Blueprints/icon_Blueprint_OverrideFunction", Icon16x16));
 		Set("GraphEditor.OverrideFunction_24x", new IMAGE_BRUSH_SVG("Starship/Blueprints/icon_Blueprint_OverrideFunction", Icon24x24));
 		Set("GraphEditor.EventGraph_16x", new IMAGE_BRUSH_SVG("Starship/Blueprints/icon_BlueprintEditor_EventGraph", Icon16x16));
@@ -6601,49 +6601,22 @@ void FStarshipEditorStyle::FStyle::SetupToolkitStyles()
 		Set( "MaterialEditor.NormalFont", DEFAULT_FONT( "Regular", 9 ) );
 		Set( "MaterialEditor.BoldFont", DEFAULT_FONT( "Bold", 9 ) );
 
-		Set( "MaterialEditor.Apply", new IMAGE_BRUSH( "Icons/icon_MatEd_Apply_40x", Icon40x40 ) );
-		Set( "MaterialEditor.Apply.Small", new IMAGE_BRUSH( "Icons/icon_MatEd_Apply_40x", Icon20x20 ) );
-
-		Set( "MaterialEditor.ShowAllMaterialParameters", new IMAGE_BRUSH( "Icons/icon_MatInsEd_Params_40x", Icon40x40 ) );
-		Set( "MaterialEditor.ShowAllMaterialParameters.Small", new IMAGE_BRUSH( "Icons/icon_MatInsEd_Params_40x", Icon20x20 ) );
-
-		Set( "MaterialEditor.SetCylinderPreview", new IMAGE_BRUSH( "Icons/icon_MatEd_Cylinder_40x", Icon40x40 ) );
-		Set( "MaterialEditor.SetCylinderPreview.Small", new IMAGE_BRUSH( "Icons/icon_MatEd_Cylinder_40x", Icon20x20 ) );
-		Set( "MaterialEditor.SetSpherePreview", new IMAGE_BRUSH( "Icons/icon_MatEd_Sphere_40x", Icon40x40 ) );
-		Set( "MaterialEditor.SetSpherePreview.Small", new IMAGE_BRUSH( "Icons/icon_MatEd_Sphere_40x", Icon20x20 ) );
-		Set( "MaterialEditor.SetPlanePreview", new IMAGE_BRUSH( "Icons/icon_MatEd_Plane_40x", Icon40x40 ) );
-		Set( "MaterialEditor.SetPlanePreview.Small", new IMAGE_BRUSH( "Icons/icon_MatEd_Plane_40x", Icon20x20 ) );
-		Set( "MaterialEditor.SetCubePreview", new IMAGE_BRUSH( "Icons/icon_MatEd_Cube_40x", Icon40x40 ) );
-		Set( "MaterialEditor.SetCubePreview.Small", new IMAGE_BRUSH( "Icons/icon_MatEd_Cube_40x", Icon20x20 ) );
-		Set( "MaterialEditor.SetPreviewMeshFromSelection", new IMAGE_BRUSH( "Icons/icon_MatEd_Mesh_40x", Icon40x40 ) );
-		Set( "MaterialEditor.SetPreviewMeshFromSelection.Small", new IMAGE_BRUSH( "Icons/icon_MatEd_Mesh_40x", Icon20x20 ) );
-		Set( "MaterialEditor.TogglePreviewGrid", new IMAGE_BRUSH( "Icons/icon_MatEd_Grid_40x", Icon40x40 ) );
-		Set( "MaterialEditor.TogglePreviewGrid.Small", new IMAGE_BRUSH( "Icons/icon_MatEd_Grid_40x", Icon20x20 ) );
-
-		Set( "MaterialEditor.ToggleMaterialStats", new IMAGE_BRUSH( "Icons/icon_MatEd_Stats_40x", Icon40x40 ) );
-		Set( "MaterialEditor.ToggleMaterialStats.Small", new IMAGE_BRUSH( "Icons/icon_MatEd_Stats_40x", Icon20x20 ) );
-		Set("MaterialEditor.ToggleMaterialStats.Tab", new IMAGE_BRUSH("Icons/icon_MatEd_Stats_40x", Icon16x16));
-		Set( "MaterialEditor.ToggleBuiltinStats", new IMAGE_BRUSH( "Icons/icon_MatEd_BuiltInStats_40x", Icon40x40 ) );
-		Set( "MaterialEditor.ToggleBuiltinStats.Small", new IMAGE_BRUSH( "Icons/icon_MatEd_BuiltInStats_40x", Icon20x20 ) );
-		Set( "MaterialEditor.TogglePlatformStats", new IMAGE_BRUSH( "Icons/icon_MobileStats_40x", Icon40x40 ) );
-		Set( "MaterialEditor.TogglePlatformStats.Small", new IMAGE_BRUSH( "Icons/icon_MobileStats_40x", Icon20x20 ) );
-		Set( "MaterialEditor.CleanUnusedExpressions", new IMAGE_BRUSH( "Icons/icon_MatEd_CleanUp_40x", Icon40x40 ) );
-		Set( "MaterialEditor.CleanUnusedExpressions.Small", new IMAGE_BRUSH( "Icons/icon_MatEd_CleanUp_40x", Icon20x20 ) );
-		Set( "MaterialEditor.ToggleRealtimeExpressions", new IMAGE_BRUSH( "Icons/icon_MatEd_LiveNodes_40x", Icon40x40 ) );
-		Set( "MaterialEditor.ToggleRealtimeExpressions.Small", new IMAGE_BRUSH( "Icons/icon_MatEd_LiveNodes_40x", Icon20x20 ) );
-		Set( "MaterialEditor.AlwaysRefreshAllPreviews", new IMAGE_BRUSH( "Icons/icon_MatEd_Refresh_40x", Icon40x40 ) );
-		Set( "MaterialEditor.AlwaysRefreshAllPreviews.Small", new IMAGE_BRUSH( "Icons/icon_MatEd_Refresh_40x", Icon20x20 ) );
-		Set( "MaterialEditor.ToggleLivePreview", new IMAGE_BRUSH( "Icons/icon_MatEd_LivePreview_40x", Icon40x40 ) );
-		Set( "MaterialEditor.ToggleLivePreview.Small", new IMAGE_BRUSH( "Icons/icon_MatEd_LivePreview_40x", Icon20x20 ) );
-		Set( "MaterialEditor.ShowHideConnectors", new IMAGE_BRUSH( "Icons/icon_MatEd_Connectors_40x", Icon40x40 ) );
-		Set( "MaterialEditor.ShowHideConnectors.Small", new IMAGE_BRUSH( "Icons/icon_MatEd_Connectors_40x", Icon20x20 ) );
-		Set( "MaterialEditor.CameraHome", new IMAGE_BRUSH( "Icons/icon_MatEd_Home_40x", Icon40x40 ) );
-		Set( "MaterialEditor.CameraHome.Small", new IMAGE_BRUSH( "Icons/icon_MatEd_Home_40x", Icon20x20 ) );
-		Set( "MaterialEditor.FindInMaterial", new IMAGE_BRUSH( "Icons/icon_Blueprint_Find_40px", Icon40x40 ) );
-		Set( "MaterialEditor.FindInMaterial.Small", new IMAGE_BRUSH( "Icons/icon_Blueprint_Find_40px", Icon20x20 ) );
-
-		Set("MaterialEditor.AddSubstanceSpecialToolBar", new IMAGE_BRUSH("Icons/icon_SubstanceFactory_White_40x", Icon40x40));
-		Set("MaterialEditor.AddSubstanceSpecialMenu", new IMAGE_BRUSH("Icons/icon_SubstanceFactory_White_40x", Icon32x32));
+		Set( "MaterialEditor.Apply", new IMAGE_BRUSH_SVG( "Starship/Common/Apply", Icon20x20 ) );
+		Set( "MaterialEditor.LiveUpdate", new IMAGE_BRUSH_SVG("Starship/MaterialEditor/LiveUpdate", Icon20x20));
+		Set( "MaterialEditor.Hierarchy", new IMAGE_BRUSH_SVG("Starship/MaterialEditor/Hierarchy", Icon20x20));
+
+		Set( "MaterialEditor.SetCylinderPreview", new IMAGE_BRUSH("Icons/AssetIcons/Cylinder_16x", Icon16x16));
+		Set( "MaterialEditor.SetSpherePreview", new IMAGE_BRUSH("Icons/AssetIcons/Sphere_16x", Icon16x16));
+		Set( "MaterialEditor.SetPlanePreview", new IMAGE_BRUSH( "Icons/AssetIcons/Plane_16x", Icon16x16) );
+		Set( "MaterialEditor.SetCubePreview", new IMAGE_BRUSH( "Icons/AssetIcons/Cube_16x", Icon16x16) );
+		Set( "MaterialEditor.SetPreviewMeshFromSelection", new IMAGE_BRUSH_SVG( "Starship/AssetIcons/StaticMesh_64", Icon16x16 ) );
+		Set( "MaterialEditor.TogglePreviewGrid", new IMAGE_BRUSH_SVG( "Starship/Common/Grid", Icon16x16 ) );
+		Set( "MaterialEditor.ToggleMaterialStats", new IMAGE_BRUSH_SVG( "Starship/Common/Statistics", Icon20x20 ) );
+		Set( "MaterialEditor.ToggleMaterialStats.Tab", new IMAGE_BRUSH_SVG("Starship/Common/Statistics", Icon16x16));
+		Set( "MaterialEditor.TogglePlatformStats", new IMAGE_BRUSH_SVG( "Starship/MaterialEditor/PlatformStats", Icon20x20 ) );
+		Set("MaterialEditor.TogglePlatformStats.Tab", new IMAGE_BRUSH_SVG("Starship/MaterialEditor/PlatformStats", Icon16x16));
+		Set( "MaterialEditor.CameraHome", new IMAGE_BRUSH_SVG( "Starship/Common/Home", Icon20x20 ) );
+		Set( "MaterialEditor.FindInMaterial", new IMAGE_BRUSH_SVG( "Starship/Common/Search", Icon20x20 ) );
 
 
 		Set("MaterialEditor.CellListViewRow", FTableRowStyle(NormalTableRowStyle)
@@ -6663,7 +6636,7 @@ void FStarshipEditorStyle::FStyle::SetupToolkitStyles()
 	// Material Instance Editor
 	{
 		Set( "MaterialInstanceEditor.Tabs.Properties", new IMAGE_BRUSH( "/Icons/icon_tab_SelectionDetails_16x", Icon16x16 ) );
-		Set( "MaterialInstanceEditor.Tabs.Parents", new IMAGE_BRUSH( "/Icons/layers_16x", Icon16x16 ) );
+		Set("MaterialEditor.ShowAllMaterialParameters", new IMAGE_BRUSH_SVG("Starship/GraphEditors/HideUnrelated", Icon20x20));
 	}
 	
 	// Sound Class Editor
diff --git a/Engine/Source/Editor/GraphEditor/Private/SGraphEditorImpl.cpp b/Engine/Source/Editor/GraphEditor/Private/SGraphEditorImpl.cpp
--- a/Engine/Source/Editor/GraphEditor/Private/SGraphEditorImpl.cpp
+++ b/Engine/Source/Editor/GraphEditor/Private/SGraphEditorImpl.cpp
@@ -1061,7 +1061,7 @@ void SGraphEditorImpl::RegisterContextMenu(const UEdGraphSchema* Schema, FToolMe
 		Menu->AddDynamicSection("EdGraphSchemaPinActions", FNewToolMenuDelegate::CreateLambda([](UToolMenu* InMenu)
 			{
 				UGraphNodeContextMenuContext* NodeContext = InMenu->FindContext<UGraphNodeContextMenuContext>();
-				if (NodeContext && NodeContext->Graph)
+				if (NodeContext && NodeContext->Graph && NodeContext->Pin)
 				{
 					if (TSharedPtr<SGraphEditorImpl> GraphEditor = StaticCastSharedPtr<SGraphEditorImpl>(FindGraphEditorForGraph(NodeContext->Graph)))
 					{
diff --git a/Engine/Source/Editor/MaterialEditor/Private/MaterialEditor.cpp b/Engine/Source/Editor/MaterialEditor/Private/MaterialEditor.cpp
--- a/Engine/Source/Editor/MaterialEditor/Private/MaterialEditor.cpp
+++ b/Engine/Source/Editor/MaterialEditor/Private/MaterialEditor.cpp
@@ -1647,100 +1647,114 @@ void FMaterialEditor::RegisterToolBar()
 		UToolMenu* ToolBar = UToolMenus::Get()->RegisterMenu(MenuName, "AssetEditor.DefaultToolBar", EMultiBoxType::ToolBar);
 
 		FToolMenuInsert InsertAfterAssetSection("Asset", EToolMenuInsertType::After);
-
-		{
-			FToolMenuSection& Section = ToolBar->AddSection("Apply", TAttribute<FText>(), InsertAfterAssetSection);
-			Section.AddEntry(FToolMenuEntry::InitToolBarButton(FMaterialEditorCommands::Get().Apply));
+		FToolMenuSection& MaterialSection = ToolBar->AddSection("MaterialTools", TAttribute<FText>(), InsertAfterAssetSection);
+		MaterialSection.AddEntry(FToolMenuEntry::InitToolBarButton(FMaterialEditorCommands::Get().Apply));
+		MaterialSection.AddEntry(FToolMenuEntry::InitToolBarButton(FMaterialEditorCommands::Get().FindInMaterial));
+		MaterialSection.AddEntry(FToolMenuEntry::InitToolBarButton(FMaterialEditorCommands::Get().CameraHome));
+		UMaterialEditorMenuContext* Context = ToolBar->FindContext<UMaterialEditorMenuContext>();
+		if (!MaterialFunction)
+		{
+			MaterialSection.AddEntry(FToolMenuEntry::InitComboButton(
+				"Hierarchy",
+				FToolUIActionChoice(),
+				FNewToolMenuDelegate::CreateLambda([](UToolMenu* InSubMenu)
+					{
+						UMaterialEditorMenuContext* SubMenuContext = InSubMenu->FindContext<UMaterialEditorMenuContext>();
+						if (SubMenuContext && SubMenuContext->MaterialEditor.IsValid())
+						{
+							SubMenuContext->MaterialEditor.Pin()->GenerateInheritanceMenu(InSubMenu);
+						}
+					}),
+				LOCTEXT("Hierarchy", "Hierarchy"),
+				FText::GetEmpty(),
+				FSlateIcon(FAppStyle::Get().GetStyleSetName(), "MaterialEditor.Hierarchy"),
+				false
+				));
 		}
+		FToolMenuEntry LiveUpdateMenu = FToolMenuEntry::InitComboButton(
+			"LiveUpdate",
+			FUIAction(),
+			FNewToolMenuDelegate::CreateLambda([](UToolMenu* InSubMenu)
+				{
+					FToolMenuSection& Section = InSubMenu->FindOrAddSection("LiveUpdateOptions");
+					Section.AddEntry(FToolMenuEntry::InitMenuEntry(FMaterialEditorCommands::Get().ToggleLivePreview));
+					Section.AddEntry(FToolMenuEntry::InitMenuEntry(FMaterialEditorCommands::Get().ToggleRealtimeExpressions));
+					Section.AddEntry(FToolMenuEntry::InitMenuEntry(FMaterialEditorCommands::Get().AlwaysRefreshAllPreviews));
+				}),
+			LOCTEXT("LiveUpdate_Label", "Live Update"),
+					LOCTEXT("LiveUpdate_Tooltip", "Set the elements of the Material Editor UI to update in realtime"),
+					FSlateIcon(FAppStyle::Get().GetStyleSetName(), "MaterialEditor.LiveUpdate")
+					);
+		LiveUpdateMenu.StyleNameOverride = "CalloutToolbar";
+		MaterialSection.AddEntry(LiveUpdateMenu);
 
-		{
-			FToolMenuSection& Section = ToolBar->AddSection("Search", TAttribute<FText>(), InsertAfterAssetSection);
-			Section.AddEntry(FToolMenuEntry::InitToolBarButton(FMaterialEditorCommands::Get().FindInMaterial));
-		}
+		
+		FToolMenuSection& GraphSection = ToolBar->AddSection("Graph", TAttribute<FText>(), InsertAfterAssetSection);
 
-		{
-			FToolMenuSection& Section = ToolBar->AddSection("Graph", TAttribute<FText>(), InsertAfterAssetSection);
-			Section.AddEntry(FToolMenuEntry::InitToolBarButton(FMaterialEditorCommands::Get().CameraHome));
-			Section.AddEntry(FToolMenuEntry::InitToolBarButton(FMaterialEditorCommands::Get().CleanUnusedExpressions));
-			Section.AddEntry(FToolMenuEntry::InitToolBarButton(FMaterialEditorCommands::Get().ShowHideConnectors));
-			Section.AddEntry(FToolMenuEntry::InitToolBarButton(FMaterialEditorCommands::Get().ToggleLivePreview));
-			Section.AddEntry(FToolMenuEntry::InitToolBarButton(FMaterialEditorCommands::Get().ToggleRealtimeExpressions));
-			Section.AddEntry(FToolMenuEntry::InitToolBarButton(FMaterialEditorCommands::Get().AlwaysRefreshAllPreviews));
-			Section.AddEntry(FToolMenuEntry::InitToolBarButton(
-				FMaterialEditorCommands::Get().ToggleHideUnrelatedNodes,
-				TAttribute<FText>(),
-				TAttribute<FText>(),
-				FSlateIcon(FEditorStyle::GetStyleSetName(), "GraphEditor.ToggleHideUnrelatedNodes")
-			));
-			Section.AddEntry(FToolMenuEntry::InitComboButton(
-				"HideUnrelatedNodesOptions",
-				FUIAction(),
-				FNewToolMenuDelegate::CreateLambda([](UToolMenu* InSubMenu)
+		FToolMenuEntry CleaningMenu = FToolMenuEntry::InitComboButton(
+			"CleanGraph",
+			FUIAction(),
+			FNewToolMenuDelegate::CreateLambda([](UToolMenu* InSubMenu)
+				{
+					FToolMenuSection& Section = InSubMenu->FindOrAddSection("General");
+					Section.AddEntry(FToolMenuEntry::InitMenuEntry(FMaterialEditorCommands::Get().CleanUnusedExpressions));
+					Section.AddEntry(FToolMenuEntry::InitMenuEntry(FMaterialEditorCommands::Get().ShowHideConnectors));
+				}),
+			LOCTEXT("GraphCleanup_Label", "Clean Graph"),
+			LOCTEXT("GraphCleanup_Tooltip", "Tools to help clean up graph nodes and connections"),
+			FSlateIcon(FAppStyle::Get().GetStyleSetName(), "GraphEditor.Clean")
+		);
+		CleaningMenu.StyleNameOverride = "CalloutToolbar";
+		GraphSection.AddEntry(CleaningMenu);
+		GraphSection.AddEntry(FToolMenuEntry::InitComboButton(
+			"NodePreview",
+			FUIAction(),
+			FNewToolMenuDelegate::CreateLambda([](UToolMenu* InSubMenu)
 				{
 					UMaterialEditorMenuContext* SubMenuContext = InSubMenu->FindContext<UMaterialEditorMenuContext>();
 					if (SubMenuContext && SubMenuContext->MaterialEditor.IsValid())
 					{
 						TSharedPtr<FMaterialEditor> MaterialEditor = StaticCastSharedPtr<FMaterialEditor>(SubMenuContext->MaterialEditor.Pin());
-						MaterialEditor->MakeHideUnrelatedNodesOptionsMenu(InSubMenu);
+						MaterialEditor->GeneratePreviewMenuContent(InSubMenu);
 					}
 				}),
-				LOCTEXT("HideUnrelatedNodesOptions", "Hide Unrelated Nodes Options"),
-				LOCTEXT("HideUnrelatedNodesOptionsMenu", "Hide Unrelated Nodes options menu"),
-				TAttribute<FSlateIcon>(),
-				true
-			));
-		}
-
-		{
-			FToolMenuSection& Section = ToolBar->AddSection("Stats", TAttribute<FText>(), InsertAfterAssetSection);
-			Section.AddEntry(FToolMenuEntry::InitToolBarButton(FMaterialEditorCommands::Get().ToggleMaterialStats));
-			Section.AddEntry(FToolMenuEntry::InitToolBarButton(FMaterialEditorCommands::Get().TogglePlatformStats));
-			Section.AddEntry(FToolMenuEntry::InitComboButton(
-				"NodePreview",
-				FUIAction(),
-				FNewToolMenuDelegate::CreateLambda([](UToolMenu* InSubMenu)
+			LOCTEXT("NodePreview_Label", "Preview State"),
+					LOCTEXT("NodePreviewToolTip", "Preview the graph state for a given feature level, material quality, or static switch value."),
+					FSlateIcon(FEditorStyle::GetStyleSetName(), "FullBlueprintEditor.SwitchToScriptingMode"),
+					false
+					));
+		GraphSection.AddEntry(FToolMenuEntry::InitToolBarButton(
+			FMaterialEditorCommands::Get().ToggleHideUnrelatedNodes,
+			TAttribute<FText>(),
+			TAttribute<FText>(),
+			FSlateIcon(FEditorStyle::GetStyleSetName(), "GraphEditor.ToggleHideUnrelatedNodes")
+		));
+		GraphSection.AddEntry(FToolMenuEntry::InitComboButton(
+			"HideUnrelatedNodesOptions",
+			FUIAction(),
+			FNewToolMenuDelegate::CreateLambda([](UToolMenu* InSubMenu)
 				{
 					UMaterialEditorMenuContext* SubMenuContext = InSubMenu->FindContext<UMaterialEditorMenuContext>();
 					if (SubMenuContext && SubMenuContext->MaterialEditor.IsValid())
 					{
 						TSharedPtr<FMaterialEditor> MaterialEditor = StaticCastSharedPtr<FMaterialEditor>(SubMenuContext->MaterialEditor.Pin());
-						MaterialEditor->GeneratePreviewMenuContent(InSubMenu);
+						MaterialEditor->MakeHideUnrelatedNodesOptionsMenu(InSubMenu);
 					}
 				}),
-				LOCTEXT("NodePreview_Label", "Preview Nodes"),
-				LOCTEXT("NodePreviewToolTip", "Preview the nodes for a given feature level and/or material quality."),
-				FSlateIcon(FEditorStyle::GetStyleSetName(), "FullBlueprintEditor.SwitchToScriptingMode"),
-				false
-			));
-		}
+			LOCTEXT("HideUnrelatedNodesOptions", "Hide Unrelated Nodes Options"),
+					LOCTEXT("HideUnrelatedNodesOptionsMenu", "Hide Unrelated Nodes options menu"),
+					TAttribute<FSlateIcon>(),
+					true
+					));
 
-		ToolBar->AddDynamicSection("Hierarchy", FNewToolMenuDelegate::CreateLambda([](UToolMenu* InMenu)
 		{
-			UMaterialEditorMenuContext* Context = InMenu->FindContext<UMaterialEditorMenuContext>();
-			TSharedPtr<FMaterialEditor> MaterialEditor = StaticCastSharedPtr<FMaterialEditor>(Context->MaterialEditor.Pin());
-			if (!MaterialEditor->MaterialFunction)
-			{
-				{
-					FToolMenuSection& Section = InMenu->AddSection("Hierarchy");
-					Section.AddEntry(FToolMenuEntry::InitComboButton(
-						"Hierarchy",
-						FToolUIActionChoice(),
-						FNewToolMenuDelegate::CreateLambda([](UToolMenu* InSubMenu)
-						{
-							UMaterialEditorMenuContext* SubMenuContext = InSubMenu->FindContext<UMaterialEditorMenuContext>();
-							if (SubMenuContext && SubMenuContext->MaterialEditor.IsValid())
-							{
-								SubMenuContext->MaterialEditor.Pin()->GenerateInheritanceMenu(InSubMenu);
-							}
-						}),
-						LOCTEXT("Hierarchy", "Hierarchy"),
-						FText::GetEmpty(),
-						FSlateIcon(FSlateIcon(FEditorStyle::GetStyleSetName(), "BTEditor.SwitchToBehaviorTreeMode")),
-						false
-					));
-				}
-			}
-		})).InsertPosition = InsertAfterAssetSection;
+			FToolMenuSection& Section = ToolBar->AddSection("Stats", TAttribute<FText>(), InsertAfterAssetSection);
+			Section.AddEntry(FToolMenuEntry::InitToolBarButton(FMaterialEditorCommands::Get().ToggleMaterialStats));
+			Section.AddEntry(FToolMenuEntry::InitToolBarButton(FMaterialEditorCommands::Get().TogglePlatformStats));
+			
+		}
+
+		
 	}
 };
 
@@ -2217,7 +2231,7 @@ void FMaterialEditor::LoadEditorSettings()
 {
 	EditorOptions = NewObject<UMaterialEditorOptions>();
 	
-	if (EditorOptions->bHideUnusedConnectors) {OnShowConnectors();}
+	if (EditorOptions->bHideUnusedConnectors) {OnHideConnectors();}
 	if (bLivePreview != EditorOptions->bLivePreviewUpdate)
 	{
 		ToggleLivePreview();
@@ -2260,7 +2274,7 @@ void FMaterialEditor::SaveEditorSettings()
 	{
 		EditorOptions->bShowGrid					= PreviewViewport->IsTogglePreviewGridChecked();
 		EditorOptions->bRealtimeMaterialViewport	= PreviewViewport->IsRealtime();
-		EditorOptions->bHideUnusedConnectors		= !IsOnShowConnectorsChecked();
+		EditorOptions->bHideUnusedConnectors		= !IsOnHideConnectorsChecked();
 		EditorOptions->bAlwaysRefreshAllPreviews	= IsOnAlwaysRefreshAllPreviews();
 		EditorOptions->bRealtimeExpressionViewport	= IsToggleRealTimeExpressionsChecked();
 		EditorOptions->bLivePreviewUpdate           = IsToggleLivePreviewChecked();
@@ -3044,9 +3058,9 @@ void FMaterialEditor::BindCommands()
 
 	ToolkitCommands->MapAction(
 		Commands.ShowHideConnectors,
-		FExecuteAction::CreateSP(this, &FMaterialEditor::OnShowConnectors),
+		FExecuteAction::CreateSP(this, &FMaterialEditor::OnHideConnectors),
 		FCanExecuteAction(),
-		FIsActionChecked::CreateSP(this, &FMaterialEditor::IsOnShowConnectorsChecked));
+		FIsActionChecked::CreateSP(this, &FMaterialEditor::IsOnHideConnectorsChecked));
 
 	ToolkitCommands->MapAction(
 		Commands.ToggleLivePreview,
@@ -3210,7 +3224,7 @@ void FMaterialEditor::OnCameraHome()
 	RecenterEditor();
 }
 
-void FMaterialEditor::OnShowConnectors()
+void FMaterialEditor::OnHideConnectors()
 {
 	bHideUnusedConnectors = !bHideUnusedConnectors;
 
@@ -3220,9 +3234,9 @@ void FMaterialEditor::OnShowConnectors()
 	}
 }
 
-bool FMaterialEditor::IsOnShowConnectorsChecked() const
+bool FMaterialEditor::IsOnHideConnectorsChecked() const
 {
-	return bHideUnusedConnectors == false;
+	return bHideUnusedConnectors == true;
 }
 
 void FMaterialEditor::ToggleLivePreview()
@@ -4582,7 +4596,7 @@ TSharedRef<SDockTab> FMaterialEditor::SpawnTab_Preview(const FSpawnTabArgs& Args
 TSharedRef<SDockTab> FMaterialEditor::SpawnTab_MaterialProperties(const FSpawnTabArgs& Args)
 {
 	TSharedPtr<SDockTab> DetailsTab = SNew(SDockTab)
-		.Icon( FEditorStyle::GetBrush("LevelEditor.Tabs.Details") )
+		.Icon( FAppStyle::Get().GetBrush("LevelEditor.Tabs.Details") )
 		.Label( LOCTEXT("MaterialDetailsTitle", "Details") )
 		[
 			MaterialDetailsView.ToSharedRef()
@@ -4690,7 +4704,7 @@ TSharedRef<SDockTab> FMaterialEditor::SpawnTab_CustomPrimitiveData(const FSpawnT
 TSharedRef<SDockTab> FMaterialEditor::SpawnTab_LayerProperties(const FSpawnTabArgs& Args)
 {
 	TSharedRef<SDockTab> SpawnedTab = SNew(SDockTab)
-		.Icon(FEditorStyle::GetBrush("MaterialInstanceEditor.Tabs.Properties"))
+		.Icon(FAppStyle::Get().GetBrush("LevelEditor.Tabs.Details"))
 		.Label(LOCTEXT("MaterialLayerPropertiesTitle", "Layer Parameter Preview"))
 		[
 			SNew(SBorder)
diff --git a/Engine/Source/Editor/MaterialEditor/Private/MaterialEditor.h b/Engine/Source/Editor/MaterialEditor/Private/MaterialEditor.h
--- a/Engine/Source/Editor/MaterialEditor/Private/MaterialEditor.h
+++ b/Engine/Source/Editor/MaterialEditor/Private/MaterialEditor.h
@@ -746,9 +746,9 @@ private:
 	bool OnApplyEnabled() const;
 	/** Command for the camera home button */
 	void OnCameraHome();
-	/** Command for the show unused connectors button */
-	void OnShowConnectors();
-	bool IsOnShowConnectorsChecked() const;
+	/** Command for the hide unused connectors button */
+	void OnHideConnectors();
+	bool IsOnHideConnectorsChecked() const;
 	/** Command for the Toggle Live Preview button */
 	void ToggleLivePreview();
 	bool IsToggleLivePreviewChecked() const;
@@ -769,16 +769,9 @@ private:
 
 	/** Make a drop down menu to control the opacity of unrelated nodes */
 	void MakeHideUnrelatedNodesOptionsMenu(class UToolMenu* Menu);
-	TOptional<float> HandleUnrelatedNodesOpacityBoxValue() const;
-	void HandleUnrelatedNodesOpacityBoxChanged(float NewOpacity);
 	void OnLockNodeStateCheckStateChanged(ECheckBoxState NewCheckedState);
 	void OnFocusWholeChainCheckStateChanged(ECheckBoxState NewCheckedState);
 
-	void ToggleReleaseStats();
-	bool IsToggleReleaseStatsChecked() const;
-	void ToggleBuiltinStats();
-	bool IsToggleBuiltinStatsChecked() const;
-
 	/** Command for using currently selected texture */
 	void OnUseCurrentTexture();
 	/** Command for converting nodes to objects */
diff --git a/Engine/Source/Editor/MaterialEditor/Private/MaterialEditorActions.cpp b/Engine/Source/Editor/MaterialEditor/Private/MaterialEditorActions.cpp
--- a/Engine/Source/Editor/MaterialEditor/Private/MaterialEditorActions.cpp
+++ b/Engine/Source/Editor/MaterialEditor/Private/MaterialEditorActions.cpp
@@ -13,7 +13,7 @@ void FMaterialEditorCommands::RegisterCommands()
 	UI_COMMAND( Apply, "Apply", "Apply changes to original material and its use in the world.", EUserInterfaceActionType::Button, FInputChord());
 	UI_COMMAND( Flatten, "Flatten", "Flatten the material to a texture for mobile devices.", EUserInterfaceActionType::Button, FInputChord() );
 
-	UI_COMMAND( ShowAllMaterialParameters, "Params", "Show or Hide all the materials parameters", EUserInterfaceActionType::ToggleButton, FInputChord() );
+	UI_COMMAND( ShowAllMaterialParameters, "Show Inactive", "Show the material instance parameters currently hidden behind static switches.", EUserInterfaceActionType::ToggleButton, FInputChord() );
 
 	UI_COMMAND( SetCylinderPreview, "Cylinder", "Sets the preview mesh to a cylinder primitive.", EUserInterfaceActionType::ToggleButton, FInputChord() );
 	UI_COMMAND( SetSpherePreview, "Sphere", "Sets the preview mesh to a sphere primitive.", EUserInterfaceActionType::ToggleButton, FInputChord() );
@@ -25,10 +25,10 @@ void FMaterialEditorCommands::RegisterCommands()
 
 	UI_COMMAND( CameraHome, "Home", "Goes home on the canvas.", EUserInterfaceActionType::Button, FInputChord() );
 	UI_COMMAND( CleanUnusedExpressions, "Clean Up", "Cleans up any unused Expressions.", EUserInterfaceActionType::Button, FInputChord() );
-	UI_COMMAND( ShowHideConnectors, "Connectors", "Show or Hide Unused Connectors", EUserInterfaceActionType::ToggleButton, FInputChord());
-	UI_COMMAND( ToggleLivePreview, "Live Preview", "Toggles real time update of the preview material.", EUserInterfaceActionType::ToggleButton, FInputChord());
-	UI_COMMAND( ToggleRealtimeExpressions, "Live Nodes", "Toggles real time update of the graph canvas.", EUserInterfaceActionType::ToggleButton, FInputChord() );
-	UI_COMMAND( AlwaysRefreshAllPreviews, "Live Update", "All nodes are previewed live.", EUserInterfaceActionType::ToggleButton, FInputChord() );
+	UI_COMMAND( ShowHideConnectors, "Hide Unused Connectors", "Hide unused connectors to clean up the graph view.", EUserInterfaceActionType::ToggleButton, FInputChord());
+	UI_COMMAND( ToggleLivePreview, "Preview Material", "Toggles real time update of the preview material.", EUserInterfaceActionType::ToggleButton, FInputChord());
+	UI_COMMAND( ToggleRealtimeExpressions, "Realtime Nodes", "Nodes impacted by time-based functions such as panners, etc. will update realtime.", EUserInterfaceActionType::ToggleButton, FInputChord() );
+	UI_COMMAND( AlwaysRefreshAllPreviews, "All Node Previews", "All node previews are updated upon any change to the graph.", EUserInterfaceActionType::ToggleButton, FInputChord() );
 	UI_COMMAND( ToggleMaterialStats, "Stats", "Toggles displaying of the material's stats.", EUserInterfaceActionType::ToggleButton, FInputChord());
 	UI_COMMAND( TogglePlatformStats, "Platform Stats", "Toggles the window that shows material stats and compilation errors for multiple platforms.", EUserInterfaceActionType::ToggleButton, FInputChord() );
 	UI_COMMAND( ToggleHideUnrelatedNodes, "Hide Unrelated", "Toggles hiding nodes which are unrelated to the selected nodes automatically.", EUserInterfaceActionType::ToggleButton, FInputChord() );
diff --git a/Engine/Source/Editor/MaterialEditor/Private/MaterialInstanceEditor.cpp b/Engine/Source/Editor/MaterialEditor/Private/MaterialInstanceEditor.cpp
--- a/Engine/Source/Editor/MaterialEditor/Private/MaterialInstanceEditor.cpp
+++ b/Engine/Source/Editor/MaterialEditor/Private/MaterialInstanceEditor.cpp
@@ -799,34 +799,30 @@ void FMaterialInstanceEditor::RegisterToolBar()
 
 		FToolMenuInsert InsertAfterAssetSection("Asset", EToolMenuInsertType::After);
 		{
-			FToolMenuSection& Section = ToolBar->AddSection("Command", TAttribute<FText>(), InsertAfterAssetSection);
+			FToolMenuSection& MaterialInstanceSection = ToolBar->AddSection("MaterialInstanceTools", TAttribute<FText>(), InsertAfterAssetSection);
 
-			Section.AddEntry(FToolMenuEntry::InitToolBarButton(FMaterialEditorCommands::Get().Apply));
-
-			Section.AddEntry(FToolMenuEntry::InitToolBarButton(FMaterialEditorCommands::Get().ShowAllMaterialParameters));
-			// TODO: support in material instance editor.
-			Section.AddEntry(FToolMenuEntry::InitToolBarButton(FMaterialEditorCommands::Get().TogglePlatformStats));
-		}
-
-		{
-			FToolMenuSection& Section = ToolBar->AddSection("Parent", TAttribute<FText>(), InsertAfterAssetSection);
-			Section.AddEntry(FToolMenuEntry::InitComboButton(
+			MaterialInstanceSection.AddEntry(FToolMenuEntry::InitToolBarButton(FMaterialEditorCommands::Get().Apply));
+			MaterialInstanceSection.AddEntry(FToolMenuEntry::InitToolBarButton(FMaterialEditorCommands::Get().ShowAllMaterialParameters));
+			MaterialInstanceSection.AddEntry(FToolMenuEntry::InitComboButton(
 				"Hierarchy",
 				FToolUIActionChoice(),
 				FNewToolMenuDelegate::CreateLambda([](UToolMenu* InSubMenu)
-				{
-					UMaterialEditorMenuContext* SubMenuContext = InSubMenu->FindContext<UMaterialEditorMenuContext>();
-					if (SubMenuContext && SubMenuContext->MaterialEditor.IsValid())
 					{
-						SubMenuContext->MaterialEditor.Pin()->GenerateInheritanceMenu(InSubMenu);
-					}
-				}),
+						UMaterialEditorMenuContext* SubMenuContext = InSubMenu->FindContext<UMaterialEditorMenuContext>();
+						if (SubMenuContext && SubMenuContext->MaterialEditor.IsValid())
+						{
+							SubMenuContext->MaterialEditor.Pin()->GenerateInheritanceMenu(InSubMenu);
+						}
+					}),
 				LOCTEXT("Hierarchy", "Hierarchy"),
-				FText::GetEmpty(),
-				FSlateIcon(FEditorStyle::GetStyleSetName(), "BTEditor.SwitchToBehaviorTreeMode"),
-				false
-			));
+						FText::GetEmpty(),
+						FSlateIcon(FAppStyle::Get().GetStyleSetName(), "MaterialEditor.Hierarchy"),
+						false
+						));
 		}
+
+		FToolMenuSection& UISection = ToolBar->AddSection("Stats", TAttribute<FText>(), InsertAfterAssetSection);
+		UISection.AddEntry(FToolMenuEntry::InitToolBarButton(FMaterialEditorCommands::Get().TogglePlatformStats));
 	}
 }
 
@@ -927,7 +923,7 @@ TSharedRef<SDockTab> FMaterialInstanceEditor::SpawnTab_Properties( const FSpawnT
 	check( Args.GetTabId().TabType == PropertiesTabId );
 
 	TSharedRef<SDockTab> SpawnedTab = SNew(SDockTab)
-		.Icon( FEditorStyle::GetBrush("MaterialInstanceEditor.Tabs.Properties") )
+		.Icon(FAppStyle::Get().GetBrush("LevelEditor.Tabs.Details"))
 		.Label(LOCTEXT("MaterialPropertiesTitle", "Details"))
 		[
 			SNew(SBorder)
@@ -948,7 +944,7 @@ TSharedRef<SDockTab> FMaterialInstanceEditor::SpawnTab_LayerProperties(const FSp
 	check(Args.GetTabId().TabType == LayerPropertiesTabId);
 
 	TSharedRef<SDockTab> SpawnedTab = SNew(SDockTab)
-		.Icon(FEditorStyle::GetBrush("MaterialInstanceEditor.Tabs.Properties"))
+		.Icon(FAppStyle::Get().GetBrush("LevelEditor.Tabs.Details"))
 		.Label(LOCTEXT("MaterialLayerPropertiesTitle", "Layer Parameters"))
 		[
 			SNew(SBorder)
diff --git a/Engine/Source/Editor/MaterialEditor/Private/MaterialStats.cpp b/Engine/Source/Editor/MaterialEditor/Private/MaterialStats.cpp
--- a/Engine/Source/Editor/MaterialEditor/Private/MaterialStats.cpp
+++ b/Engine/Source/Editor/MaterialEditor/Private/MaterialStats.cpp
@@ -912,7 +912,7 @@ TSharedRef<SDockTab> FMaterialStats::SpawnTab_Stats(const FSpawnTabArgs& Args)
 	FString TabName = FString::Printf(TEXT(""), *GetMaterialName().ToString());
 
 	TSharedRef<SDockTab> SpawnedTab = SNew(SDockTab)
-		.Icon(FEditorStyle::GetBrush("Kismet.Tabs.CompilerResults"))
+		.Icon(FAppStyle::Get().GetBrush("MaterialEditor.TogglePlatformStats.Tab"))
 		.Label(LOCTEXT("Platform Stats", "Platform Stats"))
 		.OnTabClosed_Lambda([&bShowStats = bShowStats](TSharedRef<SDockTab>) { bShowStats = false; })
 		[
@@ -936,7 +936,7 @@ TSharedRef<SDockTab> FMaterialStats::SpawnTab_OldStats(const FSpawnTabArgs& Args
 	check(Args.GetTabId() == OldStatsTabId);
 
 	TSharedRef<SDockTab> SpawnedTab = SNew(SDockTab)
-		.Icon(FEditorStyle::GetBrush("Kismet.Tabs.CompilerResults"))
+		.Icon(FAppStyle::Get().GetBrush("MaterialEditor.ToggleMaterialStats.Tab"))
 		.Label(LOCTEXT("Stats", "Stats"))
 		.OnTabClosed_Lambda([&bShowStats = bShowOldStats](TSharedRef<SDockTab>) { bShowStats = false; })
 		[
@@ -963,7 +963,7 @@ void FMaterialStats::BuildStatsTab()
 	TabManager->RegisterTabSpawner(StatsTabId, FOnSpawnTab::CreateSP(this, &FMaterialStats::SpawnTab_Stats))
 		.SetDisplayName(LOCTEXT("StatsTab", "Platform Stats"))
 		.SetGroup(ParentCategoryRef)
-		.SetIcon(FSlateIcon(FEditorStyle::GetStyleSetName(), "LevelEditor.Tabs.StatsViewer"));
+		.SetIcon(FSlateIcon(FAppStyle::Get().GetStyleSetName(),"MaterialEditor.TogglePlatformStats.Tab"));
 }
 
 void FMaterialStats::BuildOldStatsTab()
@@ -974,7 +974,7 @@ void FMaterialStats::BuildOldStatsTab()
 	TabManager->RegisterTabSpawner(OldStatsTabId, FOnSpawnTab::CreateSP(this, &FMaterialStats::SpawnTab_OldStats))
 		.SetDisplayName(LOCTEXT("OldStatsTab", "Stats"))
 		.SetGroup(ParentCategoryRef)
-		.SetIcon(FSlateIcon(FEditorStyle::GetStyleSetName(), "LevelEditor.Tabs.StatsViewer"));
+		.SetIcon(FSlateIcon(FAppStyle::Get().GetStyleSetName(), "MaterialEditor.ToggleMaterialStats.Tab"));
 }
 
 void FMaterialStats::RegisterTabs()
diff --git a/Engine/Source/Editor/UnrealEd/Private/MaterialGraphNode.cpp b/Engine/Source/Editor/UnrealEd/Private/MaterialGraphNode.cpp
--- a/Engine/Source/Editor/UnrealEd/Private/MaterialGraphNode.cpp
+++ b/Engine/Source/Editor/UnrealEd/Private/MaterialGraphNode.cpp
@@ -375,25 +375,25 @@ void UMaterialGraphNode::GetNodeContextMenuActions(UToolMenu* Menu, UGraphNodeCo
 
 	if (Context->Node)
 	{
+		FToolMenuInsert MenuPosition = FToolMenuInsert(NAME_None, EToolMenuInsertType::First);
+		FToolMenuSection& NodeSection = Menu->AddSection("MaterialSchemaNodeActions", LOCTEXT("NodeActionsMenuHeader", "Node Actions"), MenuPosition);
 		if (MaterialExpression)
 		{
 			if (MaterialExpression->IsA(UMaterialExpressionTextureBase::StaticClass()))
 			{
 				{
-					FToolMenuSection& Section = Menu->AddSection("MaterialGraphNode");
-					Section.AddMenuEntry(FMaterialEditorCommands::Get().UseCurrentTexture);
+					NodeSection.AddMenuEntry(FMaterialEditorCommands::Get().UseCurrentTexture);
 				}
 
 				// Add a 'Convert To Texture' option for convertible types
 				{
-					FToolMenuSection& Section = Menu->AddSection("MaterialEditorMenu0");
 					if ( MaterialExpression->IsA(UMaterialExpressionTextureSample::StaticClass()) && !MaterialExpression->HasAParameterName())
 					{
-						Section.AddMenuEntry(FMaterialEditorCommands::Get().ConvertToTextureObjects);
+						NodeSection.AddMenuEntry(FMaterialEditorCommands::Get().ConvertToTextureObjects);
 					}
 					else if ( MaterialExpression->IsA(UMaterialExpressionTextureObject::StaticClass()))
 					{
-						Section.AddMenuEntry(FMaterialEditorCommands::Get().ConvertToTextureSamples);
+						NodeSection.AddMenuEntry(FMaterialEditorCommands::Get().ConvertToTextureSamples);
 					}
 				}
 			}
@@ -401,23 +401,21 @@ void UMaterialGraphNode::GetNodeContextMenuActions(UToolMenu* Menu, UGraphNodeCo
 			// Add a 'Convert to Local Variables' option to reroute nodes
 			if (MaterialExpression->IsA(UMaterialExpressionReroute::StaticClass()))
 			{
-				FToolMenuSection& Section = Menu->AddSection("MaterialEditorMenu1");
-				Section.AddMenuEntry(FMaterialEditorCommands::Get().ConvertRerouteToNamedReroute);
+				NodeSection.AddMenuEntry(FMaterialEditorCommands::Get().ConvertRerouteToNamedReroute);
 			}
 
 			// Add local variables selection & conversion to reroute nodes
 			if (MaterialExpression->IsA(UMaterialExpressionNamedRerouteBase::StaticClass()))
 			{
-				FToolMenuSection& Section = Menu->AddSection("MaterialEditorMenu1");
 				if (MaterialExpression->IsA(UMaterialExpressionNamedRerouteDeclaration::StaticClass()))
 				{
-					Section.AddMenuEntry(FMaterialEditorCommands::Get().SelectNamedRerouteUsages);
+					NodeSection.AddMenuEntry(FMaterialEditorCommands::Get().SelectNamedRerouteUsages);
 				}
 				if (MaterialExpression->IsA(UMaterialExpressionNamedRerouteUsage::StaticClass()))
 				{
-					Section.AddMenuEntry(FMaterialEditorCommands::Get().SelectNamedRerouteDeclaration);
+					NodeSection.AddMenuEntry(FMaterialEditorCommands::Get().SelectNamedRerouteDeclaration);
 				}
-				Section.AddMenuEntry(FMaterialEditorCommands::Get().ConvertNamedRerouteToReroute);
+				NodeSection.AddMenuEntry(FMaterialEditorCommands::Get().ConvertNamedRerouteToReroute);
 			}
 
 			// Add a 'Convert To Parameter' option for convertible types
@@ -431,8 +429,7 @@ void UMaterialGraphNode::GetNodeContextMenuActions(UToolMenu* Menu, UGraphNodeCo
 				|| MaterialExpression->IsA(UMaterialExpressionComponentMask::StaticClass()))
 			{
 				{
-					FToolMenuSection& Section = Menu->AddSection("MaterialEditorMenu1");
-					Section.AddMenuEntry(FMaterialEditorCommands::Get().ConvertObjects);
+					NodeSection.AddMenuEntry(FMaterialEditorCommands::Get().ConvertObjects);
 				}
 			}
 
@@ -446,13 +443,11 @@ void UMaterialGraphNode::GetNodeContextMenuActions(UToolMenu* Menu, UGraphNodeCo
 				|| MaterialExpression->IsA(UMaterialExpressionTextureObjectParameter::StaticClass()))
 			{
 				{
-					FToolMenuSection& Section = Menu->AddSection("MaterialEditorMenu1");
-					Section.AddMenuEntry(FMaterialEditorCommands::Get().ConvertToConstant);
+					NodeSection.AddMenuEntry(FMaterialEditorCommands::Get().ConvertToConstant);
 				}
 			}
 
 			{
-				FToolMenuSection& Section = Menu->AddSection("MaterialEditorMenu2");
 				// Don't show preview option for bools
 				if (!MaterialExpression->IsA(UMaterialExpressionStaticBool::StaticClass())
 					&& !MaterialExpression->IsA(UMaterialExpressionStaticBoolParameter::StaticClass()))
@@ -461,36 +456,35 @@ void UMaterialGraphNode::GetNodeContextMenuActions(UToolMenu* Menu, UGraphNodeCo
 					if (bIsPreviewExpression)
 					{
 						// If we are already previewing the selected node, the menu option should tell the user that this will stop previewing
-						Section.AddMenuEntry(FMaterialEditorCommands::Get().StopPreviewNode);
+						NodeSection.AddMenuEntry(FMaterialEditorCommands::Get().StopPreviewNode);
 					}
 					else
 					{
 						// The menu option should tell the user this node will be previewed.
-						Section.AddMenuEntry(FMaterialEditorCommands::Get().StartPreviewNode);
+						NodeSection.AddMenuEntry(FMaterialEditorCommands::Get().StartPreviewNode);
 					}
 				}
 
 				if (MaterialExpression->bRealtimePreview)
 				{
-					Section.AddMenuEntry(FMaterialEditorCommands::Get().DisableRealtimePreviewNode);
+					NodeSection.AddMenuEntry(FMaterialEditorCommands::Get().DisableRealtimePreviewNode);
 				}
 				else
 				{
-					Section.AddMenuEntry(FMaterialEditorCommands::Get().EnableRealtimePreviewNode);
+					NodeSection.AddMenuEntry(FMaterialEditorCommands::Get().EnableRealtimePreviewNode);
 				}
 			}
 		}
 
 		// Break all links
 		{
-			FToolMenuSection& Section = Menu->AddSection("BreakAllLinks");
-			Section.AddMenuEntry(FGraphEditorCommands::Get().BreakNodeLinks);
+			NodeSection.AddMenuEntry(FGraphEditorCommands::Get().BreakNodeLinks);
 		}
 
 		// Separate the above frequently used options from the below less frequently used common options
 		
 		{
-			FToolMenuSection& Section = Menu->AddSection("MaterialEditorMenu3");
+			FToolMenuSection& Section = Menu->AddSection("MaterialEditorMenu3", LOCTEXT("GeneralNodeActionsMenuHeader", "General"));
 			Section.AddMenuEntry( FGenericCommands::Get().Delete );
 			Section.AddMenuEntry( FGenericCommands::Get().Cut );
 			Section.AddMenuEntry( FGenericCommands::Get().Copy );
