Fix for Hide Unrelated Connectors defaulting to on

diff --git a/Engine/Config/BaseEditorPerProjectUserSettings.ini b/Engine/Config/BaseEditorPerProjectUserSettings.ini
--- a/Engine/Config/BaseEditorPerProjectUserSettings.ini
+++ b/Engine/Config/BaseEditorPerProjectUserSettings.ini
@@ -471,7 +471,7 @@ ExpandDistributions=false
 [/Script/UnrealEd.MaterialEditorOptions]
 bShowGrid=True
 bShowBackground=False
-bHideUnusedConnectors=False
+bShouldHideUnusedConnectors=False
 bRealtimeMaterialViewport=True
 bRealtimeExpressionViewport=False
 bAlwaysRefreshAllPreviews=False
diff --git a/Engine/Source/Editor/MaterialEditor/Private/MaterialEditor.cpp b/Engine/Source/Editor/MaterialEditor/Private/MaterialEditor.cpp
--- a/Engine/Source/Editor/MaterialEditor/Private/MaterialEditor.cpp
+++ b/Engine/Source/Editor/MaterialEditor/Private/MaterialEditor.cpp
@@ -2231,7 +2231,7 @@ void FMaterialEditor::LoadEditorSettings()
 {
 	EditorOptions = NewObject<UMaterialEditorOptions>();
 	
-	if (EditorOptions->bHideUnusedConnectors) {OnHideConnectors();}
+	if (EditorOptions->bShouldHideUnusedConnectors) {OnHideConnectors();}
 	if (bLivePreview != EditorOptions->bLivePreviewUpdate)
 	{
 		ToggleLivePreview();
@@ -2274,7 +2274,7 @@ void FMaterialEditor::SaveEditorSettings()
 	{
 		EditorOptions->bShowGrid					= PreviewViewport->IsTogglePreviewGridChecked();
 		EditorOptions->bRealtimeMaterialViewport	= PreviewViewport->IsRealtime();
-		EditorOptions->bHideUnusedConnectors		= !IsOnHideConnectorsChecked();
+		EditorOptions->bShouldHideUnusedConnectors	= !IsOnHideConnectorsChecked();
 		EditorOptions->bAlwaysRefreshAllPreviews	= IsOnAlwaysRefreshAllPreviews();
 		EditorOptions->bRealtimeExpressionViewport	= IsToggleRealTimeExpressionsChecked();
 		EditorOptions->bLivePreviewUpdate           = IsToggleLivePreviewChecked();
diff --git a/Engine/Source/Editor/UnrealEd/Classes/Preferences/MaterialEditorOptions.h b/Engine/Source/Editor/UnrealEd/Classes/Preferences/MaterialEditorOptions.h
--- a/Engine/Source/Editor/UnrealEd/Classes/Preferences/MaterialEditorOptions.h
+++ b/Engine/Source/Editor/UnrealEd/Classes/Preferences/MaterialEditorOptions.h
@@ -24,7 +24,7 @@ class UNREALED_API UMaterialEditorOptions : public UObject
 
 	/** If true, don't render connectors that are not connected to anything. */
 	UPROPERTY(EditAnywhere, config, Category=Options)
-	uint32 bHideUnusedConnectors:1;
+	uint32 bShouldHideUnusedConnectors:1;
 
 	/** If true, the 3D material preview viewport updates in realtime. */
 	UPROPERTY(EditAnywhere, config, Category=Options)
