Adding shader count statistic in the material editor.  This reports the number of shaders that are in the material's shader maps by calling GetShaderNum().

diff --git a/Engine/Source/Editor/MaterialEditor/Private/MaterialEditor.cpp b/Engine/Source/Editor/MaterialEditor/Private/MaterialEditor.cpp
--- a/Engine/Source/Editor/MaterialEditor/Private/MaterialEditor.cpp
+++ b/Engine/Source/Editor/MaterialEditor/Private/MaterialEditor.cpp
@@ -2776,6 +2776,12 @@ void FMaterialEditor::UpdateMaterialinfoList_Old()
 					TSharedRef<FTokenizedMessage> Line = FTokenizedMessage::Create(EMessageSeverity::Info);
 					Line->AddToken(FTextToken::Create(FText::FromString(StrataMaterialDescription)));
 					Messages.Add(Line);
+
+					// Add shader count
+					FString ShaderCountString = FString::Printf(TEXT("Shader Count: %u"), ShaderMap->GetShaderNum());
+					TSharedRef<FTokenizedMessage> ShaderCountLine = FTokenizedMessage::Create(EMessageSeverity::Info);
+					ShaderCountLine->AddToken(FTextToken::Create(FText::FromString(ShaderCountString)));
+					Messages.Add(ShaderCountLine);
 				}
 			}
 
diff --git a/Engine/Source/Editor/MaterialEditor/Private/MaterialStats.h b/Engine/Source/Editor/MaterialEditor/Private/MaterialStats.h
--- a/Engine/Source/Editor/MaterialEditor/Private/MaterialStats.h
+++ b/Engine/Source/Editor/MaterialEditor/Private/MaterialStats.h
@@ -20,6 +20,7 @@ struct FShaderStatsInfo
 	FContent InterpolatorsCount;
 	FContent TextureSampleCount;
 	FContent VirtualTextureLookupCount;
+	FContent ShaderCount;
 	FString StrShaderErrors;
 
 	void Reset()
@@ -37,6 +38,9 @@ struct FShaderStatsInfo
 
 		VirtualTextureLookupCount.StrDescription = TEXT("Compiling...");
 		VirtualTextureLookupCount.StrDescriptionLong = TEXT("Compiling...");
+
+		ShaderCount.StrDescription = TEXT("Compiling...");
+		ShaderCount.StrDescriptionLong = TEXT("Compiling...");
 		
 		StrShaderErrors.Empty();
 	}
@@ -57,6 +61,9 @@ struct FShaderStatsInfo
 		VirtualTextureLookupCount.StrDescription.Empty();
 		VirtualTextureLookupCount.StrDescriptionLong.Empty();
 
+		ShaderCount.StrDescription.Empty();
+		ShaderCount.StrDescriptionLong.Empty();
+
 		StrShaderErrors.Empty();
 	}
 
diff --git a/Engine/Source/Editor/MaterialEditor/Private/MaterialStatsCommon.cpp b/Engine/Source/Editor/MaterialEditor/Private/MaterialStatsCommon.cpp
--- a/Engine/Source/Editor/MaterialEditor/Private/MaterialStatsCommon.cpp
+++ b/Engine/Source/Editor/MaterialEditor/Private/MaterialStatsCommon.cpp
@@ -728,6 +728,12 @@ void FMaterialStatsUtils::ExtractMatertialStatsInfo(FShaderStatsInfo& OutInfo, c
 		OutInfo.InterpolatorsCount.StrDescription = FString::Printf(TEXT("%u/%u"), TotalScalars, MaxScalars);
 		OutInfo.InterpolatorsCount.StrDescriptionLong = FString::Printf(TEXT("User interpolators: %u/%u Scalars (%u/4 Vectors) (TexCoords: %i, Custom: %i)"),
 			TotalScalars, MaxScalars, MaxScalars / 4, UVScalarsUsed, CustomInterpolatorScalarsUsed);
+
+		// extract total shader count info
+		const uint32 ShaderCount = MaterialResource->GetGameThreadShaderMap()->GetShaderNum();
+
+		OutInfo.ShaderCount.StrDescription = FString::Printf(TEXT("%u"), ShaderCount);
+		OutInfo.ShaderCount.StrDescriptionLong = FString::Printf(TEXT("Total Shaders: %u"), ShaderCount);
 	}
 }
 
diff --git a/Engine/Source/Editor/MaterialEditor/Private/MaterialStatsGrid.cpp b/Engine/Source/Editor/MaterialEditor/Private/MaterialStatsGrid.cpp
--- a/Engine/Source/Editor/MaterialEditor/Private/MaterialStatsGrid.cpp
+++ b/Engine/Source/Editor/MaterialEditor/Private/MaterialStatsGrid.cpp
@@ -342,6 +342,35 @@ void FStatsGridRow_NumVirtualTextureLookups::AddPlatform(TSharedPtr<FMaterialSta
 /*end FStatsGridRow_NumVirtualTextureLookups functions*/
 /*==============================================================================================================*/
 
+/*==============================================================================================================*/
+/* FStatsGridRow_NumVirtualTextureLookups functions*/
+
+void FStatsGridRow_NumShaders::CreateRow(TSharedPtr<FMaterialStats> StatsManager)
+{
+	// static string in the descriptor column
+	TSharedPtr<FGridCell> HeaderCell = MakeShareable(new FGridCell_StaticString(TEXT("Shader Count"), TEXT("Number of shaders compiled for this platform and quality level.")));
+	HeaderCell->SetColor(FStyleColors::Foreground);
+	HeaderCell->SetContentBold(true);
+	AddCell(FMaterialStatsGrid::DescriptorColumnName, HeaderCell);
+
+	AddCell(FMaterialStatsGrid::ShaderColumnName, MakeShareable(new FGridCell_Empty()));
+
+	FillPlatformCellsHelper(StatsManager);
+}
+
+void FStatsGridRow_NumShaders::AddPlatform(TSharedPtr<FMaterialStats> StatsManager, const TSharedPtr<FShaderPlatformSettings> Platform, const EMaterialQualityLevel::Type QualityLevel)
+{
+	// cell that will enumerate the number of shaders
+	const FString CellContent = FMaterialStatsUtils::MaterialQualityToShortString(QualityLevel);
+	TSharedPtr<FGridCell_ShaderValue> Cell = MakeShareable(new FGridCell_ShaderValue(StatsManager, EShaderInfoType::ShaderCount, ERepresentativeShader::Num, QualityLevel, Platform->GetPlatformShaderType()));
+
+	const FName ColumnName = FMaterialStatsGrid::MakePlatformColumnName(Platform, QualityLevel);
+	AddCell(ColumnName, Cell);
+}
+
+/*end FStatsGridRow_NumTotalShaders functions*/
+/*==============================================================================================================*/
+
 /***********************************************************************************************************************/
 /*FShaderStatsGrid functions*/
 
@@ -533,7 +562,8 @@ void FMaterialStatsGrid::BuildRowIds()
 		BuildKeyAndInsert(ERowType::Samplers);
 		BuildKeyAndInsert(ERowType::TextureSamples);
 		BuildKeyAndInsert(ERowType::VirtualTextureLookups);
-		BuildKeyAndInsert(ERowType::Interpolators);	
+		BuildKeyAndInsert(ERowType::Interpolators);
+		BuildKeyAndInsert(ERowType::Shaders);
 	}
 }
 
@@ -590,6 +620,11 @@ void FMaterialStatsGrid::BuildStaticRows()
 		Row->CreateRow(StatsManager);
 		StaticRows.Add(ERowType::Interpolators, Row);
 	}
+	{
+		TSharedPtr<FStatsGridRow> Row = MakeShareable(new FStatsGridRow_NumShaders());
+		Row->CreateRow(StatsManager);
+		StaticRows.Add(ERowType::Shaders, Row);
+	}
 }
 
 void FMaterialStatsGrid::BuildShaderRows()
@@ -898,6 +933,10 @@ FString FGridCell_ShaderValue::InternalGetContent(bool bLongContent)
 		case EShaderInfoType::VirtualTextureLookupCount:
 			return bLongContent ? PlatformData.ShaderStatsInfo.VirtualTextureLookupCount.StrDescriptionLong : PlatformData.ShaderStatsInfo.VirtualTextureLookupCount.StrDescription;
 		break;
+
+		case EShaderInfoType::ShaderCount:
+			return bLongContent ? PlatformData.ShaderStatsInfo.ShaderCount.StrDescriptionLong : PlatformData.ShaderStatsInfo.ShaderCount.StrDescription;
+		break;
 	}
 
 	return TEXT("");
diff --git a/Engine/Source/Editor/MaterialEditor/Private/MaterialStatsGrid.h b/Engine/Source/Editor/MaterialEditor/Private/MaterialStatsGrid.h
--- a/Engine/Source/Editor/MaterialEditor/Private/MaterialStatsGrid.h
+++ b/Engine/Source/Editor/MaterialEditor/Private/MaterialStatsGrid.h
@@ -74,6 +74,7 @@ enum class EShaderInfoType
 	InterpolatorsCount,
 	TextureSampleCount,
 	VirtualTextureLookupCount,
+	ShaderCount,
 };
 
 /** this type of cell will query certain type of informations from the material */
@@ -215,6 +216,15 @@ public:
 	void AddPlatform(TSharedPtr<FMaterialStats> StatsManager, const TSharedPtr<FShaderPlatformSettings> Platform, const EMaterialQualityLevel::Type QualityLevel) override;
 };
 
+/** this row will display the total number of shaders present in the material for a specified platform */
+class FStatsGridRow_NumShaders : public FStatsGridRow
+{
+public:
+	void CreateRow(TSharedPtr<FMaterialStats> StatsManager) override;
+
+	void AddPlatform(TSharedPtr<FMaterialStats> StatsManager, const TSharedPtr<FShaderPlatformSettings> Platform, const EMaterialQualityLevel::Type QualityLevel) override;
+};
+
 /** class that models the logical material stats grid */
 class FMaterialStatsGrid
 {
@@ -239,6 +249,7 @@ class FMaterialStatsGrid
 		Interpolators,
 		TextureSamples,
 		VirtualTextureLookups,
+		Shaders,
 
 		VertexShader,
 		FragmentShader,
