Adding Copy/Paste functionality to the group header in the material instance editor

diff --git a/Engine/Source/Editor/MaterialEditor/Private/MaterialEditorInstanceDetailCustomization.cpp b/Engine/Source/Editor/MaterialEditor/Private/MaterialEditorInstanceDetailCustomization.cpp
--- a/Engine/Source/Editor/MaterialEditor/Private/MaterialEditorInstanceDetailCustomization.cpp
+++ b/Engine/Source/Editor/MaterialEditor/Private/MaterialEditorInstanceDetailCustomization.cpp
@@ -48,6 +48,7 @@
 #include "Curves/CurveLinearColor.h"
 #include "IPropertyUtilities.h"
 #include "Engine/Texture.h"
+#include "HAL/PlatformApplicationMisc.h"
 
 #define LOCTEXT_NAMESPACE "MaterialInstanceEditor"
 
@@ -307,6 +308,22 @@ void FMaterialInstanceParameterDetails::CreateGroupsWidget(TSharedRef<IPropertyH
 			if (bCreateGroup)
 			{
 				IDetailGroup& DetailGroup = GroupsCategory.AddGroup(ParameterGroup.GroupName, FText::FromName(ParameterGroup.GroupName), false, true);
+
+				FUIAction CopyAction(
+					FExecuteAction::CreateSP(this, &FMaterialInstanceParameterDetails::OnCopyParameterValues, GroupIdx),
+					FCanExecuteAction::CreateSP(this, &FMaterialInstanceParameterDetails::CanCopyParameterValues, GroupIdx));
+				FUIAction PasteAction(
+					FExecuteAction::CreateSP(this, &FMaterialInstanceParameterDetails::OnPasteParameterValues, GroupIdx),
+					FCanExecuteAction::CreateSP(this, &FMaterialInstanceParameterDetails::CanPasteParameterValues, GroupIdx));
+				DetailGroup.HeaderRow()
+					.CopyAction(CopyAction)
+					.PasteAction(PasteAction)
+					.NameContent()
+					[
+						SNew(STextBlock)
+						.Text(FText::FromName(DetailGroup.GetGroupName()))
+					];
+
 				CreateSingleGroupWidget(ParameterGroup, ParameterGroupsProperty->GetChildHandle(GroupIdx), DetailGroup);
 			}
 		}
@@ -983,6 +1000,125 @@ EVisibility FMaterialInstanceParameterDetails::ShouldShowSubsurfaceProfile() con
 	return UseSubsurfaceProfile(ShadingModels) ? EVisibility::Visible : EVisibility::Collapsed;
 }
 
+void FMaterialInstanceParameterDetails::OnCopyParameterValues(int32 ParameterGroupIndex)
+{
+	if (!MaterialEditorInstance || !MaterialEditorInstance->ParameterGroups.IsValidIndex(ParameterGroupIndex))
+	{
+		return;
+	}
+	FEditorParameterGroup& ParameterGroup = MaterialEditorInstance->ParameterGroups[ParameterGroupIndex];
+
+	TStringBuilder<4096> CombinedValue;
+
+	const int32 NumParams = ParameterGroup.Parameters.Num();
+	for (int32 ParamIdx = 0; ParamIdx < NumParams; ++ParamIdx)
+	{
+		UDEditorParameterValue* Parameter = ParameterGroup.Parameters[ParamIdx];
+		const FName ParamName = Parameter->ParameterInfo.Name;
+
+		const TCHAR* Prefix = (ParamIdx == 0) ? TEXT("") : TEXT(",");
+
+		// Include the value in the result entry only if the parameter is overridden.
+		const bool bOverride = FMaterialPropertyHelpers::IsOverriddenExpression(Parameter);
+		if (bOverride)
+		{
+			FProperty* ParameterValueProperty = Parameter->GetClass()->FindPropertyByName("ParameterValue");
+			if (ParameterValueProperty != nullptr)
+			{
+				FString ParameterValueString;
+				if (ParameterValueProperty->ExportText_InContainer(0, ParameterValueString, Parameter, Parameter, Parameter, PPF_Copy))
+				{
+					CombinedValue.Appendf(TEXT("%s%s.Override=True,%s.Value=\"%s\""),
+						Prefix,
+						*(ParamName.ToString()),
+						*(ParamName.ToString()),
+						*(ParameterValueString.ReplaceCharWithEscapedChar()));
+				}
+			}
+		}
+		else
+		{
+			CombinedValue.Appendf(TEXT("%s%s.Override=False"), Prefix, *(ParamName.ToString()));
+		}
+	}
+
+	if (CombinedValue.Len())
+	{
+		// Copy.
+		FPlatformApplicationMisc::ClipboardCopy(*CombinedValue);
+	}
+}
+
+bool FMaterialInstanceParameterDetails::CanCopyParameterValues(int32 ParameterGroupIndex)
+{
+	return MaterialEditorInstance && MaterialEditorInstance->ParameterGroups.IsValidIndex(ParameterGroupIndex)
+		&& (MaterialEditorInstance->ParameterGroups[ParameterGroupIndex].Parameters.Num() > 0);
+}
+
+void FMaterialInstanceParameterDetails::OnPasteParameterValues(int32 ParameterGroupIndex)
+{
+	if (!MaterialEditorInstance || !MaterialEditorInstance->ParameterGroups.IsValidIndex(ParameterGroupIndex))
+	{
+		return;
+	}
+
+	FString ClipboardContent;
+	FPlatformApplicationMisc::ClipboardPaste(ClipboardContent);
+	if (!ClipboardContent.IsEmpty())
+	{
+		FScopedTransaction Transaction(LOCTEXT("PasteMaterialInstanceParameters", "Paste Material Instance Parameters"));
+
+		MaterialEditorInstance->Modify();
+
+		for (int32 ParamIdx = 0; ParamIdx < MaterialEditorInstance->ParameterGroups[ParameterGroupIndex].Parameters.Num(); ++ParamIdx)
+		{
+			UDEditorParameterValue* Parameter = MaterialEditorInstance->ParameterGroups[ParameterGroupIndex].Parameters[ParamIdx];
+			Parameter->Modify();
+
+			const FName ParamName = Parameter->ParameterInfo.Name;
+
+			const FString OverrideKey = FString::Printf(TEXT("%s.Override="), *(ParamName.ToString()));
+			bool bParsedOverride = false;
+			if (FParse::Bool(*ClipboardContent, *OverrideKey, bParsedOverride))
+			{
+				Parameter->bOverride = bParsedOverride;
+				if (bParsedOverride)
+				{
+					// Paste value.
+					const FString ValueKey = FString::Printf(TEXT("%s.Value="), *(ParamName.ToString()));
+					FString ParsedValueString;
+					if (FParse::Value(*ClipboardContent, *ValueKey, ParsedValueString))
+					{
+						ParsedValueString = ParsedValueString.ReplaceEscapedCharWithChar();
+						FProperty* ParameterValueProperty = Parameter->GetClass()->FindPropertyByName("ParameterValue");
+						if (ParameterValueProperty != nullptr)
+						{
+							void* ParameterValuePtr = ParameterValueProperty->ContainerPtrToValuePtr<void>(Parameter);
+							ParameterValueProperty->ImportText(*ParsedValueString, ParameterValuePtr, PPF_Copy, Parameter);
+						}
+					}
+				}
+			}
+		}
+
+		MaterialEditorInstance->PostEditChange();
+		FEditorSupportDelegates::RedrawAllViewports.Broadcast();
+	}
+}
+
+bool FMaterialInstanceParameterDetails::CanPasteParameterValues(int32 ParameterGroupIndex)
+{
+	// First check the same criteria as copying.
+	if (!CanCopyParameterValues(ParameterGroupIndex))
+	{
+		return false;
+	}
+
+	// Now see if there's anything to paste from the clipboard.
+	FString ClipboardContent;
+	FPlatformApplicationMisc::ClipboardPaste(ClipboardContent);
+	return !ClipboardContent.IsEmpty();
+}
 
 void FMaterialInstanceParameterDetails::CreateLightmassOverrideWidgets(IDetailLayoutBuilder& DetailLayout)
 {
diff --git a/Engine/Source/Editor/MaterialEditor/Private/MaterialEditorInstanceDetailCustomization.h b/Engine/Source/Editor/MaterialEditor/Private/MaterialEditorInstanceDetailCustomization.h
--- a/Engine/Source/Editor/MaterialEditor/Private/MaterialEditorInstanceDetailCustomization.h
+++ b/Engine/Source/Editor/MaterialEditor/Private/MaterialEditorInstanceDetailCustomization.h
@@ -74,6 +74,29 @@ private:
 	/** Returns true if the refraction options should be displayed */
 	EVisibility ShouldShowSubsurfaceProfile() const;
 
+	//Functions supporting copy/paste of entire parameter groups.
+
+	/**
+	 * Copy all parameter values in a parameter group to the clipboard
+	 * in the format "Param1=\"Value1\"\nParams2=\"Value2\"..."
+	 */
+	void OnCopyParameterValues(int32 ParameterGroupIndex);
+
+	/** Whether it is possible to copy parameter values for the given parent group index */
+	bool CanCopyParameterValues(int32 ParameterGroupIndex);
+
+	/**
+	 * Paste parameter values from the clipboard, assumed to be
+	 * in the format copied by CopyParameterValues.
+	 */
+	void OnPasteParameterValues(int32 ParameterGroupIndex);
+
+	/**
+	 * Whether it is possible to paste parameter values onto the given group index,
+	 * and if there is anything on the clipboard to paste from.
+	 */
+	bool CanPasteParameterValues(int32 ParameterGroupIndex);
+
 	/** Creates all the lightmass property override widgets. */
 	void CreateLightmassOverrideWidgets(IDetailLayoutBuilder& DetailLayout);
 
