Do not save DDC for the material stats resources (UE-110824).

diff --git a/Engine/Source/Editor/MaterialEditor/Public/MaterialStatsCommon.h b/Engine/Source/Editor/MaterialEditor/Public/MaterialStatsCommon.h
--- a/Engine/Source/Editor/MaterialEditor/Public/MaterialStatsCommon.h
+++ b/Engine/Source/Editor/MaterialEditor/Public/MaterialStatsCommon.h
@@ -18,6 +18,11 @@ public:
 		return Material;
 	}
 
+	/**
+	 * Should shaders compiled for this material be saved to disk?
+	 */
+	virtual bool IsPersistent() const override { return false; }
+
 	/** this will enable shader source extraction and pass paths to (eventual) offline shader compilers */
 	virtual void SetupExtaCompilationSettings(const EShaderPlatform Platform, FExtraShaderCompilerSettings& Settings) const override;
 };
