Fix for preview to not try output Strata type.

diff --git a/Engine/Source/Editor/MaterialEditor/Private/MaterialEditor.cpp b/Engine/Source/Editor/MaterialEditor/Private/MaterialEditor.cpp
--- a/Engine/Source/Editor/MaterialEditor/Private/MaterialEditor.cpp
+++ b/Engine/Source/Editor/MaterialEditor/Private/MaterialEditor.cpp
@@ -217,8 +217,20 @@ int32 FMatExpressionPreview::CompilePropertyAndSetMaterialProperty(EMaterialProp
 	{
 		// Hardcoding output 0 as we don't have the UI to specify any other output
 		const int32 OutputIndex = 0;
+		const uint32 Output0Type = Expression->GetOutputType(OutputIndex);
+		int32 PreviewCodeChunk = INDEX_NONE;
+		if (Output0Type != MCT_Strata)
+		{
+			PreviewCodeChunk = Expression->CompilePreview(Compiler, OutputIndex);
+		}
+		else
+		{
+			// This is used for instance to prevent any Strata typed output to be used as preview, for instance for material functions.
+			PreviewCodeChunk = Compiler->Constant(0.0f);
+		}
+
 		// Get back into gamma corrected space, as DrawTile does not do this adjustment.
-		Ret = Compiler->Power(Compiler->Max(Expression->CompilePreview(Compiler, OutputIndex), Compiler->Constant(0)), Compiler->Constant(1.f / 2.2f));
+		Ret = Compiler->Power(Compiler->Max(PreviewCodeChunk, Compiler->Constant(0)), Compiler->Constant(1.f / 2.2f));
 	}
 	else if (Property == MP_WorldPositionOffset)
 	{
