Only display the strata material description if it exists.

diff --git a/Engine/Source/Editor/MaterialEditor/Private/MaterialEditor.cpp b/Engine/Source/Editor/MaterialEditor/Private/MaterialEditor.cpp
--- a/Engine/Source/Editor/MaterialEditor/Private/MaterialEditor.cpp
+++ b/Engine/Source/Editor/MaterialEditor/Private/MaterialEditor.cpp
@@ -2785,9 +2785,12 @@ void FMaterialEditor::UpdateMaterialinfoList_Old()
 				{
 					const FString StrataMaterialDescription = ShaderMap->GetStrataMaterialDescription();
 
-					TSharedRef<FTokenizedMessage> Line = FTokenizedMessage::Create(EMessageSeverity::Info);
-					Line->AddToken(FTextToken::Create(FText::FromString(StrataMaterialDescription)));
-					Messages.Add(Line);
+					if (StrataMaterialDescription.Len())
+					{
+						TSharedRef<FTokenizedMessage> Line = FTokenizedMessage::Create(EMessageSeverity::Info);
+						Line->AddToken(FTextToken::Create(FText::FromString(StrataMaterialDescription)));
+						Messages.Add(Line);
+					}
 
 					// Add shader count
 					FString ShaderCountString = FString::Printf(TEXT("Shader Count: %u"), ShaderMap->GetShaderNum());
