Handle null vertex factory in stats collection.

diff --git a/Engine/Source/Editor/MaterialEditor/Private/MaterialStats.cpp b/Engine/Source/Editor/MaterialEditor/Private/MaterialStats.cpp
--- a/Engine/Source/Editor/MaterialEditor/Private/MaterialStats.cpp
+++ b/Engine/Source/Editor/MaterialEditor/Private/MaterialStats.cpp
@@ -241,7 +241,7 @@ bool FShaderPlatformSettings::Update()
 					{
 						FShaderType* EntryShader = Entry.Value.GetType();
 						FVertexFactoryType* VertexFactory = Entry.Value.GetVertexFactoryType();
-						FString ShaderName = FString::Printf(TEXT("%s/%s"), VertexFactory->GetName(), EntryShader->GetName());
+						FString ShaderName = FString::Printf(TEXT("%s/%s"), VertexFactory ? VertexFactory->GetName() : TEXT("NullVF"), EntryShader->GetName());
 						QualityItem.ArrShaderEntries.Add(MakeShareable(new FMaterialShaderEntry{ Entry.Key, ShaderName }));
 					}
 
