Fix null vertex factory dereference when building list of shaders for material shader code windows.

diff --git a/Engine/Source/Editor/MaterialEditor/Private/MaterialStats.cpp b/Engine/Source/Editor/MaterialEditor/Private/MaterialStats.cpp
--- a/Engine/Source/Editor/MaterialEditor/Private/MaterialStats.cpp
+++ b/Engine/Source/Editor/MaterialEditor/Private/MaterialStats.cpp
@@ -123,7 +123,9 @@ FText FShaderPlatformSettings::GetShaderCode(const EMaterialQualityLevel::Type Q
 		{
 			const TShaderRef<FShader>& Shader = *Entry;
 
-			const FName VertexFactoryName = Shader.GetVertexFactoryType()->GetFName();
+			static FName NAME_NullVF(TEXT("NullVF"));
+			FVertexFactoryType* VertexFactory = Shader.GetVertexFactoryType();
+			const FName VertexFactoryName = VertexFactory ? VertexFactory->GetFName() : NAME_NullVF;
 			const FName ShaderTypeName = Shader.GetType()->GetFName();
 			const FMemoryImageString* ShaderSource = MaterialShaderMap->GetShaderSource(VertexFactoryName, ShaderTypeName);
 			if (ShaderSource != nullptr)
