Don't crash if a MF tries to use material layers (will eventually be flagged as a compile error, but make sure we don't crash first)

diff --git a/Engine/Source/Editor/MaterialEditor/Private/SMaterialLayersFunctionsTree.cpp b/Engine/Source/Editor/MaterialEditor/Private/SMaterialLayersFunctionsTree.cpp
--- a/Engine/Source/Editor/MaterialEditor/Private/SMaterialLayersFunctionsTree.cpp
+++ b/Engine/Source/Editor/MaterialEditor/Private/SMaterialLayersFunctionsTree.cpp
@@ -2755,109 +2755,112 @@ void SMaterialLayersFunctionsMaterialTree::CreateGroupsWidget()
 				TSharedPtr<IPropertyHandle> BlendHandle = ChildHandle->GetChildHandle("Blends").ToSharedRef();
 				uint32 LayerChildren;
 				LayerHandle->GetNumChildren(LayerChildren);
-				uint32 BlendChildren;
-				BlendHandle->GetNumChildren(BlendChildren);
-				if (MaterialEditorInstance->StoredLayerPreviews.Num() != LayerChildren)
-				{
-					MaterialEditorInstance->StoredLayerPreviews.Empty();
-					MaterialEditorInstance->StoredLayerPreviews.AddDefaulted(LayerChildren);
-				}
-				if (MaterialEditorInstance->StoredBlendPreviews.Num() != BlendChildren)
+				if (LayerChildren > 0)
 				{
-					MaterialEditorInstance->StoredBlendPreviews.Empty();
-					MaterialEditorInstance->StoredBlendPreviews.AddDefaulted(BlendChildren);
-				}
-
-				TSharedRef<FSortedParamData> StackProperty = MakeShared<FSortedParamData>();
-				StackProperty->StackDataType = EStackDataType::Stack;
-				StackProperty->Parameter = Parameter;
-				StackProperty->ParameterInfo.Index = LayerChildren - 1;
-				StackProperty->NodeKey = FString::FromInt(StackProperty->ParameterInfo.Index);
-
-
-				TSharedRef<FSortedParamData> ChildProperty = MakeShared<FSortedParamData>();
-				ChildProperty->StackDataType = EStackDataType::Asset;
-				ChildProperty->Parameter = Parameter;
-				ChildProperty->ParameterHandle = LayerHandle->AsArray()->GetElement(LayerChildren - 1);
-				ChildProperty->ParameterNode = Wrapper->GetGenerator()->FindTreeNode(ChildProperty->ParameterHandle);
-				ChildProperty->ParameterInfo.Index = LayerChildren - 1;
-				ChildProperty->ParameterInfo.Association = EMaterialParameterAssociation::LayerParameter;
-				ChildProperty->NodeKey = FString::FromInt(ChildProperty->ParameterInfo.Index) + FString::FromInt(ChildProperty->ParameterInfo.Association);
-
-				UObject* AssetObject;
-				ChildProperty->ParameterHandle->GetValue(AssetObject);
-				if (AssetObject)
-				{
-					if (MaterialEditorInstance->StoredLayerPreviews[LayerChildren - 1] == nullptr)
+					uint32 BlendChildren;
+					BlendHandle->GetNumChildren(BlendChildren);
+					if (MaterialEditorInstance->StoredLayerPreviews.Num() != LayerChildren)
 					{
-						MaterialEditorInstance->StoredLayerPreviews[LayerChildren - 1] = (NewObject<UMaterialInstanceConstant>(MaterialEditorInstance, NAME_None));
+						MaterialEditorInstance->StoredLayerPreviews.Empty();
+						MaterialEditorInstance->StoredLayerPreviews.AddDefaulted(LayerChildren);
 					}
-					UMaterialInterface* EditedMaterial = Cast<UMaterialFunctionInterface>(AssetObject)->GetPreviewMaterial();
-					if (MaterialEditorInstance->StoredLayerPreviews[LayerChildren - 1] && MaterialEditorInstance->StoredLayerPreviews[LayerChildren - 1]->Parent != EditedMaterial)
+					if (MaterialEditorInstance->StoredBlendPreviews.Num() != BlendChildren)
 					{
-						MaterialEditorInstance->StoredLayerPreviews[LayerChildren - 1]->SetParentEditorOnly(EditedMaterial);
+						MaterialEditorInstance->StoredBlendPreviews.Empty();
+						MaterialEditorInstance->StoredBlendPreviews.AddDefaulted(BlendChildren);
 					}
-				}
-
-				StackProperty->Children.Add(ChildProperty);
-				LayerProperties.Add(StackProperty);
 
-				if (BlendChildren > 0 && LayerChildren > BlendChildren)
-				{
-					for (int32 Counter = BlendChildren - 1; Counter >= 0; Counter--)
+					TSharedRef<FSortedParamData> StackProperty = MakeShared<FSortedParamData>();
+					StackProperty->StackDataType = EStackDataType::Stack;
+					StackProperty->Parameter = Parameter;
+					StackProperty->ParameterInfo.Index = LayerChildren - 1;
+					StackProperty->NodeKey = FString::FromInt(StackProperty->ParameterInfo.Index);
+
+
+					TSharedRef<FSortedParamData> ChildProperty = MakeShared<FSortedParamData>();
+					ChildProperty->StackDataType = EStackDataType::Asset;
+					ChildProperty->Parameter = Parameter;
+					ChildProperty->ParameterHandle = LayerHandle->AsArray()->GetElement(LayerChildren - 1);
+					ChildProperty->ParameterNode = Wrapper->GetGenerator()->FindTreeNode(ChildProperty->ParameterHandle);
+					ChildProperty->ParameterInfo.Index = LayerChildren - 1;
+					ChildProperty->ParameterInfo.Association = EMaterialParameterAssociation::LayerParameter;
+					ChildProperty->NodeKey = FString::FromInt(ChildProperty->ParameterInfo.Index) + FString::FromInt(ChildProperty->ParameterInfo.Association);
+
+					UObject* AssetObject;
+					ChildProperty->ParameterHandle->GetValue(AssetObject);
+					if (AssetObject)
 					{
-						ChildProperty = MakeShared<FSortedParamData>();
-						ChildProperty->StackDataType = EStackDataType::Asset;
-						ChildProperty->Parameter = Parameter;
-						ChildProperty->ParameterHandle = BlendHandle->AsArray()->GetElement(Counter);
-						ChildProperty->ParameterNode = Wrapper->GetGenerator()->FindTreeNode(ChildProperty->ParameterHandle);
-						ChildProperty->ParameterInfo.Index = Counter;
-						ChildProperty->ParameterInfo.Association = EMaterialParameterAssociation::BlendParameter;
-						ChildProperty->NodeKey = FString::FromInt(ChildProperty->ParameterInfo.Index) + FString::FromInt(ChildProperty->ParameterInfo.Association);
-						ChildProperty->ParameterHandle->GetValue(AssetObject);
-						if (AssetObject)
+						if (MaterialEditorInstance->StoredLayerPreviews[LayerChildren - 1] == nullptr)
 						{
-							if (MaterialEditorInstance->StoredBlendPreviews[Counter] == nullptr)
-							{
-								MaterialEditorInstance->StoredBlendPreviews[Counter] = (NewObject<UMaterialInstanceConstant>(MaterialEditorInstance, NAME_None));
-							}
-							UMaterialInterface* EditedMaterial = Cast<UMaterialFunctionInterface>(AssetObject)->GetPreviewMaterial();
-							if (MaterialEditorInstance->StoredBlendPreviews[Counter] && MaterialEditorInstance->StoredBlendPreviews[Counter]->Parent != EditedMaterial)
-							{
-								MaterialEditorInstance->StoredBlendPreviews[Counter]->SetParentEditorOnly(EditedMaterial);
-							}
+							MaterialEditorInstance->StoredLayerPreviews[LayerChildren - 1] = (NewObject<UMaterialInstanceConstant>(MaterialEditorInstance, NAME_None));
 						}
-						LayerProperties.Last()->Children.Add(ChildProperty);
+						UMaterialInterface* EditedMaterial = Cast<UMaterialFunctionInterface>(AssetObject)->GetPreviewMaterial();
+						if (MaterialEditorInstance->StoredLayerPreviews[LayerChildren - 1] && MaterialEditorInstance->StoredLayerPreviews[LayerChildren - 1]->Parent != EditedMaterial)
+						{
+							MaterialEditorInstance->StoredLayerPreviews[LayerChildren - 1]->SetParentEditorOnly(EditedMaterial);
+						}
+					}
 
-						StackProperty = MakeShared<FSortedParamData>();
-						StackProperty->StackDataType = EStackDataType::Stack;
-						StackProperty->Parameter = Parameter;
-						StackProperty->ParameterInfo.Index = Counter;
-						StackProperty->NodeKey = FString::FromInt(StackProperty->ParameterInfo.Index);
-						LayerProperties.Add(StackProperty);
+					StackProperty->Children.Add(ChildProperty);
+					LayerProperties.Add(StackProperty);
 
-						ChildProperty = MakeShared<FSortedParamData>();
-						ChildProperty->StackDataType = EStackDataType::Asset;
-						ChildProperty->Parameter = Parameter;
-						ChildProperty->ParameterHandle = LayerHandle->AsArray()->GetElement(Counter);
-						ChildProperty->ParameterNode = Wrapper->GetGenerator()->FindTreeNode(ChildProperty->ParameterHandle);
-						ChildProperty->ParameterInfo.Index = Counter;
-						ChildProperty->ParameterInfo.Association = EMaterialParameterAssociation::LayerParameter;
-						ChildProperty->NodeKey = FString::FromInt(ChildProperty->ParameterInfo.Index) + FString::FromInt(ChildProperty->ParameterInfo.Association);
-						ChildProperty->ParameterHandle->GetValue(AssetObject);
-						if (AssetObject)
+					if (BlendChildren > 0 && LayerChildren > BlendChildren)
+					{
+						for (int32 Counter = BlendChildren - 1; Counter >= 0; Counter--)
 						{
-							if (MaterialEditorInstance->StoredLayerPreviews[Counter] == nullptr)
+							ChildProperty = MakeShared<FSortedParamData>();
+							ChildProperty->StackDataType = EStackDataType::Asset;
+							ChildProperty->Parameter = Parameter;
+							ChildProperty->ParameterHandle = BlendHandle->AsArray()->GetElement(Counter);
+							ChildProperty->ParameterNode = Wrapper->GetGenerator()->FindTreeNode(ChildProperty->ParameterHandle);
+							ChildProperty->ParameterInfo.Index = Counter;
+							ChildProperty->ParameterInfo.Association = EMaterialParameterAssociation::BlendParameter;
+							ChildProperty->NodeKey = FString::FromInt(ChildProperty->ParameterInfo.Index) + FString::FromInt(ChildProperty->ParameterInfo.Association);
+							ChildProperty->ParameterHandle->GetValue(AssetObject);
+							if (AssetObject)
 							{
-								MaterialEditorInstance->StoredLayerPreviews[Counter] = (NewObject<UMaterialInstanceConstant>(MaterialEditorInstance, NAME_None));
+								if (MaterialEditorInstance->StoredBlendPreviews[Counter] == nullptr)
+								{
+									MaterialEditorInstance->StoredBlendPreviews[Counter] = (NewObject<UMaterialInstanceConstant>(MaterialEditorInstance, NAME_None));
+								}
+								UMaterialInterface* EditedMaterial = Cast<UMaterialFunctionInterface>(AssetObject)->GetPreviewMaterial();
+								if (MaterialEditorInstance->StoredBlendPreviews[Counter] && MaterialEditorInstance->StoredBlendPreviews[Counter]->Parent != EditedMaterial)
+								{
+									MaterialEditorInstance->StoredBlendPreviews[Counter]->SetParentEditorOnly(EditedMaterial);
+								}
 							}
-							UMaterialInterface* EditedMaterial = Cast<UMaterialFunctionInterface>(AssetObject)->GetPreviewMaterial();
-							if (MaterialEditorInstance->StoredLayerPreviews[Counter] && MaterialEditorInstance->StoredLayerPreviews[Counter]->Parent != EditedMaterial)
+							LayerProperties.Last()->Children.Add(ChildProperty);
+
+							StackProperty = MakeShared<FSortedParamData>();
+							StackProperty->StackDataType = EStackDataType::Stack;
+							StackProperty->Parameter = Parameter;
+							StackProperty->ParameterInfo.Index = Counter;
+							StackProperty->NodeKey = FString::FromInt(StackProperty->ParameterInfo.Index);
+							LayerProperties.Add(StackProperty);
+
+							ChildProperty = MakeShared<FSortedParamData>();
+							ChildProperty->StackDataType = EStackDataType::Asset;
+							ChildProperty->Parameter = Parameter;
+							ChildProperty->ParameterHandle = LayerHandle->AsArray()->GetElement(Counter);
+							ChildProperty->ParameterNode = Wrapper->GetGenerator()->FindTreeNode(ChildProperty->ParameterHandle);
+							ChildProperty->ParameterInfo.Index = Counter;
+							ChildProperty->ParameterInfo.Association = EMaterialParameterAssociation::LayerParameter;
+							ChildProperty->NodeKey = FString::FromInt(ChildProperty->ParameterInfo.Index) + FString::FromInt(ChildProperty->ParameterInfo.Association);
+							ChildProperty->ParameterHandle->GetValue(AssetObject);
+							if (AssetObject)
 							{
-								MaterialEditorInstance->StoredLayerPreviews[Counter]->SetParentEditorOnly(EditedMaterial);
+								if (MaterialEditorInstance->StoredLayerPreviews[Counter] == nullptr)
+								{
+									MaterialEditorInstance->StoredLayerPreviews[Counter] = (NewObject<UMaterialInstanceConstant>(MaterialEditorInstance, NAME_None));
+								}
+								UMaterialInterface* EditedMaterial = Cast<UMaterialFunctionInterface>(AssetObject)->GetPreviewMaterial();
+								if (MaterialEditorInstance->StoredLayerPreviews[Counter] && MaterialEditorInstance->StoredLayerPreviews[Counter]->Parent != EditedMaterial)
+								{
+									MaterialEditorInstance->StoredLayerPreviews[Counter]->SetParentEditorOnly(EditedMaterial);
+								}
 							}
+							LayerProperties.Last()->Children.Add(ChildProperty);
 						}
-						LayerProperties.Last()->Children.Add(ChildProperty);
 					}
 				}
 			}
