Allow 'Output Translucent Velocity' to be overriden by material instances

diff --git a/Engine/Plugins/Editor/MaterialAnalyzer/Source/Private/SMaterialAnalyzer.cpp b/Engine/Plugins/Editor/MaterialAnalyzer/Source/Private/SMaterialAnalyzer.cpp
--- a/Engine/Plugins/Editor/MaterialAnalyzer/Source/Private/SMaterialAnalyzer.cpp
+++ b/Engine/Plugins/Editor/MaterialAnalyzer/Source/Private/SMaterialAnalyzer.cpp
@@ -35,6 +35,7 @@ SMaterialAnalyzer::SMaterialAnalyzer()
 	BasePropertyOverrideNames.Add(TEXT("bOverride_DitheredLODTransition"), TEXT("DitheredLODTransitionOverride"));
 	BasePropertyOverrideNames.Add(TEXT("bOverride_CastDynamicShadowAsMasked"), TEXT("CastDynamicShadowAsMaskedOverride"));
 	BasePropertyOverrideNames.Add(TEXT("bOverride_TwoSided"), TEXT("TwoSidedOverride"));
+	BasePropertyOverrideNames.Add(TEXT("bOverride_OutputTranslucentVelocity"), TEXT("bOutputTranslucentVelocity"));
 }
 
 SMaterialAnalyzer::~SMaterialAnalyzer()
@@ -767,6 +768,14 @@ void FAnalyzeMaterialTreeAsyncTask::DoWork()
 				bIsOverridden = CurrentMaterialInstance->BasePropertyOverrides.bOverride_TwoSided;
 			}
 		}
+		else if (BasePropertyOverrideName.Key.IsEqual(TEXT("bOverride_OutputTranslucentVelocity")))
+		{
+			TempValue = CurrentMaterialInterface->IsTranslucencyWritingVelocity();
+			if (CurrentMaterialInstance)
+			{
+				bIsOverridden = CurrentMaterialInstance->BasePropertyOverrides.bOverride_OutputTranslucentVelocity;
+			}
+		}
 
 		// Check the parent for this variable
 		FAnalyzedMaterialNodePtr Parent = CurrentMaterialNode->Parent;
diff --git a/Engine/Source/Editor/MaterialEditor/Private/MaterialEditorInstanceDetailCustomization.cpp b/Engine/Source/Editor/MaterialEditor/Private/MaterialEditorInstanceDetailCustomization.cpp
--- a/Engine/Source/Editor/MaterialEditor/Private/MaterialEditorInstanceDetailCustomization.cpp
+++ b/Engine/Source/Editor/MaterialEditor/Private/MaterialEditorInstanceDetailCustomization.cpp
@@ -1237,6 +1237,7 @@ void FMaterialInstanceParameterDetails::CreateBasePropertyOverrideWidgets(IDetai
 	TAttribute<bool> IsOverrideShadingModelEnabled = TAttribute<bool>::Create(TAttribute<bool>::FGetter::CreateSP(this, &FMaterialInstanceParameterDetails::OverrideShadingModelEnabled));
 	TAttribute<bool> IsOverrideTwoSidedEnabled = TAttribute<bool>::Create(TAttribute<bool>::FGetter::CreateSP(this, &FMaterialInstanceParameterDetails::OverrideTwoSidedEnabled));
 	TAttribute<bool> IsOverrideDitheredLODTransitionEnabled = TAttribute<bool>::Create(TAttribute<bool>::FGetter::CreateSP(this, &FMaterialInstanceParameterDetails::OverrideDitheredLODTransitionEnabled));
+	TAttribute<bool> IsOverrideOutputTranslucentVelocityEnabled = TAttribute<bool>::Create(TAttribute<bool>::FGetter::CreateSP(this, &FMaterialInstanceParameterDetails::OverrideOutputTranslucentVelocityEnabled));
 
 	TSharedRef<IPropertyHandle> BasePropertyOverridePropery = DetailLayout.GetProperty("BasePropertyOverrides");
 	TSharedPtr<IPropertyHandle> OpacityClipMaskValueProperty = BasePropertyOverridePropery->GetChildHandle("OpacityMaskClipValue");
@@ -1244,117 +1245,143 @@ void FMaterialInstanceParameterDetails::CreateBasePropertyOverrideWidgets(IDetai
 	TSharedPtr<IPropertyHandle> ShadingModelProperty = BasePropertyOverridePropery->GetChildHandle("ShadingModel");
 	TSharedPtr<IPropertyHandle> TwoSidedProperty = BasePropertyOverridePropery->GetChildHandle("TwoSided");
 	TSharedPtr<IPropertyHandle> DitheredLODTransitionProperty = BasePropertyOverridePropery->GetChildHandle("DitheredLODTransition");
+	TSharedPtr<IPropertyHandle> OutputTranslucentVelocityProperty = BasePropertyOverridePropery->GetChildHandle("bOutputTranslucentVelocity");
 
-	FIsResetToDefaultVisible IsOpacityClipMaskValuePropertyResetVisible = FIsResetToDefaultVisible::CreateLambda([this](TSharedPtr<IPropertyHandle> InHandle) {
-		return MaterialEditorInstance->Parent != nullptr ? MaterialEditorInstance->BasePropertyOverrides.OpacityMaskClipValue != MaterialEditorInstance->Parent->GetOpacityMaskClipValue() : false;
-	});
-	FResetToDefaultHandler ResetOpacityClipMaskValuePropertyHandler = FResetToDefaultHandler::CreateLambda([this](TSharedPtr<IPropertyHandle> InHandle) {
-		if (MaterialEditorInstance->Parent != nullptr)
-		{
-			MaterialEditorInstance->BasePropertyOverrides.OpacityMaskClipValue = MaterialEditorInstance->Parent->GetOpacityMaskClipValue();
-		}
-	});
-	FResetToDefaultOverride ResetOpacityClipMaskValuePropertyOverride = FResetToDefaultOverride::Create(IsOpacityClipMaskValuePropertyResetVisible, ResetOpacityClipMaskValuePropertyHandler);
-	IDetailPropertyRow& OpacityClipMaskValuePropertyRow = BasePropertyOverrideGroup.AddPropertyRow(OpacityClipMaskValueProperty.ToSharedRef());
-	OpacityClipMaskValuePropertyRow
-		.DisplayName(OpacityClipMaskValueProperty->GetPropertyDisplayName())
-		.ToolTip(OpacityClipMaskValueProperty->GetToolTipText())
-		.EditCondition(IsOverrideOpacityClipMaskValueEnabled, FOnBooleanValueChanged::CreateSP(this, &FMaterialInstanceParameterDetails::OnOverrideOpacityClipMaskValueChanged))
-		.Visibility(TAttribute<EVisibility>::Create(TAttribute<EVisibility>::FGetter::CreateSP(this, &FMaterialInstanceParameterDetails::IsOverriddenAndVisible, IsOverrideOpacityClipMaskValueEnabled)))
-		.OverrideResetToDefault(ResetOpacityClipMaskValuePropertyOverride);
-
-	FIsResetToDefaultVisible IsBlendModePropertyResetVisible = FIsResetToDefaultVisible::CreateLambda([this](TSharedPtr<IPropertyHandle> InHandle) {
-		return MaterialEditorInstance->Parent != nullptr ? MaterialEditorInstance->BasePropertyOverrides.BlendMode != MaterialEditorInstance->Parent->GetBlendMode() : false;
-	});
-	FResetToDefaultHandler ResetBlendModePropertyHandler = FResetToDefaultHandler::CreateLambda([this](TSharedPtr<IPropertyHandle> InHandle) {
-		if (MaterialEditorInstance->Parent != nullptr)
-		{
-			MaterialEditorInstance->BasePropertyOverrides.BlendMode = MaterialEditorInstance->Parent->GetBlendMode();
-		}
-	});
-	FResetToDefaultOverride ResetBlendModePropertyOverride = FResetToDefaultOverride::Create(IsBlendModePropertyResetVisible, ResetBlendModePropertyHandler);
-	IDetailPropertyRow& BlendModePropertyRow = BasePropertyOverrideGroup.AddPropertyRow(BlendModeProperty.ToSharedRef());
-	BlendModePropertyRow
-		.DisplayName(BlendModeProperty->GetPropertyDisplayName())
-		.ToolTip(BlendModeProperty->GetToolTipText())
-		.EditCondition(IsOverrideBlendModeEnabled, FOnBooleanValueChanged::CreateSP(this, &FMaterialInstanceParameterDetails::OnOverrideBlendModeChanged))
-		.Visibility(TAttribute<EVisibility>::Create(TAttribute<EVisibility>::FGetter::CreateSP(this, &FMaterialInstanceParameterDetails::IsOverriddenAndVisible, IsOverrideBlendModeEnabled)))
-		.OverrideResetToDefault(ResetBlendModePropertyOverride);
-
-	FIsResetToDefaultVisible IsShadingModelPropertyResetVisible = FIsResetToDefaultVisible::CreateLambda([this](TSharedPtr<IPropertyHandle> InHandle) {
-		if (MaterialEditorInstance->Parent != nullptr)
-		{	
-			if (MaterialEditorInstance->Parent->IsShadingModelFromMaterialExpression())
+	{
+		FIsResetToDefaultVisible IsOpacityClipMaskValuePropertyResetVisible = FIsResetToDefaultVisible::CreateLambda([this](TSharedPtr<IPropertyHandle> InHandle) {
+			return MaterialEditorInstance->Parent != nullptr ? MaterialEditorInstance->BasePropertyOverrides.OpacityMaskClipValue != MaterialEditorInstance->Parent->GetOpacityMaskClipValue() : false;
+			});
+		FResetToDefaultHandler ResetOpacityClipMaskValuePropertyHandler = FResetToDefaultHandler::CreateLambda([this](TSharedPtr<IPropertyHandle> InHandle) {
+			if (MaterialEditorInstance->Parent != nullptr)
 			{
-				return MaterialEditorInstance->BasePropertyOverrides.ShadingModel != MSM_FromMaterialExpression;
+				MaterialEditorInstance->BasePropertyOverrides.OpacityMaskClipValue = MaterialEditorInstance->Parent->GetOpacityMaskClipValue();
 			}
-			else
+			});
+		FResetToDefaultOverride ResetOpacityClipMaskValuePropertyOverride = FResetToDefaultOverride::Create(IsOpacityClipMaskValuePropertyResetVisible, ResetOpacityClipMaskValuePropertyHandler);
+		IDetailPropertyRow& OpacityClipMaskValuePropertyRow = BasePropertyOverrideGroup.AddPropertyRow(OpacityClipMaskValueProperty.ToSharedRef());
+		OpacityClipMaskValuePropertyRow
+			.DisplayName(OpacityClipMaskValueProperty->GetPropertyDisplayName())
+			.ToolTip(OpacityClipMaskValueProperty->GetToolTipText())
+			.EditCondition(IsOverrideOpacityClipMaskValueEnabled, FOnBooleanValueChanged::CreateSP(this, &FMaterialInstanceParameterDetails::OnOverrideOpacityClipMaskValueChanged))
+			.Visibility(TAttribute<EVisibility>::Create(TAttribute<EVisibility>::FGetter::CreateSP(this, &FMaterialInstanceParameterDetails::IsOverriddenAndVisible, IsOverrideOpacityClipMaskValueEnabled)))
+			.OverrideResetToDefault(ResetOpacityClipMaskValuePropertyOverride);
+	}
+	{
+		FIsResetToDefaultVisible IsBlendModePropertyResetVisible = FIsResetToDefaultVisible::CreateLambda([this](TSharedPtr<IPropertyHandle> InHandle) {
+			return MaterialEditorInstance->Parent != nullptr ? MaterialEditorInstance->BasePropertyOverrides.BlendMode != MaterialEditorInstance->Parent->GetBlendMode() : false;
+			});
+		FResetToDefaultHandler ResetBlendModePropertyHandler = FResetToDefaultHandler::CreateLambda([this](TSharedPtr<IPropertyHandle> InHandle) {
+			if (MaterialEditorInstance->Parent != nullptr)
 			{
-				return MaterialEditorInstance->BasePropertyOverrides.ShadingModel != MaterialEditorInstance->Parent->GetShadingModels().GetFirstShadingModel();
+				MaterialEditorInstance->BasePropertyOverrides.BlendMode = MaterialEditorInstance->Parent->GetBlendMode();
 			}
-		}
-		else
-		{
-			return false;
-		}
-	});
-	FResetToDefaultHandler ResetShadingModelPropertyHandler = FResetToDefaultHandler::CreateLambda([this](TSharedPtr<IPropertyHandle> InHandle) {
-		if (MaterialEditorInstance->Parent != nullptr)
-		{
-			if (MaterialEditorInstance->Parent->IsShadingModelFromMaterialExpression())
+			});
+		FResetToDefaultOverride ResetBlendModePropertyOverride = FResetToDefaultOverride::Create(IsBlendModePropertyResetVisible, ResetBlendModePropertyHandler);
+		IDetailPropertyRow& BlendModePropertyRow = BasePropertyOverrideGroup.AddPropertyRow(BlendModeProperty.ToSharedRef());
+		BlendModePropertyRow
+			.DisplayName(BlendModeProperty->GetPropertyDisplayName())
+			.ToolTip(BlendModeProperty->GetToolTipText())
+			.EditCondition(IsOverrideBlendModeEnabled, FOnBooleanValueChanged::CreateSP(this, &FMaterialInstanceParameterDetails::OnOverrideBlendModeChanged))
+			.Visibility(TAttribute<EVisibility>::Create(TAttribute<EVisibility>::FGetter::CreateSP(this, &FMaterialInstanceParameterDetails::IsOverriddenAndVisible, IsOverrideBlendModeEnabled)))
+			.OverrideResetToDefault(ResetBlendModePropertyOverride);
+	}
+	{
+		FIsResetToDefaultVisible IsShadingModelPropertyResetVisible = FIsResetToDefaultVisible::CreateLambda([this](TSharedPtr<IPropertyHandle> InHandle) {
+			if (MaterialEditorInstance->Parent != nullptr)
 			{
-				MaterialEditorInstance->BasePropertyOverrides.ShadingModel = MSM_FromMaterialExpression;
+				if (MaterialEditorInstance->Parent->IsShadingModelFromMaterialExpression())
+				{
+					return MaterialEditorInstance->BasePropertyOverrides.ShadingModel != MSM_FromMaterialExpression;
+				}
+				else
+				{
+					return MaterialEditorInstance->BasePropertyOverrides.ShadingModel != MaterialEditorInstance->Parent->GetShadingModels().GetFirstShadingModel();
+				}
 			}
 			else
 			{
-				MaterialEditorInstance->BasePropertyOverrides.ShadingModel = MaterialEditorInstance->Parent->GetShadingModels().GetFirstShadingModel();
+				return false;
 			}
-		}
-	});
-	FResetToDefaultOverride ResetShadingModelPropertyOverride = FResetToDefaultOverride::Create(IsShadingModelPropertyResetVisible, ResetShadingModelPropertyHandler);
-	IDetailPropertyRow& ShadingModelPropertyRow = BasePropertyOverrideGroup.AddPropertyRow(ShadingModelProperty.ToSharedRef());
-	ShadingModelPropertyRow
-		.DisplayName(ShadingModelProperty->GetPropertyDisplayName())
-		.ToolTip(ShadingModelProperty->GetToolTipText())
-		.EditCondition(IsOverrideShadingModelEnabled, FOnBooleanValueChanged::CreateSP(this, &FMaterialInstanceParameterDetails::OnOverrideShadingModelChanged))
-		.Visibility(TAttribute<EVisibility>::Create(TAttribute<EVisibility>::FGetter::CreateSP(this, &FMaterialInstanceParameterDetails::IsOverriddenAndVisible, IsOverrideShadingModelEnabled)))
-		.OverrideResetToDefault(ResetShadingModelPropertyOverride);
-
-	FIsResetToDefaultVisible IsTwoSidedPropertyResetVisible = FIsResetToDefaultVisible::CreateLambda([this](TSharedPtr<IPropertyHandle> InHandle) {
-		return MaterialEditorInstance->Parent != nullptr ? MaterialEditorInstance->BasePropertyOverrides.TwoSided != MaterialEditorInstance->Parent->IsTwoSided() : false;
-	});
-	FResetToDefaultHandler ResetTwoSidedValuePropertyHandler = FResetToDefaultHandler::CreateLambda([this](TSharedPtr<IPropertyHandle> InHandle) {
-		if (MaterialEditorInstance->Parent != nullptr)
-		{
-			MaterialEditorInstance->BasePropertyOverrides.TwoSided = MaterialEditorInstance->Parent->IsTwoSided();
-		}
-	});
-	FResetToDefaultOverride ResetTwoSidedPropertyOverride = FResetToDefaultOverride::Create(IsTwoSidedPropertyResetVisible, ResetTwoSidedValuePropertyHandler);
-	IDetailPropertyRow& TwoSidedPropertyRow = BasePropertyOverrideGroup.AddPropertyRow(TwoSidedProperty.ToSharedRef());
-	TwoSidedPropertyRow
-		.DisplayName(TwoSidedProperty->GetPropertyDisplayName())
-		.ToolTip(TwoSidedProperty->GetToolTipText())
-		.EditCondition(IsOverrideTwoSidedEnabled, FOnBooleanValueChanged::CreateSP(this, &FMaterialInstanceParameterDetails::OnOverrideTwoSidedChanged))
-		.Visibility(TAttribute<EVisibility>::Create(TAttribute<EVisibility>::FGetter::CreateSP(this, &FMaterialInstanceParameterDetails::IsOverriddenAndVisible, IsOverrideTwoSidedEnabled)))
-		.OverrideResetToDefault(ResetTwoSidedPropertyOverride);
-
-	FIsResetToDefaultVisible IsDitheredLODTransitionPropertyResetVisible = FIsResetToDefaultVisible::CreateLambda([this](TSharedPtr<IPropertyHandle> InHandle) {
-		return MaterialEditorInstance->Parent != nullptr ? MaterialEditorInstance->BasePropertyOverrides.DitheredLODTransition != MaterialEditorInstance->Parent->IsDitheredLODTransition() : false;
-	});
-	FResetToDefaultHandler ResetDitheredLODTransitionPropertyHandler = FResetToDefaultHandler::CreateLambda([this](TSharedPtr<IPropertyHandle> InHandle) {
-		if (MaterialEditorInstance->Parent != nullptr)
-		{
-			MaterialEditorInstance->BasePropertyOverrides.DitheredLODTransition = MaterialEditorInstance->Parent->IsDitheredLODTransition();
-		}
-	});
-	FResetToDefaultOverride ResetDitheredLODTransitionPropertyOverride = FResetToDefaultOverride::Create(IsDitheredLODTransitionPropertyResetVisible, ResetDitheredLODTransitionPropertyHandler);
-	IDetailPropertyRow& DitheredLODTransitionPropertyRow = BasePropertyOverrideGroup.AddPropertyRow(DitheredLODTransitionProperty.ToSharedRef());
-	DitheredLODTransitionPropertyRow
-		.DisplayName(DitheredLODTransitionProperty->GetPropertyDisplayName())
-		.ToolTip(DitheredLODTransitionProperty->GetToolTipText())
-		.EditCondition(IsOverrideDitheredLODTransitionEnabled, FOnBooleanValueChanged::CreateSP(this, &FMaterialInstanceParameterDetails::OnOverrideDitheredLODTransitionChanged))
-		.Visibility(TAttribute<EVisibility>::Create(TAttribute<EVisibility>::FGetter::CreateSP(this, &FMaterialInstanceParameterDetails::IsOverriddenAndVisible, IsOverrideDitheredLODTransitionEnabled)))
-		.OverrideResetToDefault(ResetDitheredLODTransitionPropertyOverride);
+			});
+		FResetToDefaultHandler ResetShadingModelPropertyHandler = FResetToDefaultHandler::CreateLambda([this](TSharedPtr<IPropertyHandle> InHandle) {
+			if (MaterialEditorInstance->Parent != nullptr)
+			{
+				if (MaterialEditorInstance->Parent->IsShadingModelFromMaterialExpression())
+				{
+					MaterialEditorInstance->BasePropertyOverrides.ShadingModel = MSM_FromMaterialExpression;
+				}
+				else
+				{
+					MaterialEditorInstance->BasePropertyOverrides.ShadingModel = MaterialEditorInstance->Parent->GetShadingModels().GetFirstShadingModel();
+				}
+			}
+			});
+		FResetToDefaultOverride ResetShadingModelPropertyOverride = FResetToDefaultOverride::Create(IsShadingModelPropertyResetVisible, ResetShadingModelPropertyHandler);
+		IDetailPropertyRow& ShadingModelPropertyRow = BasePropertyOverrideGroup.AddPropertyRow(ShadingModelProperty.ToSharedRef());
+		ShadingModelPropertyRow
+			.DisplayName(ShadingModelProperty->GetPropertyDisplayName())
+			.ToolTip(ShadingModelProperty->GetToolTipText())
+			.EditCondition(IsOverrideShadingModelEnabled, FOnBooleanValueChanged::CreateSP(this, &FMaterialInstanceParameterDetails::OnOverrideShadingModelChanged))
+			.Visibility(TAttribute<EVisibility>::Create(TAttribute<EVisibility>::FGetter::CreateSP(this, &FMaterialInstanceParameterDetails::IsOverriddenAndVisible, IsOverrideShadingModelEnabled)))
+			.OverrideResetToDefault(ResetShadingModelPropertyOverride);
+	}
+	{
+		FIsResetToDefaultVisible IsTwoSidedPropertyResetVisible = FIsResetToDefaultVisible::CreateLambda([this](TSharedPtr<IPropertyHandle> InHandle) {
+			return MaterialEditorInstance->Parent != nullptr ? MaterialEditorInstance->BasePropertyOverrides.TwoSided != MaterialEditorInstance->Parent->IsTwoSided() : false;
+			});
+		FResetToDefaultHandler ResetTwoSidedValuePropertyHandler = FResetToDefaultHandler::CreateLambda([this](TSharedPtr<IPropertyHandle> InHandle) {
+			if (MaterialEditorInstance->Parent != nullptr)
+			{
+				MaterialEditorInstance->BasePropertyOverrides.TwoSided = MaterialEditorInstance->Parent->IsTwoSided();
+			}
+			});
+		FResetToDefaultOverride ResetTwoSidedPropertyOverride = FResetToDefaultOverride::Create(IsTwoSidedPropertyResetVisible, ResetTwoSidedValuePropertyHandler);
+		IDetailPropertyRow& TwoSidedPropertyRow = BasePropertyOverrideGroup.AddPropertyRow(TwoSidedProperty.ToSharedRef());
+		TwoSidedPropertyRow
+			.DisplayName(TwoSidedProperty->GetPropertyDisplayName())
+			.ToolTip(TwoSidedProperty->GetToolTipText())
+			.EditCondition(IsOverrideTwoSidedEnabled, FOnBooleanValueChanged::CreateSP(this, &FMaterialInstanceParameterDetails::OnOverrideTwoSidedChanged))
+			.Visibility(TAttribute<EVisibility>::Create(TAttribute<EVisibility>::FGetter::CreateSP(this, &FMaterialInstanceParameterDetails::IsOverriddenAndVisible, IsOverrideTwoSidedEnabled)))
+			.OverrideResetToDefault(ResetTwoSidedPropertyOverride);
+	}
+	{
+		FIsResetToDefaultVisible IsDitheredLODTransitionPropertyResetVisible = FIsResetToDefaultVisible::CreateLambda([this](TSharedPtr<IPropertyHandle> InHandle) {
+			return MaterialEditorInstance->Parent != nullptr ? MaterialEditorInstance->BasePropertyOverrides.DitheredLODTransition != MaterialEditorInstance->Parent->IsDitheredLODTransition() : false;
+			});
+		FResetToDefaultHandler ResetDitheredLODTransitionPropertyHandler = FResetToDefaultHandler::CreateLambda([this](TSharedPtr<IPropertyHandle> InHandle) {
+			if (MaterialEditorInstance->Parent != nullptr)
+			{
+				MaterialEditorInstance->BasePropertyOverrides.DitheredLODTransition = MaterialEditorInstance->Parent->IsDitheredLODTransition();
+			}
+			});
+		FResetToDefaultOverride ResetDitheredLODTransitionPropertyOverride = FResetToDefaultOverride::Create(IsDitheredLODTransitionPropertyResetVisible, ResetDitheredLODTransitionPropertyHandler);
+		IDetailPropertyRow& DitheredLODTransitionPropertyRow = BasePropertyOverrideGroup.AddPropertyRow(DitheredLODTransitionProperty.ToSharedRef());
+		DitheredLODTransitionPropertyRow
+			.DisplayName(DitheredLODTransitionProperty->GetPropertyDisplayName())
+			.ToolTip(DitheredLODTransitionProperty->GetToolTipText())
+			.EditCondition(IsOverrideDitheredLODTransitionEnabled, FOnBooleanValueChanged::CreateSP(this, &FMaterialInstanceParameterDetails::OnOverrideDitheredLODTransitionChanged))
+			.Visibility(TAttribute<EVisibility>::Create(TAttribute<EVisibility>::FGetter::CreateSP(this, &FMaterialInstanceParameterDetails::IsOverriddenAndVisible, IsOverrideDitheredLODTransitionEnabled)))
+			.OverrideResetToDefault(ResetDitheredLODTransitionPropertyOverride);
+	}
+	{
+		FIsResetToDefaultVisible IsOutputTranslucentVelocityPropertyResetVisible = FIsResetToDefaultVisible::CreateLambda([this](TSharedPtr<IPropertyHandle> InHandle) {
+			return MaterialEditorInstance->Parent != nullptr ? MaterialEditorInstance->BasePropertyOverrides.bOutputTranslucentVelocity != MaterialEditorInstance->Parent->IsTranslucencyWritingVelocity() : false;
+			});
+		FResetToDefaultHandler ResetOutputTranslucentVelocityPropertyHandler = FResetToDefaultHandler::CreateLambda([this](TSharedPtr<IPropertyHandle> InHandle) {
+			if (MaterialEditorInstance->Parent != nullptr)
+			{
+				MaterialEditorInstance->BasePropertyOverrides.bOutputTranslucentVelocity = MaterialEditorInstance->Parent->IsTranslucencyWritingVelocity();
+			}
+			});
+		FResetToDefaultOverride ResetOutputTranslucentVelocityPropertyOverride = FResetToDefaultOverride::Create(IsOutputTranslucentVelocityPropertyResetVisible, ResetOutputTranslucentVelocityPropertyHandler);
+		IDetailPropertyRow& OutputTranslucentVelocityPropertyRow = BasePropertyOverrideGroup.AddPropertyRow(OutputTranslucentVelocityProperty.ToSharedRef());
+		OutputTranslucentVelocityPropertyRow
+			.DisplayName(OutputTranslucentVelocityProperty->GetPropertyDisplayName())
+			.ToolTip(OutputTranslucentVelocityProperty->GetToolTipText())
+			.EditCondition(IsOverrideOutputTranslucentVelocityEnabled, FOnBooleanValueChanged::CreateSP(this, &FMaterialInstanceParameterDetails::OnOverrideOutputTranslucentVelocityChanged))
+			.Visibility(TAttribute<EVisibility>::Create(TAttribute<EVisibility>::FGetter::CreateSP(this, &FMaterialInstanceParameterDetails::IsOverriddenAndVisible, IsOverrideOutputTranslucentVelocityEnabled)))
+			.OverrideResetToDefault(ResetOutputTranslucentVelocityPropertyOverride);
+	}
 }
 
 EVisibility FMaterialInstanceParameterDetails::IsOverriddenAndVisible(TAttribute<bool> IsOverridden) const
@@ -1392,6 +1419,11 @@ bool FMaterialInstanceParameterDetails::OverrideDitheredLODTransitionEnabled() c
 	return MaterialEditorInstance->BasePropertyOverrides.bOverride_DitheredLODTransition;
 }
 
+bool FMaterialInstanceParameterDetails::OverrideOutputTranslucentVelocityEnabled() const
+{
+	return MaterialEditorInstance->BasePropertyOverrides.bOverride_OutputTranslucentVelocity;
+}
+
 void FMaterialInstanceParameterDetails::OnOverrideOpacityClipMaskValueChanged(bool NewValue)
 {
 	MaterialEditorInstance->BasePropertyOverrides.bOverride_OpacityMaskClipValue = NewValue;
@@ -1427,5 +1459,12 @@ void FMaterialInstanceParameterDetails::OnOverrideDitheredLODTransitionChanged(b
 	FEditorSupportDelegates::RedrawAllViewports.Broadcast();
 }
 
+void FMaterialInstanceParameterDetails::OnOverrideOutputTranslucentVelocityChanged(bool NewValue)
+{
+	MaterialEditorInstance->BasePropertyOverrides.bOverride_OutputTranslucentVelocity = NewValue;
+	MaterialEditorInstance->PostEditChange();
+	FEditorSupportDelegates::RedrawAllViewports.Broadcast();
+}
+
 #undef LOCTEXT_NAMESPACE
 
diff --git a/Engine/Source/Editor/MaterialEditor/Private/MaterialEditorInstanceDetailCustomization.h b/Engine/Source/Editor/MaterialEditor/Private/MaterialEditorInstanceDetailCustomization.h
--- a/Engine/Source/Editor/MaterialEditor/Private/MaterialEditorInstanceDetailCustomization.h
+++ b/Engine/Source/Editor/MaterialEditor/Private/MaterialEditorInstanceDetailCustomization.h
@@ -112,11 +112,13 @@ private:
 	bool OverrideShadingModelEnabled() const;
 	bool OverrideTwoSidedEnabled() const;
 	bool OverrideDitheredLODTransitionEnabled() const;
+	bool OverrideOutputTranslucentVelocityEnabled() const;
 	void OnOverrideOpacityClipMaskValueChanged(bool NewValue);
 	void OnOverrideBlendModeChanged(bool NewValue);
 	void OnOverrideShadingModelChanged(bool NewValue);
 	void OnOverrideTwoSidedChanged(bool NewValue);
 	void OnOverrideDitheredLODTransitionChanged(bool NewValue);
+	void OnOverrideOutputTranslucentVelocityChanged(bool NewValue);
 
 private:
 	/** Object that stores all of the possible parameters we can edit */
diff --git a/Engine/Source/Editor/UnrealEd/Private/PreviewMaterial.cpp b/Engine/Source/Editor/UnrealEd/Private/PreviewMaterial.cpp
--- a/Engine/Source/Editor/UnrealEd/Private/PreviewMaterial.cpp
+++ b/Engine/Source/Editor/UnrealEd/Private/PreviewMaterial.cpp
@@ -1808,6 +1808,10 @@ void UMaterialEditorInstanceConstant::CopyBasePropertiesFromParent()
 	{
 		BasePropertyOverrides.TwoSided = SourceInstance->IsTwoSided();
 	}
+	if (!BasePropertyOverrides.bOverride_OutputTranslucentVelocity)
+	{
+		BasePropertyOverrides.bOutputTranslucentVelocity = SourceInstance->IsTranslucencyWritingVelocity();
+	}
 	if (!BasePropertyOverrides.DitheredLODTransition)
 	{
 		BasePropertyOverrides.DitheredLODTransition = SourceInstance->IsDitheredLODTransition();
diff --git a/Engine/Source/Runtime/Engine/Classes/Materials/MaterialInstance.h b/Engine/Source/Runtime/Engine/Classes/Materials/MaterialInstance.h
--- a/Engine/Source/Runtime/Engine/Classes/Materials/MaterialInstance.h
+++ b/Engine/Source/Runtime/Engine/Classes/Materials/MaterialInstance.h
@@ -357,6 +357,7 @@ class UMaterialInstance : public UMaterialInterface
 	uint8 TwoSided : 1;
 	uint8 DitheredLODTransition : 1;
 	uint8 bCastDynamicShadowAsMasked : 1;
+	uint8 bOutputTranslucentVelocity : 1;
 	uint8 bIsShadingModelFromMaterialExpression : 1;
 
 	TEnumAsByte<EBlendMode> BlendMode;
@@ -551,6 +552,7 @@ public:
 	ENGINE_API virtual FMaterialShadingModelField GetShadingModels() const override;
 	ENGINE_API virtual bool IsShadingModelFromMaterialExpression() const override;
 	ENGINE_API virtual bool IsTwoSided() const override;
+	ENGINE_API virtual bool IsTranslucencyWritingVelocity() const override;
 	ENGINE_API virtual bool IsDitheredLODTransition() const override;
 	ENGINE_API virtual bool IsMasked() const override;
 	
diff --git a/Engine/Source/Runtime/Engine/Classes/Materials/MaterialInstanceBasePropertyOverrides.h b/Engine/Source/Runtime/Engine/Classes/Materials/MaterialInstanceBasePropertyOverrides.h
--- a/Engine/Source/Runtime/Engine/Classes/Materials/MaterialInstanceBasePropertyOverrides.h
+++ b/Engine/Source/Runtime/Engine/Classes/Materials/MaterialInstanceBasePropertyOverrides.h
@@ -37,6 +37,10 @@ struct ENGINE_API FMaterialInstanceBasePropertyOverrides
 	UPROPERTY(EditAnywhere, Category = Material)
 	uint8 bOverride_TwoSided : 1;
 
+	/** Enables override of the output velocity property. */
+	UPROPERTY(EditAnywhere, Category = Material)
+	uint8 bOverride_OutputTranslucentVelocity : 1;
+
 	/** Indicates that the material should be rendered without backface culling and the normal should be flipped for backfaces. */
 	UPROPERTY(EditAnywhere, Category = Material, meta = (editcondition = "bOverride_TwoSided"))
 	uint8 TwoSided : 1;
@@ -49,6 +53,10 @@ struct ENGINE_API FMaterialInstanceBasePropertyOverrides
 	UPROPERTY(EditAnywhere, Category = Material, meta = (editcondition = "bOverride_CastShadowAsMasked", NoSpinbox = true))
 	uint8 bCastDynamicShadowAsMasked:1;
 
+	/** Whether the material should output velocity even though it has a translucent blend mode. */
+	UPROPERTY(EditAnywhere, Category = Material, meta = (editcondition = "bOverride_OutputTranslucentVelocity"))
+	uint8 bOutputTranslucentVelocity : 1;
+
 	/** The blend mode */
 	UPROPERTY(EditAnywhere, Category = Material, meta = (editcondition = "bOverride_BlendMode"))
 	TEnumAsByte<EBlendMode> BlendMode;
diff --git a/Engine/Source/Runtime/Engine/Classes/Materials/MaterialInstanceDynamic.h b/Engine/Source/Runtime/Engine/Classes/Materials/MaterialInstanceDynamic.h
--- a/Engine/Source/Runtime/Engine/Classes/Materials/MaterialInstanceDynamic.h
+++ b/Engine/Source/Runtime/Engine/Classes/Materials/MaterialInstanceDynamic.h
@@ -157,6 +157,7 @@ class ENGINE_API UMaterialInstanceDynamic : public UMaterialInstance
 	virtual bool IsShadingModelFromMaterialExpression() const override;
 	virtual EBlendMode GetBlendMode() const override;
 	virtual bool IsTwoSided() const override;
+	virtual bool IsTranslucencyWritingVelocity() const override;
 	virtual bool IsDitheredLODTransition() const override;
 	virtual bool IsMasked() const override;
 
diff --git a/Engine/Source/Runtime/Engine/Private/Materials/MaterialInstance.cpp b/Engine/Source/Runtime/Engine/Private/Materials/MaterialInstance.cpp
--- a/Engine/Source/Runtime/Engine/Private/Materials/MaterialInstance.cpp
+++ b/Engine/Source/Runtime/Engine/Private/Materials/MaterialInstance.cpp
@@ -1878,6 +1878,11 @@ bool UMaterialInstanceDynamic::IsTwoSided() const
 	return Parent ? Parent->IsTwoSided() : false;
 }
 
+bool UMaterialInstanceDynamic::IsTranslucencyWritingVelocity() const
+{
+	return Parent ? Parent->IsTranslucencyWritingVelocity() : false;
+}
+
 bool UMaterialInstanceDynamic::IsDitheredLODTransition() const
 {
 	return Parent ? Parent->IsDitheredLODTransition() : false;
@@ -2728,6 +2733,7 @@ void UMaterialInstance::UpdateOverridableBaseProperties()
 		TwoSided = 0;
 		DitheredLODTransition = 0;
 		bIsShadingModelFromMaterialExpression = 0;
+		bOutputTranslucentVelocity = false;
 		return;
 	}
 
@@ -2751,6 +2757,16 @@ void UMaterialInstance::UpdateOverridableBaseProperties()
 		BasePropertyOverrides.bCastDynamicShadowAsMasked = bCastDynamicShadowAsMasked;
 	}
 
+	if(BasePropertyOverrides.bOverride_OutputTranslucentVelocity)
+	{
+		bOutputTranslucentVelocity = BasePropertyOverrides.bOutputTranslucentVelocity;
+	}
+	else
+	{
+		bOutputTranslucentVelocity = Parent->IsTranslucencyWritingVelocity();
+		BasePropertyOverrides.bOutputTranslucentVelocity = bOutputTranslucentVelocity;
+	}
+
 	if (BasePropertyOverrides.bOverride_BlendMode)
 	{
 		BlendMode = BasePropertyOverrides.BlendMode;
@@ -4510,6 +4526,15 @@ void UMaterialInstance::GetBasePropertyOverridesHash(FSHAHash& OutHash)const
 		bHasOverrides = true;
 	}
 
+	bool bUsedIsTranslucencyWritingVelocity = IsTranslucencyWritingVelocity();
+	if (bUsedIsTranslucencyWritingVelocity != Mat->IsTranslucencyWritingVelocity())
+	{
+		const FString HashString = TEXT("bOverride_OutputTranslucentVelocity");
+		Hash.UpdateWithString(*HashString, HashString.Len());
+		Hash.Update((uint8*)&bUsedIsTranslucencyWritingVelocity, sizeof(bUsedIsTranslucencyWritingVelocity));
+		bHasOverrides = true;
+	}
+
 	if (bHasOverrides)
 	{
 		Hash.Final();
@@ -4526,7 +4551,8 @@ bool UMaterialInstance::HasOverridenBaseProperties()const
 		(GetShadingModels() != Parent->GetShadingModels()) ||
 		(IsTwoSided() != Parent->IsTwoSided()) ||
 		(IsDitheredLODTransition() != Parent->IsDitheredLODTransition()) ||
-		(GetCastDynamicShadowAsMasked() != Parent->GetCastDynamicShadowAsMasked())
+		(GetCastDynamicShadowAsMasked() != Parent->GetCastDynamicShadowAsMasked()) ||
+		(IsTranslucencyWritingVelocity() != Parent->IsTranslucencyWritingVelocity())
 		))
 	{
 		return true;
@@ -4560,6 +4586,11 @@ bool UMaterialInstance::IsTwoSided() const
 	return TwoSided;
 }
 
+bool UMaterialInstance::IsTranslucencyWritingVelocity() const
+{
+	return bOutputTranslucentVelocity && IsTranslucentBlendMode(GetBlendMode());
+}
+
 bool UMaterialInstance::IsDitheredLODTransition() const
 {
 	return DitheredLODTransition;
diff --git a/Engine/Source/Runtime/Engine/Private/Materials/MaterialShared.cpp b/Engine/Source/Runtime/Engine/Private/Materials/MaterialShared.cpp
--- a/Engine/Source/Runtime/Engine/Private/Materials/MaterialShared.cpp
+++ b/Engine/Source/Runtime/Engine/Private/Materials/MaterialShared.cpp
@@ -1537,7 +1537,7 @@ bool FMaterialResource::IsTranslucencyWritingCustomDepth() const
 
 bool FMaterialResource::IsTranslucencyWritingVelocity() const
 {
-	return Material->IsTranslucencyWritingVelocity();
+	return MaterialInstance ? MaterialInstance->IsTranslucencyWritingVelocity() : Material->IsTranslucencyWritingVelocity();
 }
 
 bool FMaterialResource::IsMasked() const 
@@ -4146,9 +4146,11 @@ FMaterialInstanceBasePropertyOverrides::FMaterialInstanceBasePropertyOverrides()
 	,bOverride_DitheredLODTransition(false)
 	,bOverride_CastDynamicShadowAsMasked(false)
 	,bOverride_TwoSided(false)
+	,bOverride_OutputTranslucentVelocity(false)
 	,TwoSided(0)
 	,DitheredLODTransition(0)
 	,bCastDynamicShadowAsMasked(false)
+	,bOutputTranslucentVelocity(false)
 	,BlendMode(BLEND_Opaque)
 	,ShadingModel(MSM_DefaultLit)
 	, OpacityMaskClipValue(.333333f)
@@ -4159,15 +4161,17 @@ FMaterialInstanceBasePropertyOverrides::FMaterialInstanceBasePropertyOverrides()
 bool FMaterialInstanceBasePropertyOverrides::operator==(const FMaterialInstanceBasePropertyOverrides& Other)const
 {
 	return	bOverride_OpacityMaskClipValue == Other.bOverride_OpacityMaskClipValue &&
-			bOverride_BlendMode == Other.bOverride_BlendMode &&
-			bOverride_ShadingModel == Other.bOverride_ShadingModel &&
-			bOverride_TwoSided == Other.bOverride_TwoSided &&
-			bOverride_DitheredLODTransition == Other.bOverride_DitheredLODTransition &&
-			OpacityMaskClipValue == Other.OpacityMaskClipValue &&
-			BlendMode == Other.BlendMode &&
-			ShadingModel == Other.ShadingModel &&
-			TwoSided == Other.TwoSided &&
-			DitheredLODTransition == Other.DitheredLODTransition;
+		bOverride_BlendMode == Other.bOverride_BlendMode &&
+		bOverride_ShadingModel == Other.bOverride_ShadingModel &&
+		bOverride_TwoSided == Other.bOverride_TwoSided &&
+		bOverride_DitheredLODTransition == Other.bOverride_DitheredLODTransition &&
+		bOverride_OutputTranslucentVelocity == Other.bOverride_OutputTranslucentVelocity &&
+		OpacityMaskClipValue == Other.OpacityMaskClipValue &&
+		BlendMode == Other.BlendMode &&
+		ShadingModel == Other.ShadingModel &&
+		TwoSided == Other.TwoSided &&
+		DitheredLODTransition == Other.DitheredLODTransition &&
+		bCastDynamicShadowAsMasked == Other.bCastDynamicShadowAsMasked;
 }
 
 bool FMaterialInstanceBasePropertyOverrides::operator!=(const FMaterialInstanceBasePropertyOverrides& Other)const
