Improvements to material editor responsiveness with stats open.

- Count SM finalization time into the time budget to prevent super long ticks.
- Make material stats panel update no more frequently than once in 5 seconds.

diff --git a/Engine/Source/Editor/MaterialEditor/Private/MaterialStats.cpp b/Engine/Source/Editor/MaterialEditor/Private/MaterialStats.cpp
--- a/Engine/Source/Editor/MaterialEditor/Private/MaterialStats.cpp
+++ b/Engine/Source/Editor/MaterialEditor/Private/MaterialStats.cpp
@@ -179,6 +179,9 @@ bool FShaderPlatformSettings::CheckShaders()
 {
 	bool bRetValue = false;
 
+	// prevent stats 
+	const double kMinimumTimeBetweenCompilationsSeconds = 5.0;
+
 	if (Material != nullptr)
 	{
 		// check and triggers shader recompilation if needed
@@ -186,7 +189,8 @@ bool FShaderPlatformSettings::CheckShaders()
 		{
 			auto &Data = PlatformData[QualityLevelIndex];
 			const bool bNeedsShaders = (bPresentInGrid && Data.bExtractStats) || Data.bExtractCode;
-			if (Data.bNeedShaderRecompilation && bNeedsShaders)
+			const double CurrentTime = FPlatformTime::Seconds();
+			if (Data.bNeedShaderRecompilation && bNeedsShaders && (CurrentTime - Data.LastTimeCompilationRequested) > kMinimumTimeBetweenCompilationsSeconds)
 			{
 				Data.MaterialResourcesStats->CancelCompilation();
 
@@ -197,6 +201,7 @@ bool FShaderPlatformSettings::CheckShaders()
 					MaterialInstance->UpdateCachedLayerParameters();
 				}
 
+				Data.LastTimeCompilationRequested = CurrentTime;
 				Data.MaterialResourcesStats->CacheShaders(PlatformShaderID);
 
 				Data.bCompilingShaders = true;
diff --git a/Engine/Source/Editor/MaterialEditor/Private/MaterialStats.h b/Engine/Source/Editor/MaterialEditor/Private/MaterialStats.h
--- a/Engine/Source/Editor/MaterialEditor/Private/MaterialStats.h
+++ b/Engine/Source/Editor/MaterialEditor/Private/MaterialStats.h
@@ -119,6 +119,9 @@ public:
 		TWeakPtr<class SDockTab> CodeViewerTab;
 
 		FShaderStatsInfo ShaderStatsInfo;
+
+		/** Time when the previous compilation was requested.*/
+		double LastTimeCompilationRequested = 0.0;
 	};
 	/////////////////////
 
diff --git a/Engine/Source/Runtime/Engine/Private/ShaderCompiler/ShaderCompiler.cpp b/Engine/Source/Runtime/Engine/Private/ShaderCompiler/ShaderCompiler.cpp
--- a/Engine/Source/Runtime/Engine/Private/ShaderCompiler/ShaderCompiler.cpp
+++ b/Engine/Source/Runtime/Engine/Private/ShaderCompiler/ShaderCompiler.cpp
@@ -3760,7 +3760,9 @@ void FShaderCompilingManager::ProcessCompiledShaderMaps(
 
 				// Make a clone of the compiling shader map to use for rendering
 				// This will allow rendering to proceed with the clone, while async compilation continues to potentially update the compiling shader map
+				double StartTime = FPlatformTime::Seconds();
 				ShaderMapToUseForRendering = CompilingShaderMap->AcquireFinalizedClone();
+				TimeBudget -= (FPlatformTime::Seconds() - StartTime);
 			}
 
 			if (!bSuccess || ResultArray.Num() == 0)
