Extend UMaterialEditingLibrary to allow settings values of material layer parameters

diff --git a/Engine/Source/Editor/MaterialEditor/Private/MaterialEditingLibrary.cpp b/Engine/Source/Editor/MaterialEditor/Private/MaterialEditingLibrary.cpp
--- a/Engine/Source/Editor/MaterialEditor/Private/MaterialEditingLibrary.cpp
+++ b/Engine/Source/Editor/MaterialEditor/Private/MaterialEditingLibrary.cpp
@@ -990,76 +990,76 @@ void UMaterialEditingLibrary::ClearAllMaterialInstanceParameters(UMaterialInstan
 }
 
 
-float UMaterialEditingLibrary::GetMaterialInstanceScalarParameterValue(UMaterialInstanceConstant* Instance, FName ParameterName)
+float UMaterialEditingLibrary::GetMaterialInstanceScalarParameterValue(UMaterialInstanceConstant* Instance, FName ParameterName, EMaterialParameterAssociation Association)
 {
 	float Result = 0.f;
 	if (Instance)
 	{
-		Instance->GetScalarParameterValue(ParameterName, Result);
+		Instance->GetScalarParameterValue(FHashedMaterialParameterInfo(ParameterName, Association, Association == EMaterialParameterAssociation::LayerParameter ? 0 : INDEX_NONE), Result);
 	}
 	return Result;
 }
 
-bool UMaterialEditingLibrary::SetMaterialInstanceScalarParameterValue(UMaterialInstanceConstant* Instance, FName ParameterName, float Value)
+bool UMaterialEditingLibrary::SetMaterialInstanceScalarParameterValue(UMaterialInstanceConstant* Instance, FName ParameterName, float Value, EMaterialParameterAssociation Association)
 {
 	bool bResult = false;
 	if (Instance)
 	{
-		Instance->SetScalarParameterValueEditorOnly(ParameterName, Value);
+		Instance->SetScalarParameterValueEditorOnly(FMaterialParameterInfo(ParameterName, Association, Association == EMaterialParameterAssociation::LayerParameter ? 0 : INDEX_NONE), Value);
 	}
 	return bResult;
 }
 
 
-UTexture* UMaterialEditingLibrary::GetMaterialInstanceTextureParameterValue(UMaterialInstanceConstant* Instance, FName ParameterName)
+UTexture* UMaterialEditingLibrary::GetMaterialInstanceTextureParameterValue(UMaterialInstanceConstant* Instance, FName ParameterName, EMaterialParameterAssociation Association)
 {
 	UTexture* Result = nullptr;
 	if (Instance)
 	{
-		Instance->GetTextureParameterValue(ParameterName, Result);
+		Instance->GetTextureParameterValue(FHashedMaterialParameterInfo(ParameterName, Association, Association == EMaterialParameterAssociation::LayerParameter ? 0 : INDEX_NONE), Result);
 	}
 	return Result;
 }
 
-bool UMaterialEditingLibrary::SetMaterialInstanceTextureParameterValue(UMaterialInstanceConstant* Instance, FName ParameterName, UTexture* Value)
+bool UMaterialEditingLibrary::SetMaterialInstanceTextureParameterValue(UMaterialInstanceConstant* Instance, FName ParameterName, UTexture* Value, EMaterialParameterAssociation Association)
 {
 	bool bResult = false;
 	if (Instance)
 	{
-		Instance->SetTextureParameterValueEditorOnly(ParameterName, Value);
+		Instance->SetTextureParameterValueEditorOnly(FMaterialParameterInfo(ParameterName, Association, Association == EMaterialParameterAssociation::LayerParameter ? 0 : INDEX_NONE), Value);
 	}
 	return bResult;
 }
 
 
-FLinearColor UMaterialEditingLibrary::GetMaterialInstanceVectorParameterValue(UMaterialInstanceConstant* Instance, FName ParameterName)
+FLinearColor UMaterialEditingLibrary::GetMaterialInstanceVectorParameterValue(UMaterialInstanceConstant* Instance, FName ParameterName, EMaterialParameterAssociation Association)
 {
 	FLinearColor Result = FLinearColor::Black;
 	if (Instance)
 	{
-		Instance->GetVectorParameterValue(ParameterName, Result);
+		Instance->GetVectorParameterValue(FHashedMaterialParameterInfo(ParameterName, Association, Association == EMaterialParameterAssociation::LayerParameter ? 0 : INDEX_NONE), Result);
 	}
 	return Result;
 }
 
-bool UMaterialEditingLibrary::SetMaterialInstanceVectorParameterValue(UMaterialInstanceConstant* Instance, FName ParameterName, FLinearColor Value)
+bool UMaterialEditingLibrary::SetMaterialInstanceVectorParameterValue(UMaterialInstanceConstant* Instance, FName ParameterName, FLinearColor Value, EMaterialParameterAssociation Association)
 {
 	bool bResult = false;
 	if (Instance)
 	{
-		Instance->SetVectorParameterValueEditorOnly(ParameterName, Value);
+		Instance->SetVectorParameterValueEditorOnly(FMaterialParameterInfo(ParameterName, Association, Association == EMaterialParameterAssociation::LayerParameter ? 0 : INDEX_NONE), Value);
 	}
 	return bResult;
 }
 
 
-bool UMaterialEditingLibrary::GetMaterialInstanceStaticSwitchParameterValue(UMaterialInstanceConstant* Instance, FName ParameterName)
+bool UMaterialEditingLibrary::GetMaterialInstanceStaticSwitchParameterValue(UMaterialInstanceConstant* Instance, FName ParameterName, EMaterialParameterAssociation Association)
 {
 	bool bResult = false;
 	if (Instance)
 	{
 		FGuid OutGuid;
-		Instance->GetStaticSwitchParameterValue(ParameterName, bResult, OutGuid);
+		Instance->GetStaticSwitchParameterValue(FHashedMaterialParameterInfo(ParameterName, Association, Association == EMaterialParameterAssociation::LayerParameter ? 0 : INDEX_NONE), bResult, OutGuid);
 	}
 	return bResult;
 }
diff --git a/Engine/Source/Editor/MaterialEditor/Public/MaterialEditingLibrary.h b/Engine/Source/Editor/MaterialEditor/Public/MaterialEditingLibrary.h
--- a/Engine/Source/Editor/MaterialEditor/Public/MaterialEditingLibrary.h
+++ b/Engine/Source/Editor/MaterialEditor/Public/MaterialEditingLibrary.h
@@ -250,33 +250,33 @@ public:
 
 	/** Get the current scalar (float) parameter value from a Material Instance */
 	UFUNCTION(BlueprintPure, Category = "MaterialEditing")
-	static float GetMaterialInstanceScalarParameterValue(UMaterialInstanceConstant* Instance, FName ParameterName);
+	static float GetMaterialInstanceScalarParameterValue(UMaterialInstanceConstant* Instance, FName ParameterName, EMaterialParameterAssociation Association = EMaterialParameterAssociation::GlobalParameter);
 
 	/** Set the scalar (float) parameter value for a Material Instance */
 	UFUNCTION(BlueprintCallable, Category = "MaterialEditing")
-	static bool SetMaterialInstanceScalarParameterValue(UMaterialInstanceConstant* Instance, FName ParameterName, float Value);
+	static bool SetMaterialInstanceScalarParameterValue(UMaterialInstanceConstant* Instance, FName ParameterName, float Value, EMaterialParameterAssociation Association = EMaterialParameterAssociation::GlobalParameter);
 
 
 	/** Get the current texture parameter value from a Material Instance */
 	UFUNCTION(BlueprintPure, Category = "MaterialEditing")
-	static UTexture* GetMaterialInstanceTextureParameterValue(UMaterialInstanceConstant* Instance, FName ParameterName);
+	static UTexture* GetMaterialInstanceTextureParameterValue(UMaterialInstanceConstant* Instance, FName ParameterName, EMaterialParameterAssociation Association = EMaterialParameterAssociation::GlobalParameter);
 
 	/** Set the texture parameter value for a Material Instance */
 	UFUNCTION(BlueprintCallable, Category = "MaterialEditing")
-	static bool SetMaterialInstanceTextureParameterValue(UMaterialInstanceConstant* Instance, FName ParameterName, UTexture* Value);
+	static bool SetMaterialInstanceTextureParameterValue(UMaterialInstanceConstant* Instance, FName ParameterName, UTexture* Value, EMaterialParameterAssociation Association = EMaterialParameterAssociation::GlobalParameter);
 
 
 	/** Get the current vector parameter value from a Material Instance */
 	UFUNCTION(BlueprintPure, Category = "MaterialEditing")
-	static FLinearColor GetMaterialInstanceVectorParameterValue(UMaterialInstanceConstant* Instance, FName ParameterName);
+	static FLinearColor GetMaterialInstanceVectorParameterValue(UMaterialInstanceConstant* Instance, FName ParameterName, EMaterialParameterAssociation Association = EMaterialParameterAssociation::GlobalParameter);
 
 	/** Set the vector parameter value for a Material Instance */
 	UFUNCTION(BlueprintCallable, Category = "MaterialEditing")
-	static bool SetMaterialInstanceVectorParameterValue(UMaterialInstanceConstant* Instance, FName ParameterName, FLinearColor Value);
+	static bool SetMaterialInstanceVectorParameterValue(UMaterialInstanceConstant* Instance, FName ParameterName, FLinearColor Value, EMaterialParameterAssociation Association = EMaterialParameterAssociation::GlobalParameter);
 
 	/** Get the current static switch parameter value from a Material Instance */
 	UFUNCTION(BlueprintPure, Category = "MaterialEditing")
-	static bool GetMaterialInstanceStaticSwitchParameterValue(UMaterialInstanceConstant* Instance, FName ParameterName);
+	static bool GetMaterialInstanceStaticSwitchParameterValue(UMaterialInstanceConstant* Instance, FName ParameterName, EMaterialParameterAssociation Association = EMaterialParameterAssociation::GlobalParameter);
 
 	/** Called after making modifications to a Material Instance to recompile shaders etc. */
 	UFUNCTION(BlueprintCallable, Category = "MaterialEditing")
