UMaterialEditingLibrary support for setting static switches

diff --git a/Engine/Source/Editor/MaterialEditor/Private/MaterialEditingLibrary.cpp b/Engine/Source/Editor/MaterialEditor/Private/MaterialEditingLibrary.cpp
--- a/Engine/Source/Editor/MaterialEditor/Private/MaterialEditingLibrary.cpp
+++ b/Engine/Source/Editor/MaterialEditor/Private/MaterialEditingLibrary.cpp
@@ -1064,6 +1064,16 @@ bool UMaterialEditingLibrary::GetMaterialInstanceStaticSwitchParameterValue(UMat
 	return bResult;
 }
 
+bool UMaterialEditingLibrary::SetMaterialInstanceStaticSwitchParameterValue(UMaterialInstanceConstant* Instance, FName ParameterName, bool Value, EMaterialParameterAssociation Association)
+{
+	bool bResult = false;
+	if (Instance)
+	{
+		Instance->SetStaticSwitchParameterValueEditorOnly(FMaterialParameterInfo(ParameterName, Association, Association == EMaterialParameterAssociation::LayerParameter ? 0 : INDEX_NONE), Value);
+	}
+	return bResult;
+}
+
 void UMaterialEditingLibrary::UpdateMaterialInstance(UMaterialInstanceConstant* Instance)
 {
 	if (Instance)
diff --git a/Engine/Source/Editor/MaterialEditor/Public/MaterialEditingLibrary.h b/Engine/Source/Editor/MaterialEditor/Public/MaterialEditingLibrary.h
--- a/Engine/Source/Editor/MaterialEditor/Public/MaterialEditingLibrary.h
+++ b/Engine/Source/Editor/MaterialEditor/Public/MaterialEditingLibrary.h
@@ -278,6 +278,10 @@ public:
 	UFUNCTION(BlueprintPure, Category = "MaterialEditing")
 	static bool GetMaterialInstanceStaticSwitchParameterValue(UMaterialInstanceConstant* Instance, FName ParameterName, EMaterialParameterAssociation Association = EMaterialParameterAssociation::GlobalParameter);
 
+	/** Set the static switch parameter value for a Material Instance */
+	UFUNCTION(BlueprintPure, Category = "MaterialEditing")
+	static bool SetMaterialInstanceStaticSwitchParameterValue(UMaterialInstanceConstant* Instance, FName ParameterName, bool Value, EMaterialParameterAssociation Association = EMaterialParameterAssociation::GlobalParameter);
+
 	/** Called after making modifications to a Material Instance to recompile shaders etc. */
 	UFUNCTION(BlueprintCallable, Category = "MaterialEditing")
 	static void UpdateMaterialInstance(UMaterialInstanceConstant* Instance);
diff --git a/Engine/Source/Runtime/Engine/Classes/Materials/MaterialInstance.h b/Engine/Source/Runtime/Engine/Classes/Materials/MaterialInstance.h
--- a/Engine/Source/Runtime/Engine/Classes/Materials/MaterialInstance.h
+++ b/Engine/Source/Runtime/Engine/Classes/Materials/MaterialInstance.h
@@ -629,6 +629,9 @@ public:
 	void GetAllShaderMaps(TArray<FMaterialShaderMap*>& OutShaderMaps);
 
 #if WITH_EDITORONLY_DATA
+
+	ENGINE_API void SetStaticSwitchParameterValueEditorOnly(const FMaterialParameterInfo& ParameterInfo, bool Value);
+
 	/**
 	 * Builds a composited set of static parameters, including inherited and overridden values
 	 */
diff --git a/Engine/Source/Runtime/Engine/Private/Materials/MaterialInstance.cpp b/Engine/Source/Runtime/Engine/Private/Materials/MaterialInstance.cpp
--- a/Engine/Source/Runtime/Engine/Private/Materials/MaterialInstance.cpp
+++ b/Engine/Source/Runtime/Engine/Private/Materials/MaterialInstance.cpp
@@ -2074,6 +2074,24 @@ UPhysicalMaterial* UMaterialInstance::GetPhysicalMaterialFromMap(int32 Index) co
 }
 
 #if WITH_EDITORONLY_DATA
+
+void UMaterialInstance::SetStaticSwitchParameterValueEditorOnly(const FMaterialParameterInfo& ParameterInfo, bool Value)
+{
+	check(GIsEditor || IsRunningCommandlet());
+
+	for (FStaticSwitchParameter& StaticSwitches : StaticParameters.StaticSwitchParameters)
+	{
+		if (StaticSwitches.ParameterInfo == ParameterInfo)
+		{
+			StaticSwitches.bOverride = true;
+			StaticSwitches.Value = Value;
+			return;
+		}
+	}
+
+	new(StaticParameters.StaticSwitchParameters) FStaticSwitchParameter(ParameterInfo, Value, true, FGuid());
+}
+
 void UMaterialInstance::GetStaticParameterValues(FStaticParameterSet& OutStaticParameters)
 {
 	check(IsInGameThread());
