UMaterialEditingLibrary::SetMaterialInstanceStaticSwitchParameterValue should be BlueprintCallable and not BlueprintPure

diff --git a/Engine/Source/Editor/MaterialEditor/Public/MaterialEditingLibrary.h b/Engine/Source/Editor/MaterialEditor/Public/MaterialEditingLibrary.h
--- a/Engine/Source/Editor/MaterialEditor/Public/MaterialEditingLibrary.h
+++ b/Engine/Source/Editor/MaterialEditor/Public/MaterialEditingLibrary.h
@@ -279,7 +279,7 @@ public:
 	static bool GetMaterialInstanceStaticSwitchParameterValue(UMaterialInstanceConstant* Instance, FName ParameterName, EMaterialParameterAssociation Association = EMaterialParameterAssociation::GlobalParameter);
 
 	/** Set the static switch parameter value for a Material Instance */
-	UFUNCTION(BlueprintPure, Category = "MaterialEditing")
+	UFUNCTION(BlueprintCallable, Category = "MaterialEditing")
 	static bool SetMaterialInstanceStaticSwitchParameterValue(UMaterialInstanceConstant* Instance, FName ParameterName, bool Value, EMaterialParameterAssociation Association = EMaterialParameterAssociation::GlobalParameter);
 
 	/** Called after making modifications to a Material Instance to recompile shaders etc. */
