Work around in UMaterialEditingLibrary::SetMaterialInstanceStaticSwitchParameterValue when MICs are generated programmatically and have never been opened by the editor. MaterialLayersParameters not being on the MIC can cause static switches to get wiped when it is eventually opened.

diff --git a/Engine/Source/Editor/MaterialEditor/Private/MaterialEditingLibrary.cpp b/Engine/Source/Editor/MaterialEditor/Private/MaterialEditingLibrary.cpp
--- a/Engine/Source/Editor/MaterialEditor/Private/MaterialEditingLibrary.cpp
+++ b/Engine/Source/Editor/MaterialEditor/Private/MaterialEditingLibrary.cpp
@@ -1070,6 +1070,11 @@ bool UMaterialEditingLibrary::SetMaterialInstanceStaticSwitchParameterValue(UMat
 	if (Instance)
 	{
 		Instance->SetStaticSwitchParameterValueEditorOnly(FMaterialParameterInfo(ParameterName, Association, Association == EMaterialParameterAssociation::LayerParameter ? 0 : INDEX_NONE), Value);
+
+		// The material instance editor window puts MaterialLayersParameters into our StaticParameters, if we don't do this, our settings could get wiped out on first launch of the material editor.
+		// If there's ever a cleaner and more isolated way of populating MaterialLayersParameters, we should do that instead.
+		UMaterialEditorInstanceConstant* MaterialEditorInstance = NewObject<UMaterialEditorInstanceConstant>(GetTransientPackage(), NAME_None, RF_Transactional);
+		MaterialEditorInstance->SetSourceInstance(Instance);
 	}
 	return bResult;
 }
