Avoid crash when opening material function instance, don't look for expressions connected to inputs, since nothing is connected in this context

diff --git a/Engine/Source/Editor/MaterialEditor/Private/MaterialEditorUtilities.cpp b/Engine/Source/Editor/MaterialEditor/Private/MaterialEditorUtilities.cpp
--- a/Engine/Source/Editor/MaterialEditor/Private/MaterialEditorUtilities.cpp
+++ b/Engine/Source/Editor/MaterialEditor/Private/MaterialEditorUtilities.cpp
@@ -279,20 +279,21 @@ bool FMaterialEditorUtilities::GetStaticSwitchExpressionValue(UMaterialInstance*
 	if(FunctionInputExpression && FunctionInputExpression->InputType == FunctionInput_StaticBool)
 	{
 		FGetVisibleMaterialParametersFunctionState* TopmostFunctionState = FunctionStack.Pop();
-		check(TopmostFunctionState->FunctionCall);
-		const TArray<FFunctionExpressionInput>* FunctionInputs = &TopmostFunctionState->FunctionCall->FunctionInputs;
-
-		// Get the FFunctionExpressionInput which stores information about the input node from the parent that this is linked to.
-		const FFunctionExpressionInput* MatchingInput = FindInputById(FunctionInputExpression, *FunctionInputs);
-		if (MatchingInput && (MatchingInput->Input.Expression || !FunctionInputExpression->bUsePreviewValueAsDefault))
-		{
-			GetStaticSwitchExpressionValue(MaterialInstance, MatchingInput->Input.Expression, bOutValue, OutExpressionID, FunctionStack);
-		}
-		else
+		if (TopmostFunctionState->FunctionCall)
 		{
-			GetStaticSwitchExpressionValue(MaterialInstance, FunctionInputExpression->Preview.Expression, bOutValue, OutExpressionID, FunctionStack);
-		}
+			const TArray<FFunctionExpressionInput>* FunctionInputs = &TopmostFunctionState->FunctionCall->FunctionInputs;
 
+			// Get the FFunctionExpressionInput which stores information about the input node from the parent that this is linked to.
+			const FFunctionExpressionInput* MatchingInput = FindInputById(FunctionInputExpression, *FunctionInputs);
+			if (MatchingInput && (MatchingInput->Input.Expression || !FunctionInputExpression->bUsePreviewValueAsDefault))
+			{
+				GetStaticSwitchExpressionValue(MaterialInstance, MatchingInput->Input.Expression, bOutValue, OutExpressionID, FunctionStack);
+			}
+			else
+			{
+				GetStaticSwitchExpressionValue(MaterialInstance, FunctionInputExpression->Preview.Expression, bOutValue, OutExpressionID, FunctionStack);
+			}
+		}
 		FunctionStack.Push(TopmostFunctionState);
 	}
 
