Modifying a material function will continue to recompile all materials that depend on that function, but the materials will no longer be marked dirty
This is fine, because materials will refresh the stateID of any dependent MF when loaded, and use updated stateIDs to form the ShaderMapID

diff --git a/Engine/Source/Editor/MaterialEditor/Private/MaterialEditingLibrary.cpp b/Engine/Source/Editor/MaterialEditor/Private/MaterialEditingLibrary.cpp
--- a/Engine/Source/Editor/MaterialEditor/Private/MaterialEditingLibrary.cpp
+++ b/Engine/Source/Editor/MaterialEditor/Private/MaterialEditingLibrary.cpp
@@ -862,16 +862,15 @@ void UMaterialEditingLibrary::UpdateMaterialFunction(UMaterialFunctionInterface*
 {
 	if (MaterialFunction)
 	{
-		// mark the function as changed
-		MaterialFunction->PreEditChange(nullptr);
-		MaterialFunction->PostEditChange();
-		MaterialFunction->MarkPackageDirty();
-
 		// Create a material update context so we can safely update materials using this function.
 		{
 			FMaterialUpdateContext UpdateContext;
 
-			// Go through all function instances in memory and update them if they are children
+			// mark the function as changed
+			MaterialFunction->ForceRecompileForRendering(UpdateContext, PreviewMaterial);
+			MaterialFunction->MarkPackageDirty();
+
+			// Go through all function instances in memory and recompile them if they are children
 			for (TObjectIterator<UMaterialFunctionInstance> It; It; ++It)
 			{
 				UMaterialFunctionInstance* FunctionInstance = *It;
@@ -881,80 +880,19 @@ void UMaterialEditingLibrary::UpdateMaterialFunction(UMaterialFunctionInterface*
 				if (Functions.Contains(MaterialFunction))
 				{
 					FunctionInstance->UpdateParameterSet();
-					FunctionInstance->PreEditChange(nullptr);
-					FunctionInstance->PostEditChange();
-					FunctionInstance->MarkPackageDirty();
-				}
-			}
-
-			// Go through all materials in memory and recompile them if they use this material function
-			for (TObjectIterator<UMaterial> It; It; ++It)
-			{
-				UMaterial* CurrentMaterial = *It;
-				if (CurrentMaterial != PreviewMaterial)
-				{
-					bool bRecompile = false;
-
-					// Preview materials often use expressions for rendering that are not in their Expressions array, 
-					// And therefore their MaterialFunctionInfos are not up to date.
-					// However we don't want to trigger this if the Material is a preview material itself. This can now be the case with thumbnail preview materials for material functions.
-					if (CurrentMaterial->bIsPreviewMaterial && (PreviewMaterial != nullptr) && !PreviewMaterial->bIsPreviewMaterial)
-					{
-						bRecompile = true;
-					}
-					else
-					{
-						TArray<UMaterialFunctionInterface*> Functions;
-						CurrentMaterial->GetDependentFunctions(Functions);
-						if (Functions.Contains(MaterialFunction))
-						{
-							bRecompile = true;
-						}
-					}
+					FunctionInstance->ForceRecompileForRendering(UpdateContext, PreviewMaterial);
 
-					if (bRecompile)
-					{
-						UpdateContext.AddMaterial(CurrentMaterial);
-
-						// Propagate the function change to this material
-						CurrentMaterial->PreEditChange(nullptr);
-						CurrentMaterial->PostEditChange();
-						CurrentMaterial->MarkPackageDirty();
-
-						if (CurrentMaterial->MaterialGraph)
-						{
-							CurrentMaterial->MaterialGraph->RebuildGraph();
-						}
-
-						// if this instance was opened in an editor notify the change
-						if (IMaterialEditor* MaterialEditor = MaterialEditingLibraryImpl::FindMaterialEditorForAsset(CurrentMaterial))
-						{
-							MaterialEditor->NotifyExternalMaterialChange();
-						}
-					}
+					// ForceRecompileForRendering will update StateId, so need to mark the package as dirty
+					FunctionInstance->MarkPackageDirty();
 				}
 			}
 
-			// Go through all material instances in memory and recompile them if they use this material function
-			for (TObjectIterator<UMaterialInstance> It; It; ++It)
+			// Notify material editor for any materials that we are updating
+			for (UMaterialInterface* CurrentMaterial : UpdateContext.GetUpdatedMaterials())
 			{
-				UMaterialInstance* CurrentInstance = *It;
-				if (CurrentInstance->GetBaseMaterial())
+				if (IMaterialEditor* MaterialEditor = MaterialEditingLibraryImpl::FindMaterialEditorForAsset(CurrentMaterial))
 				{
-					TArray<UMaterialFunctionInterface*> Functions;
-					CurrentInstance->GetDependentFunctions(Functions);
-					if (Functions.Contains(MaterialFunction))
-					{
-						UpdateContext.AddMaterialInstance(CurrentInstance);
-						CurrentInstance->PreEditChange(nullptr);
-						CurrentInstance->PostEditChange();
-
-						// if this instance was opened in an editor notify the change
-						if (IMaterialEditor* MaterialEditor = MaterialEditingLibraryImpl::FindMaterialEditorForAsset(CurrentInstance))
-						{
-							MaterialEditor->NotifyExternalMaterialChange();
-						}
-					}
+					MaterialEditor->NotifyExternalMaterialChange();
 				}
 			}
 		}
diff --git a/Engine/Source/Runtime/Engine/Classes/Materials/MaterialFunction.h b/Engine/Source/Runtime/Engine/Classes/Materials/MaterialFunction.h
--- a/Engine/Source/Runtime/Engine/Classes/Materials/MaterialFunction.h
+++ b/Engine/Source/Runtime/Engine/Classes/Materials/MaterialFunction.h
@@ -104,6 +104,8 @@ public:
 
 	/** Get the inputs and outputs that this function exposes, for a function call expression to use. */
 	virtual void GetInputsAndOutputs(TArray<struct FFunctionExpressionInput>& OutInputs, TArray<struct FFunctionExpressionOutput>& OutOutputs) const override;
+
+	virtual void ForceRecompileForRendering(FMaterialUpdateContext& UpdateContext, UMaterial* InPreviewMaterial) override;
 #endif
 
 	virtual bool ValidateFunctionUsage(class FMaterialCompiler* Compiler, const FFunctionExpressionOutput& Output) override;
diff --git a/Engine/Source/Runtime/Engine/Classes/Materials/MaterialFunctionInterface.h b/Engine/Source/Runtime/Engine/Classes/Materials/MaterialFunctionInterface.h
--- a/Engine/Source/Runtime/Engine/Classes/Materials/MaterialFunctionInterface.h
+++ b/Engine/Source/Runtime/Engine/Classes/Materials/MaterialFunctionInterface.h
@@ -74,6 +74,8 @@ public:
 
 	virtual void GetInputsAndOutputs(TArray<struct FFunctionExpressionInput>& OutInputs, TArray<struct FFunctionExpressionOutput>& OutOutputs) const
 		PURE_VIRTUAL(UMaterialFunctionInterface::GetInputsAndOutputs,);
+
+	virtual void ForceRecompileForRendering(FMaterialUpdateContext& UpdateContext, UMaterial* InPreviewMaterial);
 #endif
 
 	virtual bool ValidateFunctionUsage(class FMaterialCompiler* Compiler, const FFunctionExpressionOutput& Output)
diff --git a/Engine/Source/Runtime/Engine/Private/Materials/Material.cpp b/Engine/Source/Runtime/Engine/Private/Materials/Material.cpp
--- a/Engine/Source/Runtime/Engine/Private/Materials/Material.cpp
+++ b/Engine/Source/Runtime/Engine/Private/Materials/Material.cpp
@@ -6425,6 +6425,7 @@ void UMaterial::NotifyCompilationFinished(UMaterialInterface* Material)
 
 void UMaterial::ForceRecompileForRendering()
 {
+	UpdateCachedExpressionData();
 	CacheResourceShadersForRendering( false );
 }
 
diff --git a/Engine/Source/Runtime/Engine/Private/Materials/MaterialExpressions.cpp b/Engine/Source/Runtime/Engine/Private/Materials/MaterialExpressions.cpp
--- a/Engine/Source/Runtime/Engine/Private/Materials/MaterialExpressions.cpp
+++ b/Engine/Source/Runtime/Engine/Private/Materials/MaterialExpressions.cpp
@@ -12330,41 +12330,64 @@ void UMaterialFunctionInterface::GetAssetRegistryTags(TArray<FAssetRegistryTag>&
 
 #if WITH_EDITOR
 void UMaterialFunctionInterface::PostEditChangeProperty(FPropertyChangedEvent& PropertyChangedEvent)
+{
+	{
+		FMaterialUpdateContext UpdateContext;
+		ForceRecompileForRendering(UpdateContext, nullptr);
+	}
+	Super::PostEditChangeProperty(PropertyChangedEvent);
+}
+
+void UMaterialFunctionInterface::ForceRecompileForRendering(FMaterialUpdateContext& UpdateContext, UMaterial* InPreviewMaterial)
 {
 	//@todo - recreate guid only when needed, not when a comment changes
 	StateId = FGuid::NewGuid();
 
 	// Go through all materials in memory and recompile them if they use this function
-	FMaterialUpdateContext UpdateContext;
 	for (TObjectIterator<UMaterialInterface> It; It; ++It)
 	{
-		UMaterialInterface* CurrentMaterial = *It;
+		UMaterialInterface* CurrentMaterialInterface = *It;
+		if (CurrentMaterialInterface == InPreviewMaterial)
+		{
+			continue;
+		}
 
 		bool bRecompile = false;
-		UMaterialFunctionInterface* Self = this;
-		CurrentMaterial->IterateDependentFunctions(
-			[Self, &bRecompile](UMaterialFunctionInterface* InFunction)
+
+		// Preview materials often use expressions for rendering that are not in their Expressions array, 
+		// And therefore their MaterialFunctionInfos are not up to date.
+		// However we don't want to trigger this if the Material is a preview material itself. This can now be the case with thumbnail preview materials for material functions.
+		if (InPreviewMaterial && !InPreviewMaterial->bIsPreviewMaterial)
+		{
+			UMaterial* CurrentMaterial = Cast<UMaterial>(CurrentMaterialInterface);
+			if (CurrentMaterial && CurrentMaterial->bIsPreviewMaterial)
 			{
-				if (InFunction == Self)
+				bRecompile = true;
+			}
+		}
+
+		if (!bRecompile)
+		{
+			UMaterialFunctionInterface* Self = this;
+			CurrentMaterialInterface->IterateDependentFunctions(
+				[Self, &bRecompile](UMaterialFunctionInterface* InFunction)
 				{
-					bRecompile = true;
-					return false;
-				}
-				return true;
-			});
+					if (InFunction == Self)
+					{
+						bRecompile = true;
+						return false;
+					}
+					return true;
+				});
+		}
 
 		if (bRecompile)
 		{
-			UpdateContext.AddMaterialInterface(CurrentMaterial);
-
 			// Propagate the change to this material
-			CurrentMaterial->PreEditChange(nullptr);
-			CurrentMaterial->PostEditChange();
-			CurrentMaterial->MarkPackageDirty();
+			UpdateContext.AddMaterialInterface(CurrentMaterialInterface);
+			CurrentMaterialInterface->ForceRecompileForRendering();
 		}
 	}
-
-	Super::PostEditChangeProperty(PropertyChangedEvent);
 }
 #endif // WITH_EDITOR
 
@@ -12479,14 +12502,20 @@ void UMaterialFunction::UpdateDependentFunctionCandidates()
 #endif
 
 void UMaterialFunction::PostEditChangeProperty(FPropertyChangedEvent& PropertyChangedEvent)
+{
+	Super::PostEditChangeProperty(PropertyChangedEvent);
+}
+
+void UMaterialFunction::ForceRecompileForRendering(FMaterialUpdateContext& UpdateContext, UMaterial* InPreviewMaterial)
 {
 #if WITH_EDITORONLY_DATA
 	UpdateInputOutputTypes();
 	UpdateDependentFunctionCandidates();
 #endif
 
-	Super::PostEditChangeProperty(PropertyChangedEvent);
+	Super::ForceRecompileForRendering(UpdateContext, InPreviewMaterial);
 }
+
 #endif // WITH_EDITOR
 
 void UMaterialFunction::Serialize(FArchive& Ar)
diff --git a/Engine/Source/Runtime/Engine/Private/Materials/MaterialInstance.cpp b/Engine/Source/Runtime/Engine/Private/Materials/MaterialInstance.cpp
--- a/Engine/Source/Runtime/Engine/Private/Materials/MaterialInstance.cpp
+++ b/Engine/Source/Runtime/Engine/Private/Materials/MaterialInstance.cpp
@@ -2723,6 +2723,7 @@ bool UMaterialInstance::GetGroupName(const FHashedMaterialParameterInfo& Paramet
 #if WITH_EDITOR
 void UMaterialInstance::ForceRecompileForRendering()
 {
+	UpdateCachedLayerParameters();
 	CacheResourceShadersForRendering();
 }
 #endif // WITH_EDITOR
diff --git a/Engine/Source/Runtime/Engine/Public/MaterialShared.h b/Engine/Source/Runtime/Engine/Public/MaterialShared.h
--- a/Engine/Source/Runtime/Engine/Public/MaterialShared.h
+++ b/Engine/Source/Runtime/Engine/Public/MaterialShared.h
@@ -2773,6 +2773,8 @@ public:
 
 	/** Adds a material interface that has been updated to the context. */
 	ENGINE_API void AddMaterialInterface(UMaterialInterface* Instance);
+
+	inline const TSet<UMaterialInterface*>& GetUpdatedMaterials() const { return UpdatedMaterialInterfaces; }
 };
 
 /**
