Fix pin index assignments when collapsing material nodes

diff --git a/Engine/Source/Editor/MaterialEditor/Private/MaterialEditor.cpp b/Engine/Source/Editor/MaterialEditor/Private/MaterialEditor.cpp
--- a/Engine/Source/Editor/MaterialEditor/Private/MaterialEditor.cpp
+++ b/Engine/Source/Editor/MaterialEditor/Private/MaterialEditor.cpp
@@ -900,6 +900,43 @@ void FMaterialEditor::UpdatePreviewViewportsVisibility()
 	}
 }
 
+static void AssignPinSourceIndices(UEdGraphNode* Node)
+{
+	int32 NumInputDataPins = 0;
+	int32 NumOutputDataPins = 0;
+	int32 NumInputExecPins = 0;
+	int32 NumOutputExecPins = 0;
+
+	for (UEdGraphPin* Pin : Node->Pins)
+	{
+		int32 SourceIndex = INDEX_NONE;
+		if (Pin->PinType.PinCategory == UMaterialGraphSchema::PC_Exec)
+		{
+			switch (Pin->Direction)
+			{
+			case EGPD_Input: SourceIndex = NumInputExecPins++; break;
+			case EGPD_Output: SourceIndex = NumOutputExecPins++; break;
+			default: checkNoEntry(); break;
+			}
+		}
+		else
+		{
+			switch (Pin->Direction)
+			{
+			case EGPD_Input: SourceIndex = NumInputDataPins++; break;
+			case EGPD_Output: SourceIndex = NumOutputDataPins++; break;
+			default: checkNoEntry(); break;
+			}
+		}
+
+		if (Pin->SourceIndex == INDEX_NONE)
+		{
+			Pin->SourceIndex = SourceIndex;
+		}
+		ensure(Pin->SourceIndex == SourceIndex);
+	}
+}
+
 void FMaterialEditor::CollapseNodesIntoGraph(UEdGraphNode* InGatewayNode, UMaterialGraphNode* InEntryNode, UMaterialGraphNode* InResultNode, UEdGraph* InSourceGraph, UEdGraph* InDestinationGraph, TSet<UEdGraphNode*>& InCollapsableNodes)
 {
 	const UMaterialGraphSchema* MaterialSchema = GetDefault<UMaterialGraphSchema>(); // ?
@@ -1071,6 +1108,10 @@ void FMaterialEditor::CollapseNodesIntoGraph(UEdGraphNode* InGatewayNode, UMater
 		InResultNode->MaterialExpression->MaterialExpressionEditorX = InResultNode->NodePosX;
 		InResultNode->MaterialExpression->MaterialExpressionEditorY = InResultNode->NodePosY;
 	}
+
+	AssignPinSourceIndices(InGatewayNode);
+	AssignPinSourceIndices(InEntryNode);
+	AssignPinSourceIndices(InResultNode);
 }
 
 void FMaterialEditor::CollapseNodes(TSet<UEdGraphNode*>& InCollapsableNodes)
@@ -1101,6 +1142,7 @@ void FMaterialEditor::CollapseNodes(TSet<UEdGraphNode*>& InCollapsableNodes)
 		SourceGraph, 
 		DestinationGraph, 
 		InCollapsableNodes);
+	AssignPinSourceIndices(GatewayNode);
 
 	UpdateMaterialAfterGraphChange();
 
diff --git a/Engine/Source/Editor/UnrealEd/Private/MaterialGraph.cpp b/Engine/Source/Editor/UnrealEd/Private/MaterialGraph.cpp
--- a/Engine/Source/Editor/UnrealEd/Private/MaterialGraph.cpp
+++ b/Engine/Source/Editor/UnrealEd/Private/MaterialGraph.cpp
@@ -419,6 +419,7 @@ void UMaterialGraph::LinkMaterialExpressionsFromGraph() const
 					{
 						UEdGraphPin* ConnectedPin = Pin->LinkedTo[0];
 						UMaterialGraphNode* ConnectedNode = CastChecked<UMaterialGraphNode>(ConnectedPin->GetOwningNode());
+						check(ConnectedPin->SourceIndex != INDEX_NONE);
 
 						if (!ConnectedNode->MaterialExpression->IsExpressionConnected(&MaterialInput, ConnectedPin->SourceIndex))
 						{
