Refactored CPD view to skip duplicate entries.

diff --git a/Engine/Source/Editor/MaterialEditor/Private/SMaterialEditorCustomPrimitiveDataWidget.cpp b/Engine/Source/Editor/MaterialEditor/Private/SMaterialEditorCustomPrimitiveDataWidget.cpp
--- a/Engine/Source/Editor/MaterialEditor/Private/SMaterialEditorCustomPrimitiveDataWidget.cpp
+++ b/Engine/Source/Editor/MaterialEditor/Private/SMaterialEditorCustomPrimitiveDataWidget.cpp
@@ -58,7 +58,7 @@ private:
 
 void SMaterialCustomPrimitiveDataPanel::Refresh()
 {
-	Items.Empty();
+	TMap<int32, TSet<FString>> IndexBuckets;
 
 	if (MaterialEditorInstance && MaterialEditorInstance->PreviewMaterial)
 	{
@@ -69,13 +69,8 @@ void SMaterialCustomPrimitiveDataPanel::Refresh()
 		{
 			if (Expr->bUseCustomPrimitiveData)
 			{
-				FString FunctionName;
-				if (Expr->GraphNode)
-				{
-					Expr->GraphNode->GetGraph()->GetName(FunctionName);
-				}
-
-				Items.Add(MakeShareable(new FCustomPrimitiveDataRowData((int32)Expr->PrimitiveDataIndex, Expr->GetParameterName().ToString(), FunctionName)));
+				TSet<FString>& Params = IndexBuckets.FindOrAdd((int32)Expr->PrimitiveDataIndex);
+				Params.Add(Expr->GetParameterName().ToString());
 			}
 		}
 
@@ -86,51 +81,36 @@ void SMaterialCustomPrimitiveDataPanel::Refresh()
 		{
 			if (Expr->bUseCustomPrimitiveData)
 			{
-				Items.Add(MakeShareable(new FCustomPrimitiveDataRowData((int32)Expr->PrimitiveDataIndex + 0, Expr->GetParameterName().ToString() + ".r")));
-				Items.Add(MakeShareable(new FCustomPrimitiveDataRowData((int32)Expr->PrimitiveDataIndex + 1, Expr->GetParameterName().ToString() + ".g")));
-				Items.Add(MakeShareable(new FCustomPrimitiveDataRowData((int32)Expr->PrimitiveDataIndex + 2, Expr->GetParameterName().ToString() + ".b")));
-				Items.Add(MakeShareable(new FCustomPrimitiveDataRowData((int32)Expr->PrimitiveDataIndex + 3, Expr->GetParameterName().ToString() + ".a")));
+				TArray<FString> Elements({".r", ".g" , ".b",  ".a" });
+				for (int32 i = 0; i < 4; i++)
+				{
+					TSet<FString>& Params = IndexBuckets.FindOrAdd((int32)Expr->PrimitiveDataIndex + i);
+					Params.Add(Expr->GetParameterName().ToString() + Elements[i]);
+				}
 			}
 		}
 
-		// Sort the items before construction the list
-		Items.StableSort([](const TSharedPtr<FCustomPrimitiveDataRowData>& A, const TSharedPtr<FCustomPrimitiveDataRowData>& B) {
-			int32 SlotA = A->Slot;
-			int32 SlotB = B->Slot;
-			if (SlotA == SlotB)
-			{
-				return A->Name < B->Name;
-			}
-			return SlotA < SlotB;
-		});
-
-		TSet<int32> DuplicateIndices;
+		Items.Empty();
 
-		// Once sorted, find suplicate slots
-		for (int32 i = 1; i < Items.Num(); i++)
+		// Go through and mark slots as duplicate while copying over to the shared ptr array
+		for (auto& Params : IndexBuckets)
 		{
-			TSharedPtr<FCustomPrimitiveDataRowData>& PrevItem = Items[i - 1];
-			TSharedPtr<FCustomPrimitiveDataRowData>& CurrItem = Items[i];
-			const int32 PrevSlot = PrevItem->Slot;
-			const int32 CurrSlot = CurrItem->Slot;
-			const FString& PrevName = PrevItem->Name;
-			const FString& CurrName = CurrItem->Name;
-
-			// Duplicate found, store the slot for later processing
-			if ((CurrSlot == PrevSlot) && (CurrName != PrevName))
+			for (FString& ParamName : Params.Value)
 			{
-				DuplicateIndices.Add(CurrItem->Slot);
+				Items.Add(MakeShareable(new FCustomPrimitiveDataRowData(Params.Key, ParamName, Params.Value.Num() > 1)));
 			}
 		}
 
-		// Go through and mark slots as duplicate
-		for (int32 i = 0; i < Items.Num(); i++)
-		{
-			if (DuplicateIndices.Find(Items[i]->Slot))
+		// Sort the items
+		Items.StableSort([](const TSharedPtr<FCustomPrimitiveDataRowData>& A, const TSharedPtr<FCustomPrimitiveDataRowData>& B) {
+			int32 SlotA = A->Slot;
+			int32 SlotB = B->Slot;
+			if (SlotA == SlotB)
 			{
-				Items[i]->bIsDuplicate = true;
+				return A->Name < B->Name;
 			}
-		}
+			return SlotA < SlotB;
+			});
 	}
 
 	ListViewWidget->RequestListRefresh();
diff --git a/Engine/Source/Editor/MaterialEditor/Private/SMaterialEditorCustomPrimitiveDataWidget.h b/Engine/Source/Editor/MaterialEditor/Private/SMaterialEditorCustomPrimitiveDataWidget.h
--- a/Engine/Source/Editor/MaterialEditor/Private/SMaterialEditorCustomPrimitiveDataWidget.h
+++ b/Engine/Source/Editor/MaterialEditor/Private/SMaterialEditorCustomPrimitiveDataWidget.h
@@ -19,7 +19,7 @@ class UMaterialEditorPreviewParameters;
 /** Data for a row in the list of custom primitive data entries*/
 struct FCustomPrimitiveDataRowData
 {
-	FCustomPrimitiveDataRowData(int32 InSlot, FString InName, FString InUsage = "") : Slot(InSlot), Name(InName) {}
+	FCustomPrimitiveDataRowData(int32 InSlot, FString InName, bool bInIsDuplicate = false) : Slot(InSlot), Name(InName), bIsDuplicate(bInIsDuplicate) {}
 
 	int32 Slot = 0;
 	FString Name = "";
