Force mobile dynamic pointlight to use static branch.

#p4v-preflight-copy 17870962

diff --git a/Engine/Plugins/Runtime/Oculus/OculusVR/Source/OculusEditor/Private/OculusToolWidget.cpp b/Engine/Plugins/Runtime/Oculus/OculusVR/Source/OculusEditor/Private/OculusToolWidget.cpp
--- a/Engine/Plugins/Runtime/Oculus/OculusVR/Source/OculusEditor/Private/OculusToolWidget.cpp
+++ b/Engine/Plugins/Runtime/Oculus/OculusVR/Source/OculusEditor/Private/OculusToolWidget.cpp
@@ -923,9 +923,7 @@ FReply SOculusToolWidget::MobileNumDynamicPointLightsFix(bool text)
 {
 	URendererSettings* Settings = GetMutableDefault<URendererSettings>();
 	Settings->MobileNumDynamicPointLights = 0;
-	Settings->bMobileDynamicPointLightsUseStaticBranch = true;
 	Settings->UpdateSinglePropertyInConfigFile(Settings->GetClass()->FindPropertyByName(GET_MEMBER_NAME_CHECKED(URendererSettings, MobileNumDynamicPointLights)), Settings->GetDefaultConfigFilename());
-	Settings->UpdateSinglePropertyInConfigFile(Settings->GetClass()->FindPropertyByName(GET_MEMBER_NAME_CHECKED(URendererSettings, bMobileDynamicPointLightsUseStaticBranch)), Settings->GetDefaultConfigFilename());
 	SuggestRestart();
 	return FReply::Handled();
 }
@@ -933,7 +931,7 @@ FReply SOculusToolWidget::MobileNumDynamicPointLightsFix(bool text)
 EVisibility SOculusToolWidget::MobileNumDynamicPointLightsVisibility(FName tag) const
 {
 	URendererSettings* Settings = GetMutableDefault<URendererSettings>();
-	if (Settings->MobileNumDynamicPointLights == 0 && Settings->bMobileDynamicPointLightsUseStaticBranch)
+	if (Settings->MobileNumDynamicPointLights == 0)
 	{
 		return EVisibility::Collapsed;
 	}
diff --git a/Engine/Source/Editor/MaterialEditor/Private/MaterialStatsCommon.cpp b/Engine/Source/Editor/MaterialEditor/Private/MaterialStatsCommon.cpp
--- a/Engine/Source/Editor/MaterialEditor/Private/MaterialStatsCommon.cpp
+++ b/Engine/Source/Editor/MaterialEditor/Private/MaterialStatsCommon.cpp
@@ -240,7 +240,6 @@ FString FMaterialStatsUtils::RepresentativeShaderTypeToString(const ERepresentat
 			return TEXT("Stationary surface + CSM");
 		break;
 
-		case ERepresentativeShader::StationarySurface1PointLight:
 		case ERepresentativeShader::StationarySurfaceNPointLights:
 			return TEXT("Stationary surface + Point Lights");
 		break;
@@ -432,11 +431,8 @@ void FMaterialStatsUtils::GetRepresentativeShaderTypesAndDescriptions(TMap<FName
 
 				static auto* CVarPointLights = IConsoleManager::Get().FindTConsoleVariableDataInt(TEXT("r.MobileNumDynamicPointLights"));
 				const bool bPointLights = CVarPointLights->GetValueOnAnyThread() > 0;
-
-				static auto* CVarPointLightsStaticBranch = IConsoleManager::Get().FindTConsoleVariableDataInt(TEXT("r.MobileDynamicPointLightsUseStaticBranch"));
-				const bool bPointLightsStaticBranch = CVarPointLightsStaticBranch->GetValueOnAnyThread() != 0;
 				
-				const int32 NumPointLights = bPointLightsStaticBranch ? CVarPointLights->GetValueOnAnyThread() : 1;
+				const int32 NumPointLights = CVarPointLights->GetValueOnAnyThread();
 				
 				if (bAllowDistanceFieldShadows)// distance field shadows
 				{
@@ -476,13 +472,6 @@ void FMaterialStatsUtils::GetRepresentativeShaderTypesAndDescriptions(TMap<FName
 
 						if (bPointLights) // add point lights shaders + distance field shadows
 						{
-							static const FName Name_HDRLinear64_OneLight = bOnlySkyPermutation ? 
-								TEXT("TMobileBasePassPSFMobileDistanceFieldShadowsAndLQLightMapPolicy1HDRLinear64Skylight") :
-								TEXT("TMobileBasePassPSFMobileDistanceFieldShadowsAndLQLightMapPolicy1HDRLinear64");
-							static const FName Name_LDRGamma32_OneLight = bOnlySkyPermutation ? 
-								TEXT("TMobileBasePassPSFMobileDistanceFieldShadowsAndLQLightMapPolicy1LDRGamma32Skylight") : 
-								TEXT("TMobileBasePassPSFMobileDistanceFieldShadowsAndLQLightMapPolicy1LDRGamma32");
-
 							static const FName Name_HDRLinear64_NLights = bOnlySkyPermutation ? 
 								TEXT("TMobileBasePassPSFMobileDistanceFieldShadowsAndLQLightMapPolicyINT32_MAXHDRLinear64Skylight") : 
 								TEXT("TMobileBasePassPSFMobileDistanceFieldShadowsAndLQLightMapPolicyINT32_MAXHDRLinear64");
@@ -490,15 +479,11 @@ void FMaterialStatsUtils::GetRepresentativeShaderTypesAndDescriptions(TMap<FName
 								TEXT("TMobileBasePassPSFMobileDistanceFieldShadowsAndLQLightMapPolicyINT32_MAXLDRGamma32Skylight") :
 								TEXT("TMobileBasePassPSFMobileDistanceFieldShadowsAndLQLightMapPolicyINT32_MAXLDRGamma32");
 
-							static const FName ShaderName = bMobileHDR ?
-								(bPointLightsStaticBranch ? Name_HDRLinear64_NLights : Name_HDRLinear64_OneLight) :
-								(bPointLightsStaticBranch ? Name_LDRGamma32_NLights : Name_LDRGamma32_OneLight);
+							static const FName ShaderName = bMobileHDR ? Name_HDRLinear64_NLights : Name_LDRGamma32_NLights;
 
-							const FString Description = FString::Printf(TEXT("Mobile base pass shader with distance field shadows, CSM and %s point light(s) %s"), bPointLightsStaticBranch ? TEXT("N") : TEXT("1"), DescSuffix);
+							const FString Description = FString::Printf(TEXT("Mobile base pass shader with distance field shadows, CSM and %s point light(N) %s"), DescSuffix);
 
-							FRepresentativeShaderInfo ShaderInfo = bPointLightsStaticBranch ?
-								FRepresentativeShaderInfo(ERepresentativeShader::StationarySurfaceNPointLights, ShaderName, Description) :
-								FRepresentativeShaderInfo(ERepresentativeShader::StationarySurface1PointLight, ShaderName, Description);
+							FRepresentativeShaderInfo ShaderInfo = FRepresentativeShaderInfo(ERepresentativeShader::StationarySurfaceNPointLights, ShaderName, Description);
 
 							ShaderTypeNamesAndDescriptions.FindOrAdd(FLocalVertexFactoryName).Add(ShaderInfo);
 								
@@ -521,13 +506,6 @@ void FMaterialStatsUtils::GetRepresentativeShaderTypesAndDescriptions(TMap<FName
 
 					if (bPointLights) // add point lights + lightmap
 					{
-						static const FName Name_HDRLinear64_OneLight = bOnlySkyPermutation ? 
-							TEXT("TMobileBasePassPSTLightMapPolicyLQ1HDRLinear64Skylight") : 
-							TEXT("TMobileBasePassPSTLightMapPolicyLQ1HDRLinear64");
-						static const FName Name_LDRGamma32_OneLight = bOnlySkyPermutation ? 
-							TEXT("TMobileBasePassPSTLightMapPolicyLQ1LDRGamma32Skylight") : 
-							TEXT("TMobileBasePassPSTLightMapPolicyLQ1LDRGamma32");
-
 						static const FName Name_HDRLinear64_NLights = bOnlySkyPermutation ? 
 							TEXT("TMobileBasePassPSTLightMapPolicyLQINT32_MAXHDRLinear64Skylight") :
 							TEXT("TMobileBasePassPSTLightMapPolicyLQINT32_MAXHDRLinear64");
@@ -535,15 +513,11 @@ void FMaterialStatsUtils::GetRepresentativeShaderTypesAndDescriptions(TMap<FName
 							TEXT("TMobileBasePassPSTLightMapPolicyLQINT32_MAXLDRGamma32Skylight") :
 							TEXT("TMobileBasePassPSTLightMapPolicyLQINT32_MAXLDRGamma32");
 
-						static const FName ShaderName = bMobileHDR ?
-							(bPointLightsStaticBranch ? Name_HDRLinear64_NLights : Name_HDRLinear64_OneLight) :
-							(bPointLightsStaticBranch ? Name_LDRGamma32_NLights : Name_LDRGamma32_OneLight);
+						static const FName ShaderName = bMobileHDR ? Name_HDRLinear64_NLights : Name_LDRGamma32_NLights;
 
-						const FString Description = FString::Printf(TEXT("Mobile base pass shader with static lighting and %s point light(s) %s"), bPointLightsStaticBranch ? TEXT("N") : TEXT("1"), DescSuffix);
+						const FString Description = FString::Printf(TEXT("Mobile base pass shader with static lighting and %s point light(N) %s"), DescSuffix);
 
-						FRepresentativeShaderInfo ShaderInfo = bPointLightsStaticBranch ?
-							FRepresentativeShaderInfo(ERepresentativeShader::StationarySurfaceNPointLights, ShaderName, Description) :
-							FRepresentativeShaderInfo(ERepresentativeShader::StationarySurface1PointLight, ShaderName, Description);
+						FRepresentativeShaderInfo ShaderInfo = FRepresentativeShaderInfo(ERepresentativeShader::StationarySurfaceNPointLights, ShaderName, Description);
 
 						ShaderTypeNamesAndDescriptions.FindOrAdd(FLocalVertexFactoryName).Add(ShaderInfo);
 					}
diff --git a/Engine/Source/Editor/MaterialEditor/Private/MaterialStatsGrid.cpp b/Engine/Source/Editor/MaterialEditor/Private/MaterialStatsGrid.cpp
--- a/Engine/Source/Editor/MaterialEditor/Private/MaterialStatsGrid.cpp
+++ b/Engine/Source/Editor/MaterialEditor/Private/MaterialStatsGrid.cpp
@@ -163,7 +163,6 @@ FStatsGridRow_Shaders::EShaderClass FStatsGridRow_Shaders::GetShaderClass(const
 	{
 		case ERepresentativeShader::StationarySurface:
 		case ERepresentativeShader::StationarySurfaceCSM:
-		case ERepresentativeShader::StationarySurface1PointLight:
 		case ERepresentativeShader::StationarySurfaceNPointLights:
 		case ERepresentativeShader::DynamicallyLitObject:
 		case ERepresentativeShader::UIDefaultFragmentShader:
diff --git a/Engine/Source/Editor/MaterialEditor/Public/MaterialStatsCommon.h b/Engine/Source/Editor/MaterialEditor/Public/MaterialStatsCommon.h
--- a/Engine/Source/Editor/MaterialEditor/Public/MaterialStatsCommon.h
+++ b/Engine/Source/Editor/MaterialEditor/Public/MaterialStatsCommon.h
@@ -44,7 +44,6 @@ enum class ERepresentativeShader
 	FirstFragmentShader,
 	StationarySurface = FirstFragmentShader,
 	StationarySurfaceCSM,
-	StationarySurface1PointLight,
 	StationarySurfaceNPointLights,
 	DynamicallyLitObject,
 	UIDefaultFragmentShader,
diff --git a/Engine/Source/Runtime/Core/Private/HAL/ConsoleManager.cpp b/Engine/Source/Runtime/Core/Private/HAL/ConsoleManager.cpp
--- a/Engine/Source/Runtime/Core/Private/HAL/ConsoleManager.cpp
+++ b/Engine/Source/Runtime/Core/Private/HAL/ConsoleManager.cpp
@@ -2114,13 +2114,6 @@ static TAutoConsoleVariable<int32> CVarMobileNumDynamicPointLights(
 	TEXT("The number of dynamic point lights to support on mobile devices. Setting this to 0 for games which do not require dynamic point lights will reduce the number of shaders generated."), 
 	ECVF_RenderThreadSafe | ECVF_ReadOnly);
 
-static TAutoConsoleVariable<int32> CVarMobileDynamicPointLightsUseStaticBranch(
-	TEXT("r.MobileDynamicPointLightsUseStaticBranch"),
-	1,
-	TEXT("0: Generate unique forward rendering base pass shaders for 0, 1, ... N mobile dynamic point lights. (faster but generates many more shaders)\n")
-	TEXT("1: Use a shared shader with static branching for rendering 1 or more dynamic point lights (slightly slower but reduces shaders generated, recommended for most games)."),
-	ECVF_RenderThreadSafe | ECVF_ReadOnly);
-
 static TAutoConsoleVariable<int32> CVarMobileEnableStaticAndCSMShadowReceivers(
 	TEXT("r.Mobile.EnableStaticAndCSMShadowReceivers"),
 	1,
diff --git a/Engine/Source/Runtime/Engine/Classes/Engine/RendererSettings.h b/Engine/Source/Runtime/Engine/Classes/Engine/RendererSettings.h
--- a/Engine/Source/Runtime/Engine/Classes/Engine/RendererSettings.h
+++ b/Engine/Source/Runtime/Engine/Classes/Engine/RendererSettings.h
@@ -851,12 +851,6 @@ class ENGINE_API URendererSettings : public UDeveloperSettings
 		ConfigRestartRequired = true))
 		uint32 MobileNumDynamicPointLights;
 
-	UPROPERTY(config, EditAnywhere, Category = MobileShaderPermutationReduction, meta = (
-		ConsoleVariable = "r.MobileDynamicPointLightsUseStaticBranch", DisplayName = "Use Shared Movable Spotlight / Point Light Shaders",
-		ToolTip = "If this setting is enabled, the same shader will be used for any number of dynamic spotlights or point lights (up to the maximum specified above) hitting a surface. This is slightly slower but reduces the number of shaders generated. Changing this setting requires restarting the editor.",
-		ConfigRestartRequired = true))
-		uint32 bMobileDynamicPointLightsUseStaticBranch : 1;
-
 	UPROPERTY(config, EditAnywhere, Category = MobileShaderPermutationReduction, meta = (
 		ConsoleVariable = "r.Mobile.EnableMovableSpotlights",
 		DisplayName = "Support Movable Spotlights",
diff --git a/Engine/Source/Runtime/Engine/Private/SceneManagement.cpp b/Engine/Source/Runtime/Engine/Private/SceneManagement.cpp
--- a/Engine/Source/Runtime/Engine/Private/SceneManagement.cpp
+++ b/Engine/Source/Runtime/Engine/Private/SceneManagement.cpp
@@ -1109,7 +1109,6 @@ void FReadOnlyCVARCache::Init()
 	static const auto CVarMobileEnableStaticAndCSMShadowReceivers = IConsoleManager::Get().FindTConsoleVariableDataInt(TEXT("r.Mobile.EnableStaticAndCSMShadowReceivers"));
 	static const auto CVarMobileAllowDistanceFieldShadows = IConsoleManager::Get().FindTConsoleVariableDataInt(TEXT("r.Mobile.AllowDistanceFieldShadows"));
 	static const auto CVarMobileNumDynamicPointLights = IConsoleManager::Get().FindTConsoleVariableDataInt(TEXT("r.MobileNumDynamicPointLights"));
-	static const auto CVarMobileDynamicPointLightsUseStaticBranch = IConsoleManager::Get().FindTConsoleVariableDataInt(TEXT("r.MobileDynamicPointLightsUseStaticBranch"));
 	static const auto CVarMobileSkyLightPermutation = IConsoleManager::Get().FindTConsoleVariableDataInt(TEXT("r.Mobile.SkyLightPermutation"));
 	static const auto CVarMobileEnableMovableSpotLights = IConsoleManager::Get().FindTConsoleVariableDataInt(TEXT("r.Mobile.EnableMovableSpotlights"));
 	static const auto CVarMobileEnableMovableSpotLightsShadow = IConsoleManager::Get().FindTConsoleVariableDataInt(TEXT("r.Mobile.EnableMovableSpotlightsShadow"));
@@ -1127,7 +1126,6 @@ void FReadOnlyCVARCache::Init()
 	bMobileAllowDistanceFieldShadows = CVarMobileAllowDistanceFieldShadows->GetValueOnAnyThread() != 0;
 	bMobileEnableStaticAndCSMShadowReceivers = CVarMobileEnableStaticAndCSMShadowReceivers->GetValueOnAnyThread() != 0;
 	NumMobileMovablePointLights = CVarMobileNumDynamicPointLights->GetValueOnAnyThread();
-	bMobileMovablePointLightsUseStaticBranch = CVarMobileDynamicPointLightsUseStaticBranch->GetValueOnAnyThread() != 0;
 	MobileSkyLightPermutation = CVarMobileSkyLightPermutation->GetValueOnAnyThread();
 	bMobileEnableMovableSpotlights = CVarMobileEnableMovableSpotLights->GetValueOnAnyThread() != 0;
 	bMobileEnableMovableSpotlightsShadow = bMobileEnableMovableSpotlights && CVarMobileEnableMovableSpotLightsShadow->GetValueOnAnyThread() != 0;
diff --git a/Engine/Source/Runtime/Engine/Public/SceneManagement.h b/Engine/Source/Runtime/Engine/Public/SceneManagement.h
--- a/Engine/Source/Runtime/Engine/Public/SceneManagement.h
+++ b/Engine/Source/Runtime/Engine/Public/SceneManagement.h
@@ -3219,7 +3219,6 @@ struct FReadOnlyCVARCache
 	bool bMobileEnableStaticAndCSMShadowReceivers;
 	int32 NumMobileMovablePointLights;
 	int32 MobileSkyLightPermutation;
-	bool bMobileMovablePointLightsUseStaticBranch;
 	bool bMobileEnableMovableSpotlights;
 	bool bMobileEnableMovableSpotlightsShadow;
 	
diff --git a/Engine/Source/Runtime/Renderer/Private/MobileBasePass.cpp b/Engine/Source/Runtime/Renderer/Private/MobileBasePass.cpp
--- a/Engine/Source/Runtime/Renderer/Private/MobileBasePass.cpp
+++ b/Engine/Source/Runtime/Renderer/Private/MobileBasePass.cpp
@@ -129,42 +129,6 @@ bool MobileBasePass::GetShaders(
 			VertexShader,
 			PixelShader
 			);
-	case 1:
-		return GetMobileBasePassShaders<1>(
-			LightMapPolicyType, 
-			MaterialResource, 
-			VertexFactoryType, 
-			bEnableSkyLight, 
-			VertexShader,
-			PixelShader
-			);
-	case 2:
-		return GetMobileBasePassShaders<2>(
-			LightMapPolicyType,
-			MaterialResource,
-			VertexFactoryType,
-			bEnableSkyLight,
-			VertexShader,
-			PixelShader
-			);
-	case 3:
-		return GetMobileBasePassShaders<3>(
-			LightMapPolicyType,
-			MaterialResource,
-			VertexFactoryType,
-			bEnableSkyLight,
-			VertexShader,
-			PixelShader
-			);
-	case 4:
-		return GetMobileBasePassShaders<4>(
-			LightMapPolicyType,
-			MaterialResource,
-			VertexFactoryType,
-			bEnableSkyLight,
-			VertexShader,
-			PixelShader
-			);
 	case 0:
 	default:
 		return GetMobileBasePassShaders<0>(
@@ -202,7 +166,7 @@ int32 MobileBasePass::CalcNumMovablePointLights(const FMaterial& InMaterial, con
 	const FReadOnlyCVARCache& ReadOnlyCVARCache = FReadOnlyCVARCache::Get();
 	const bool bIsUnlit = InMaterial.GetShadingModels().IsUnlit();
 	int32 OutNumMovablePointLights = (InPrimitiveSceneProxy && !bIsUnlit) ? FMath::Min<int32>(InPrimitiveSceneProxy->GetPrimitiveSceneInfo()->NumMobileMovablePointLights, ReadOnlyCVARCache.NumMobileMovablePointLights) : 0;
-	if (OutNumMovablePointLights > 0 && ReadOnlyCVARCache.bMobileMovablePointLightsUseStaticBranch)
+	if (OutNumMovablePointLights > 0)
 	{
 		OutNumMovablePointLights = INT32_MAX;
 	}
diff --git a/Engine/Source/Runtime/Renderer/Private/MobileBasePassRendering.cpp b/Engine/Source/Runtime/Renderer/Private/MobileBasePassRendering.cpp
--- a/Engine/Source/Runtime/Renderer/Private/MobileBasePassRendering.cpp
+++ b/Engine/Source/Runtime/Renderer/Private/MobileBasePassRendering.cpp
@@ -70,7 +70,6 @@ static TAutoConsoleVariable<int32> CVarMobileUseHWsRGBEncoding(
 
 namespace MobileBasePass
 {
-	FShaderPlatformCachedIniValue<bool> MobileDynamicPointLightsUseStaticBranchIniValue(TEXT("/Script/Engine.RendererSettings"), TEXT("r.MobileDynamicPointLightsUseStaticBranch"));
 	FShaderPlatformCachedIniValue<int32> MobileNumDynamicPointLightsIniValue(TEXT("/Script/Engine.RendererSettings"), TEXT("r.MobileNumDynamicPointLights"));
 };
 
@@ -106,10 +105,6 @@ static_assert(MAX_BASEPASS_DYNAMIC_POINT_LIGHTS == 4, "If you change MAX_BASEPAS
 #define IMPLEMENT_MOBILE_SHADING_BASEPASS_LIGHTMAPPED_SHADER_TYPE(LightMapPolicyType,LightMapPolicyName) \
 	IMPLEMENT_MOBILE_SHADING_BASEPASS_LIGHTMAPPED_VERTEX_SHADER_TYPE(LightMapPolicyType, LightMapPolicyName) \
 	IMPLEMENT_MOBILE_SHADING_BASEPASS_LIGHTMAPPED_PIXEL_SHADER_TYPE(LightMapPolicyType, LightMapPolicyName, 0) \
-	IMPLEMENT_MOBILE_SHADING_BASEPASS_LIGHTMAPPED_PIXEL_SHADER_TYPE(LightMapPolicyType, LightMapPolicyName, 1) \
-	IMPLEMENT_MOBILE_SHADING_BASEPASS_LIGHTMAPPED_PIXEL_SHADER_TYPE(LightMapPolicyType, LightMapPolicyName, 2) \
-	IMPLEMENT_MOBILE_SHADING_BASEPASS_LIGHTMAPPED_PIXEL_SHADER_TYPE(LightMapPolicyType, LightMapPolicyName, 3) \
-	IMPLEMENT_MOBILE_SHADING_BASEPASS_LIGHTMAPPED_PIXEL_SHADER_TYPE(LightMapPolicyType, LightMapPolicyName, 4) \
 	IMPLEMENT_MOBILE_SHADING_BASEPASS_LIGHTMAPPED_PIXEL_SHADER_TYPE(LightMapPolicyType, LightMapPolicyName, INT32_MAX)
 
 // Implement shader types per lightmap policy 
diff --git a/Engine/Source/Runtime/Renderer/Private/MobileBasePassRendering.h b/Engine/Source/Runtime/Renderer/Private/MobileBasePassRendering.h
--- a/Engine/Source/Runtime/Renderer/Private/MobileBasePassRendering.h
+++ b/Engine/Source/Runtime/Renderer/Private/MobileBasePassRendering.h
@@ -372,7 +372,6 @@ namespace MobileBasePass
 
 	bool StationarySkyLightHasBeenApplied(const FScene* Scene, ELightMapPolicyType LightMapPolicyType);
 
-	extern FShaderPlatformCachedIniValue<bool> MobileDynamicPointLightsUseStaticBranchIniValue;
 	extern FShaderPlatformCachedIniValue<int32> MobileNumDynamicPointLightsIniValue;
 };
 
@@ -400,7 +399,6 @@ public:
 		// We compile the point light shader combinations based on the project settings
 		static auto* MobileSkyLightPermutationCVar = IConsoleManager::Get().FindTConsoleVariableDataInt(TEXT("r.Mobile.SkyLightPermutation"));
 
-		const bool bMobileDynamicPointLightsUseStaticBranch = MobileBasePass::MobileDynamicPointLightsUseStaticBranchIniValue.Get(Parameters.Platform);
 		const int32 MobileNumDynamicPointLights = MobileBasePass::MobileNumDynamicPointLightsIniValue.Get(Parameters.Platform);
 		const int32 MobileSkyLightPermutationOptions = MobileSkyLightPermutationCVar->GetValueOnAnyThread();
 		const bool bDeferredShading = IsMobileDeferredShadingEnabled(Parameters.Platform);
@@ -424,8 +422,7 @@ public:
 
 		const bool bShouldCacheByNumDynamicPointLights =
 			(NumMovablePointLights == 0 ||
-			(bIsLit && NumMovablePointLights == INT32_MAX && bMobileDynamicPointLightsUseStaticBranch && MobileNumDynamicPointLights > 0) ||	// single shader for variable number of point lights
-				(bIsLit && NumMovablePointLights <= MobileNumDynamicPointLights && !bMobileDynamicPointLightsUseStaticBranch));				// unique 1...N point light shaders
+			(bIsLit && NumMovablePointLights == INT32_MAX && MobileNumDynamicPointLights > 0));		// single shader for variable number of point lights
 
 		return TMobileBasePassPSBaseType<LightMapPolicyType>::ShouldCompilePermutation(Parameters) && 
 				ShouldCacheShaderByPlatformAndOutputFormat(Parameters.Platform, OutputFormat) && 
