Mobile: fix a broken string formating for material stats

diff --git a/Engine/Source/Editor/MaterialEditor/Private/MaterialStatsCommon.cpp b/Engine/Source/Editor/MaterialEditor/Private/MaterialStatsCommon.cpp
--- a/Engine/Source/Editor/MaterialEditor/Private/MaterialStatsCommon.cpp
+++ b/Engine/Source/Editor/MaterialEditor/Private/MaterialStatsCommon.cpp
@@ -481,7 +481,7 @@ void FMaterialStatsUtils::GetRepresentativeShaderTypesAndDescriptions(TMap<FName
 
 							static const FName ShaderName = bMobileHDR ? Name_HDRLinear64_NLights : Name_LDRGamma32_NLights;
 
-							const FString Description = FString::Printf(TEXT("Mobile base pass shader with distance field shadows, CSM and %s point light(N) %s"), DescSuffix);
+							const FString Description = FString::Printf(TEXT("Mobile base pass shader with distance field shadows, CSM and %d point light(s) %s"), NumPointLights, DescSuffix);
 
 							FRepresentativeShaderInfo ShaderInfo = FRepresentativeShaderInfo(ERepresentativeShader::StationarySurfaceNPointLights, ShaderName, Description);
 
@@ -515,7 +515,7 @@ void FMaterialStatsUtils::GetRepresentativeShaderTypesAndDescriptions(TMap<FName
 
 						static const FName ShaderName = bMobileHDR ? Name_HDRLinear64_NLights : Name_LDRGamma32_NLights;
 
-						const FString Description = FString::Printf(TEXT("Mobile base pass shader with static lighting and %s point light(N) %s"), DescSuffix);
+						const FString Description = FString::Printf(TEXT("Mobile base pass shader with static lighting and %d point light(s) %s"), NumPointLights, DescSuffix);
 
 						FRepresentativeShaderInfo ShaderInfo = FRepresentativeShaderInfo(ERepresentativeShader::StationarySurfaceNPointLights, ShaderName, Description);
 
