Deprecated bDisplayResetToDefault and the other checkboxes in IPropertyHandle::CreatePropertyNameWidget().

diff --git a/Engine/Plugins/Runtime/DataRegistry/Source/DataRegistryEditor/Private/DataRegistryIdCustomization.cpp b/Engine/Plugins/Runtime/DataRegistry/Source/DataRegistryEditor/Private/DataRegistryIdCustomization.cpp
--- a/Engine/Plugins/Runtime/DataRegistry/Source/DataRegistryEditor/Private/DataRegistryIdCustomization.cpp
+++ b/Engine/Plugins/Runtime/DataRegistry/Source/DataRegistryEditor/Private/DataRegistryIdCustomization.cpp
@@ -244,7 +244,7 @@ void FDataRegistryIdCustomization::CustomizeHeader(TSharedRef<class IPropertyHan
 	HeaderRow
 	.NameContent()
 	[
-		InStructPropertyHandle->CreatePropertyNameWidget(FText::GetEmpty(), FText::GetEmpty(), false)
+		InStructPropertyHandle->CreatePropertyNameWidget()
 	];
 
 	FDataTableEditorUtils::AddSearchForReferencesContextMenu(HeaderRow, FExecuteAction::CreateSP(this, &FDataRegistryIdCustomization::OnSearchForReferences));
diff --git a/Engine/Plugins/VirtualProduction/DMX/DMXProtocol/Source/DMXProtocolEditor/Private/DetailsCustomizations/DMXInputPortConfigCustomization.cpp b/Engine/Plugins/VirtualProduction/DMX/DMXProtocol/Source/DMXProtocolEditor/Private/DetailsCustomizations/DMXInputPortConfigCustomization.cpp
--- a/Engine/Plugins/VirtualProduction/DMX/DMXProtocol/Source/DMXProtocolEditor/Private/DetailsCustomizations/DMXInputPortConfigCustomization.cpp
+++ b/Engine/Plugins/VirtualProduction/DMX/DMXProtocol/Source/DMXProtocolEditor/Private/DetailsCustomizations/DMXInputPortConfigCustomization.cpp
@@ -33,14 +33,10 @@ void FDMXInputPortConfigCustomization::CustomizeHeader(TSharedRef<IPropertyHandl
 {
 	PropertyUtilities = StructCustomizationUtils.GetPropertyUtilities();
 
-	const bool bDisplayResetToDefault = false;
-	const FText DisplayNameOverride = FText::GetEmpty();
-	const FText DisplayToolTipOverride = FText::GetEmpty();
-
 	HeaderRow
 	.NameContent()
 	[
-		StructPropertyHandle->CreatePropertyNameWidget(DisplayNameOverride, DisplayToolTipOverride, bDisplayResetToDefault)
+		StructPropertyHandle->CreatePropertyNameWidget()
 	];
 }
 
diff --git a/Engine/Plugins/VirtualProduction/DMX/DMXProtocol/Source/DMXProtocolEditor/Private/DetailsCustomizations/DMXOutputPortConfigCustomization.cpp b/Engine/Plugins/VirtualProduction/DMX/DMXProtocol/Source/DMXProtocolEditor/Private/DetailsCustomizations/DMXOutputPortConfigCustomization.cpp
--- a/Engine/Plugins/VirtualProduction/DMX/DMXProtocol/Source/DMXProtocolEditor/Private/DetailsCustomizations/DMXOutputPortConfigCustomization.cpp
+++ b/Engine/Plugins/VirtualProduction/DMX/DMXProtocol/Source/DMXProtocolEditor/Private/DetailsCustomizations/DMXOutputPortConfigCustomization.cpp
@@ -34,14 +34,10 @@ void FDMXOutputPortConfigCustomization::CustomizeHeader(TSharedRef<IPropertyHand
 {
 	PropertyUtilities = StructCustomizationUtils.GetPropertyUtilities();
 
-	const bool bDisplayResetToDefault = false;
-	const FText DisplayNameOverride = FText::GetEmpty();
-	const FText DisplayToolTipOverride = FText::GetEmpty();
-
 	HeaderRow
 	.NameContent()
 	[
-		StructPropertyHandle->CreatePropertyNameWidget(DisplayNameOverride, DisplayToolTipOverride, bDisplayResetToDefault)
+		StructPropertyHandle->CreatePropertyNameWidget()
 	];
 }
 
diff --git a/Engine/Plugins/VirtualProduction/DMX/DMXProtocol/Source/DMXProtocolEditor/Private/DetailsCustomizations/DMXPortReferenceCustomizationBase.cpp b/Engine/Plugins/VirtualProduction/DMX/DMXProtocol/Source/DMXProtocolEditor/Private/DetailsCustomizations/DMXPortReferenceCustomizationBase.cpp
--- a/Engine/Plugins/VirtualProduction/DMX/DMXProtocol/Source/DMXProtocolEditor/Private/DetailsCustomizations/DMXPortReferenceCustomizationBase.cpp
+++ b/Engine/Plugins/VirtualProduction/DMX/DMXProtocol/Source/DMXProtocolEditor/Private/DetailsCustomizations/DMXPortReferenceCustomizationBase.cpp
@@ -20,7 +20,6 @@
 #include "Widgets/Layout/SWrapBox.h"
 #include "Widgets/Text/STextBlock.h" 
 
-
 #define LOCTEXT_NAMESPACE "DMXPortConfigCustomizationBase"
 
 FDMXPortReferenceCustomizationBase::FDMXPortReferenceCustomizationBase()
@@ -28,14 +27,10 @@ FDMXPortReferenceCustomizationBase::FDMXPortReferenceCustomizationBase()
 
 void FDMXPortReferenceCustomizationBase::CustomizeHeader(TSharedRef<IPropertyHandle> StructPropertyHandle, class FDetailWidgetRow& HeaderRow, IPropertyTypeCustomizationUtils& StructCustomizationUtils)
 {
-	const bool bDisplayResetToDefault = false;
-	const FText DisplayNameOverride = FText::GetEmpty();
-	const FText DisplayToolTipOverride = FText::GetEmpty();
-
 	HeaderRow
 	.NameContent()
 	[
-		StructPropertyHandle->CreatePropertyNameWidget(DisplayNameOverride, DisplayToolTipOverride, bDisplayResetToDefault)
+		StructPropertyHandle->CreatePropertyNameWidget()
 	];
 }
 
diff --git a/Engine/Source/Editor/DetailCustomizations/Private/AttenuationSettingsCustomizations.cpp b/Engine/Source/Editor/DetailCustomizations/Private/AttenuationSettingsCustomizations.cpp
--- a/Engine/Source/Editor/DetailCustomizations/Private/AttenuationSettingsCustomizations.cpp
+++ b/Engine/Source/Editor/DetailCustomizations/Private/AttenuationSettingsCustomizations.cpp
@@ -117,15 +117,10 @@ TSharedRef<IPropertyTypeCustomization> FForceFeedbackAttenuationSettingsCustomiz
 
 void FBaseAttenuationSettingsCustomization::CustomizeHeader(TSharedRef<class IPropertyHandle> StructPropertyHandle, class FDetailWidgetRow& HeaderRow, IPropertyTypeCustomizationUtils& StructCustomizationUtils)
 {
-	// We'll set up reset to default ourselves
-	const bool bDisplayResetToDefault = false;
-	const FText DisplayNameOverride = FText::GetEmpty();
-	const FText DisplayToolTipOverride = FText::GetEmpty();
-
 	HeaderRow
 		.NameContent()
 		[
-			StructPropertyHandle->CreatePropertyNameWidget(DisplayNameOverride, DisplayToolTipOverride, bDisplayResetToDefault)
+			StructPropertyHandle->CreatePropertyNameWidget()
 		];
 }
 
diff --git a/Engine/Source/Editor/DetailCustomizations/Private/ColorStructCustomization.cpp b/Engine/Source/Editor/DetailCustomizations/Private/ColorStructCustomization.cpp
--- a/Engine/Source/Editor/DetailCustomizations/Private/ColorStructCustomization.cpp
+++ b/Engine/Source/Editor/DetailCustomizations/Private/ColorStructCustomization.cpp
@@ -46,11 +46,6 @@ void FColorStructCustomization::CustomizeHeader(TSharedRef<class IPropertyHandle
 
 void FColorStructCustomization::MakeHeaderRow(TSharedRef<class IPropertyHandle>& InStructPropertyHandle, FDetailWidgetRow& Row)
 {
-	// We'll set up reset to default ourselves
-	const bool bDisplayResetToDefault = false;
-	const FText DisplayNameOverride = FText::GetEmpty();
-	const FText DisplayToolTipOverride = FText::GetEmpty();
-	
 	TSharedPtr<SWidget> ColorWidget;
 	float ContentWidth = 125.0f;
 
@@ -68,7 +63,7 @@ void FColorStructCustomization::MakeHeaderRow(TSharedRef<class IPropertyHandle>&
 
 	Row.NameContent()
 	[
-		StructPropertyHandle->CreatePropertyNameWidget(DisplayNameOverride, DisplayToolTipOverride, bDisplayResetToDefault)
+		StructPropertyHandle->CreatePropertyNameWidget()
 	]
 	.ValueContent()
 	.MinDesiredWidth(ContentWidth)
diff --git a/Engine/Source/Editor/DetailCustomizations/Private/CurveColorCustomization.cpp b/Engine/Source/Editor/DetailCustomizations/Private/CurveColorCustomization.cpp
--- a/Engine/Source/Editor/DetailCustomizations/Private/CurveColorCustomization.cpp
+++ b/Engine/Source/Editor/DetailCustomizations/Private/CurveColorCustomization.cpp
@@ -86,7 +86,7 @@ void FCurveColorCustomization::CustomizeHeader(TSharedRef<IPropertyHandle> InStr
 		HeaderRow
 			.NameContent()
 			[
-				InStructPropertyHandle->CreatePropertyNameWidget( FText::GetEmpty(), FText::GetEmpty(), false )
+				InStructPropertyHandle->CreatePropertyNameWidget( FText::GetEmpty(), FText::GetEmpty() )
 			]
 			.ValueContent()
 			.HAlign(HAlign_Fill)
@@ -123,7 +123,7 @@ void FCurveColorCustomization::CustomizeHeader(TSharedRef<IPropertyHandle> InStr
 		HeaderRow
 			.NameContent()
 			[
-				InStructPropertyHandle->CreatePropertyNameWidget( FText::GetEmpty(), FText::GetEmpty(), false )
+				InStructPropertyHandle->CreatePropertyNameWidget()
 			]
 			.ValueContent()
 			[
diff --git a/Engine/Source/Editor/DetailCustomizations/Private/DataTableCategoryCustomization.h b/Engine/Source/Editor/DetailCustomizations/Private/DataTableCategoryCustomization.h
--- a/Engine/Source/Editor/DetailCustomizations/Private/DataTableCategoryCustomization.h
+++ b/Engine/Source/Editor/DetailCustomizations/Private/DataTableCategoryCustomization.h
@@ -33,7 +33,7 @@ public:
 		HeaderRow
 			.NameContent()
 			[
-				InStructPropertyHandle->CreatePropertyNameWidget(FText::GetEmpty(), FText::GetEmpty(), false)
+				InStructPropertyHandle->CreatePropertyNameWidget()
 			];
 	}
 
diff --git a/Engine/Source/Editor/DetailCustomizations/Private/DataTableCustomization.cpp b/Engine/Source/Editor/DetailCustomizations/Private/DataTableCustomization.cpp
--- a/Engine/Source/Editor/DetailCustomizations/Private/DataTableCustomization.cpp
+++ b/Engine/Source/Editor/DetailCustomizations/Private/DataTableCustomization.cpp
@@ -28,7 +28,7 @@ void FDataTableCustomizationLayout::CustomizeHeader(TSharedRef<class IPropertyHa
 	HeaderRow
 	.NameContent()
 	[
-		InStructPropertyHandle->CreatePropertyNameWidget(FText::GetEmpty(), FText::GetEmpty(), false)
+		InStructPropertyHandle->CreatePropertyNameWidget()
 	];
 
 	FDataTableEditorUtils::AddSearchForReferencesContextMenu(HeaderRow, FExecuteAction::CreateSP(this, &FDataTableCustomizationLayout::OnSearchForReferences));
diff --git a/Engine/Source/Editor/DetailCustomizations/Private/DebugCameraControllerSettingsCustomization.cpp b/Engine/Source/Editor/DetailCustomizations/Private/DebugCameraControllerSettingsCustomization.cpp
--- a/Engine/Source/Editor/DetailCustomizations/Private/DebugCameraControllerSettingsCustomization.cpp
+++ b/Engine/Source/Editor/DetailCustomizations/Private/DebugCameraControllerSettingsCustomization.cpp
@@ -24,11 +24,6 @@ FDebugCameraControllerSettingsViewModeIndexCustomization::FDebugCameraController
 
 void FDebugCameraControllerSettingsViewModeIndexCustomization::CustomizeHeader(TSharedRef<class IPropertyHandle> StructPropertyHandle, class FDetailWidgetRow& HeaderRow, IPropertyTypeCustomizationUtils& StructCustomizationUtils)
 {
-	// We'll set up reset to default ourselves
-	const bool bDisplayResetToDefault = false;
-	const FText DisplayNameOverride = FText::GetEmpty();
-	const FText DisplayToolTipOverride = FText::GetEmpty();
-
 	uint32 NumChildren;
 	StructPropertyHandle->GetNumChildren(NumChildren);
 	TSharedPtr<IPropertyHandle> ViewModeIndexHandle;
@@ -62,7 +57,7 @@ void FDebugCameraControllerSettingsViewModeIndexCustomization::CustomizeHeader(T
 	HeaderRow
 		.NameContent()
 		[
-			StructPropertyHandle->CreatePropertyNameWidget(DisplayNameOverride, DisplayToolTipOverride, bDisplayResetToDefault)
+			StructPropertyHandle->CreatePropertyNameWidget()
 		]
 		.ValueContent()
 		.MinDesiredWidth(500)
diff --git a/Engine/Source/Editor/DetailCustomizations/Private/MathStructCustomizations.cpp b/Engine/Source/Editor/DetailCustomizations/Private/MathStructCustomizations.cpp
--- a/Engine/Source/Editor/DetailCustomizations/Private/MathStructCustomizations.cpp
+++ b/Engine/Source/Editor/DetailCustomizations/Private/MathStructCustomizations.cpp
@@ -44,18 +44,13 @@ void FMathStructCustomization::CustomizeChildren(TSharedRef<class IPropertyHandl
 
 void FMathStructCustomization::MakeHeaderRow(TSharedRef<class IPropertyHandle>& StructPropertyHandle, FDetailWidgetRow& Row)
 {
-	// We'll set up reset to default ourselves
-	const bool bDisplayResetToDefault = false;
-	const FText DisplayNameOverride = FText::GetEmpty();
-	const FText DisplayToolTipOverride = FText::GetEmpty();
-
 	TWeakPtr<IPropertyHandle> StructWeakHandlePtr = StructPropertyHandle;
 
 	TSharedPtr<SHorizontalBox> HorizontalBox;
 
 	Row.NameContent()
 	[
-		StructPropertyHandle->CreatePropertyNameWidget(DisplayNameOverride, DisplayToolTipOverride, bDisplayResetToDefault)
+		StructPropertyHandle->CreatePropertyNameWidget()
 	]
 	.ValueContent()
 	// Make enough space for each child handle
diff --git a/Engine/Source/Editor/DetailCustomizations/Private/ParticleSysParamStructCustomization.cpp b/Engine/Source/Editor/DetailCustomizations/Private/ParticleSysParamStructCustomization.cpp
--- a/Engine/Source/Editor/DetailCustomizations/Private/ParticleSysParamStructCustomization.cpp
+++ b/Engine/Source/Editor/DetailCustomizations/Private/ParticleSysParamStructCustomization.cpp
@@ -30,14 +30,10 @@ FParticleSysParamStructCustomization::FParticleSysParamStructCustomization()
 
 void FParticleSysParamStructCustomization::CustomizeHeader(TSharedRef<IPropertyHandle> StructPropertyHandle, FDetailWidgetRow& HeaderRow, IPropertyTypeCustomizationUtils& StructCustomizationUtils)
 {
-	const bool bDisplayResetToDefault = false;
-	const FText DisplayNameOverride = FText::GetEmpty();
-	const FText DisplayToolTipOverride = FText::GetEmpty();
-
 	HeaderRow
 		.NameContent()
 		[
-			StructPropertyHandle->CreatePropertyNameWidget(DisplayNameOverride, DisplayToolTipOverride, bDisplayResetToDefault)
+			StructPropertyHandle->CreatePropertyNameWidget()
 		];
 }
 
diff --git a/Engine/Source/Editor/DetailCustomizations/Private/RawDistributionVectorStructCustomization.cpp b/Engine/Source/Editor/DetailCustomizations/Private/RawDistributionVectorStructCustomization.cpp
--- a/Engine/Source/Editor/DetailCustomizations/Private/RawDistributionVectorStructCustomization.cpp
+++ b/Engine/Source/Editor/DetailCustomizations/Private/RawDistributionVectorStructCustomization.cpp
@@ -505,14 +505,10 @@ TSharedRef<IPropertyTypeCustomization> FRawDistributionVectorStructCustomization
 
 void FRawDistributionVectorStructCustomization::CustomizeHeader(TSharedRef<IPropertyHandle> StructPropertyHandle, FDetailWidgetRow& HeaderRow, IPropertyTypeCustomizationUtils& StructCustomizationUtils)
 {
-	const bool bDisplayResetToDefault = false;
-	const FText DisplayNameOverride = FText::GetEmpty();
-	const FText DisplayToolTipOverride = FText::GetEmpty();
-
 	HeaderRow
 	.NameContent()
 	[
-		StructPropertyHandle->CreatePropertyNameWidget(DisplayNameOverride, DisplayToolTipOverride, bDisplayResetToDefault)
+		StructPropertyHandle->CreatePropertyNameWidget()
 	]
 	.ValueContent()
 	.MinDesiredWidth(1)
diff --git a/Engine/Source/Editor/DetailCustomizations/Private/SlateBrushCustomization.cpp b/Engine/Source/Editor/DetailCustomizations/Private/SlateBrushCustomization.cpp
--- a/Engine/Source/Editor/DetailCustomizations/Private/SlateBrushCustomization.cpp
+++ b/Engine/Source/Editor/DetailCustomizations/Private/SlateBrushCustomization.cpp
@@ -1402,7 +1402,7 @@ void FSlateBrushStructCustomization::CustomizeChildren( TSharedRef<IPropertyHand
 				.HeaderRow()
 				.NameContent()
 				[
-					StructPropertyHandle->CreatePropertyNameWidget(NSLOCTEXT("UnrealEd", "Preview", "Preview"), FText::GetEmpty(), false)
+					StructPropertyHandle->CreatePropertyNameWidget(NSLOCTEXT("UnrealEd", "Preview", "Preview"))
 				]
 			.ValueContent()
 				.MinDesiredWidth(1)
diff --git a/Engine/Source/Editor/DetailCustomizations/Public/Customizations/CurveTableCustomization.h b/Engine/Source/Editor/DetailCustomizations/Public/Customizations/CurveTableCustomization.h
--- a/Engine/Source/Editor/DetailCustomizations/Public/Customizations/CurveTableCustomization.h
+++ b/Engine/Source/Editor/DetailCustomizations/Public/Customizations/CurveTableCustomization.h
@@ -39,7 +39,7 @@ public:
 		HeaderRow
 			.NameContent()
 			[
-				InStructPropertyHandle->CreatePropertyNameWidget( FText::GetEmpty(), FText::GetEmpty(), false )
+				InStructPropertyHandle->CreatePropertyNameWidget( FText::GetEmpty(), FText::GetEmpty() )
 			];
 	}
 
diff --git a/Engine/Source/Editor/MaterialEditor/Private/MaterialEditorInstanceDetailCustomization.cpp b/Engine/Source/Editor/MaterialEditor/Private/MaterialEditorInstanceDetailCustomization.cpp
--- a/Engine/Source/Editor/MaterialEditor/Private/MaterialEditorInstanceDetailCustomization.cpp
+++ b/Engine/Source/Editor/MaterialEditor/Private/MaterialEditorInstanceDetailCustomization.cpp
@@ -573,7 +573,7 @@ void FMaterialInstanceParameterDetails::CreateMaskParameterValueWidget(UDEditorP
 			.HAlign(HAlign_Left)
 			.AutoWidth()
 			[
-				RMaskProperty->CreatePropertyNameWidget(FText::GetEmpty(), FText::GetEmpty(), false)
+				RMaskProperty->CreatePropertyNameWidget()
 			]
 		+ SHorizontalBox::Slot()
 			.HAlign(HAlign_Left)
@@ -586,7 +586,7 @@ void FMaterialInstanceParameterDetails::CreateMaskParameterValueWidget(UDEditorP
 			.Padding(FMargin(10.0f, 0.0f, 0.0f, 0.0f))
 			.AutoWidth()
 			[
-				GMaskProperty->CreatePropertyNameWidget(FText::GetEmpty(), FText::GetEmpty(), false)
+				GMaskProperty->CreatePropertyNameWidget()
 			]
 		+ SHorizontalBox::Slot()
 			.HAlign(HAlign_Left)
@@ -599,7 +599,7 @@ void FMaterialInstanceParameterDetails::CreateMaskParameterValueWidget(UDEditorP
 			.Padding(FMargin(10.0f, 0.0f, 0.0f, 0.0f))
 			.AutoWidth()
 			[
-				BMaskProperty->CreatePropertyNameWidget(FText::GetEmpty(), FText::GetEmpty(), false)
+				BMaskProperty->CreatePropertyNameWidget()
 			]
 		+ SHorizontalBox::Slot()
 			.HAlign(HAlign_Left)
@@ -612,7 +612,7 @@ void FMaterialInstanceParameterDetails::CreateMaskParameterValueWidget(UDEditorP
 			.Padding(FMargin(10.0f, 0.0f, 0.0f, 0.0f))
 			.AutoWidth()
 			[
-				AMaskProperty->CreatePropertyNameWidget(FText::GetEmpty(), FText::GetEmpty(), false)
+				AMaskProperty->CreatePropertyNameWidget()
 			]
 		+ SHorizontalBox::Slot()
 			.HAlign(HAlign_Left)
diff --git a/Engine/Source/Editor/MaterialEditor/Private/SMaterialLayersFunctionsTree.cpp b/Engine/Source/Editor/MaterialEditor/Private/SMaterialLayersFunctionsTree.cpp
--- a/Engine/Source/Editor/MaterialEditor/Private/SMaterialLayersFunctionsTree.cpp
+++ b/Engine/Source/Editor/MaterialEditor/Private/SMaterialLayersFunctionsTree.cpp
@@ -832,7 +832,7 @@ void SMaterialLayersFunctionsInstanceTreeItem::Construct(const FArguments& InArg
 					.HAlign(HAlign_Left)
 					.AutoWidth()
 					[
-						RMaskProperty->CreatePropertyNameWidget(FText::GetEmpty(), FText::GetEmpty(), false)
+						RMaskProperty->CreatePropertyNameWidget()
 					]
 					+ SHorizontalBox::Slot()
 					.HAlign(HAlign_Left)
@@ -845,7 +845,7 @@ void SMaterialLayersFunctionsInstanceTreeItem::Construct(const FArguments& InArg
 					.Padding(FMargin(10.0f, 0.0f, 0.0f, 0.0f))
 					.AutoWidth()
 					[
-						GMaskProperty->CreatePropertyNameWidget(FText::GetEmpty(), FText::GetEmpty(), false)
+						GMaskProperty->CreatePropertyNameWidget()
 					]
 					+ SHorizontalBox::Slot()
 					.HAlign(HAlign_Left)
@@ -858,7 +858,7 @@ void SMaterialLayersFunctionsInstanceTreeItem::Construct(const FArguments& InArg
 					.Padding(FMargin(10.0f, 0.0f, 0.0f, 0.0f))
 					.AutoWidth()
 					[
-						BMaskProperty->CreatePropertyNameWidget(FText::GetEmpty(), FText::GetEmpty(), false)
+						BMaskProperty->CreatePropertyNameWidget()
 					]
 					+ SHorizontalBox::Slot()
 					.HAlign(HAlign_Left)
@@ -871,7 +871,7 @@ void SMaterialLayersFunctionsInstanceTreeItem::Construct(const FArguments& InArg
 					.Padding(FMargin(10.0f, 0.0f, 0.0f, 0.0f))
 					.AutoWidth()
 					[
-						AMaskProperty->CreatePropertyNameWidget(FText::GetEmpty(), FText::GetEmpty(), false)
+						AMaskProperty->CreatePropertyNameWidget()
 					]
 					+ SHorizontalBox::Slot()
 					.HAlign(HAlign_Left)
@@ -2372,7 +2372,7 @@ void SMaterialLayersFunctionsMaterialTreeItem::Construct(const FArguments& InArg
 					.HAlign(HAlign_Left)
 					.AutoWidth()
 					[
-						RMaskProperty->CreatePropertyNameWidget(FText::GetEmpty(), FText::GetEmpty(), false)
+						RMaskProperty->CreatePropertyNameWidget()
 					]
 			+ SHorizontalBox::Slot()
 				.HAlign(HAlign_Left)
@@ -2385,7 +2385,7 @@ void SMaterialLayersFunctionsMaterialTreeItem::Construct(const FArguments& InArg
 				.Padding(FMargin(10.0f, 0.0f, 0.0f, 0.0f))
 				.AutoWidth()
 				[
-					GMaskProperty->CreatePropertyNameWidget(FText::GetEmpty(), FText::GetEmpty(), false)
+					GMaskProperty->CreatePropertyNameWidget()
 				]
 			+ SHorizontalBox::Slot()
 				.HAlign(HAlign_Left)
@@ -2398,7 +2398,7 @@ void SMaterialLayersFunctionsMaterialTreeItem::Construct(const FArguments& InArg
 				.Padding(FMargin(10.0f, 0.0f, 0.0f, 0.0f))
 				.AutoWidth()
 				[
-					BMaskProperty->CreatePropertyNameWidget(FText::GetEmpty(), FText::GetEmpty(), false)
+					BMaskProperty->CreatePropertyNameWidget()
 				]
 			+ SHorizontalBox::Slot()
 				.HAlign(HAlign_Left)
@@ -2411,7 +2411,7 @@ void SMaterialLayersFunctionsMaterialTreeItem::Construct(const FArguments& InArg
 				.Padding(FMargin(10.0f, 0.0f, 0.0f, 0.0f))
 				.AutoWidth()
 				[
-					AMaskProperty->CreatePropertyNameWidget(FText::GetEmpty(), FText::GetEmpty(), false)
+					AMaskProperty->CreatePropertyNameWidget()
 				]
 			+ SHorizontalBox::Slot()
 				.HAlign(HAlign_Left)
diff --git a/Engine/Source/Editor/MaterialEditor/Private/SMaterialParametersOverviewWidget.cpp b/Engine/Source/Editor/MaterialEditor/Private/SMaterialParametersOverviewWidget.cpp
--- a/Engine/Source/Editor/MaterialEditor/Private/SMaterialParametersOverviewWidget.cpp
+++ b/Engine/Source/Editor/MaterialEditor/Private/SMaterialParametersOverviewWidget.cpp
@@ -368,7 +368,7 @@ void SMaterialParametersOverviewTreeItem::Construct(const FArguments& InArgs, co
 					.HAlign(HAlign_Left)
 					.AutoWidth()
 					[
-						RMaskProperty->CreatePropertyNameWidget(FText::GetEmpty(), FText::GetEmpty(), false)
+						RMaskProperty->CreatePropertyNameWidget()
 					]
 					+ SHorizontalBox::Slot()
 					.HAlign(HAlign_Left)
@@ -381,7 +381,7 @@ void SMaterialParametersOverviewTreeItem::Construct(const FArguments& InArgs, co
 					.Padding(FMargin(10.0f, 0.0f, 0.0f, 0.0f))
 					.AutoWidth()
 					[
-						GMaskProperty->CreatePropertyNameWidget(FText::GetEmpty(), FText::GetEmpty(), false)
+						GMaskProperty->CreatePropertyNameWidget()
 					]
 					+ SHorizontalBox::Slot()
 					.HAlign(HAlign_Left)
@@ -394,7 +394,7 @@ void SMaterialParametersOverviewTreeItem::Construct(const FArguments& InArgs, co
 					.Padding(FMargin(10.0f, 0.0f, 0.0f, 0.0f))
 					.AutoWidth()
 					[
-						BMaskProperty->CreatePropertyNameWidget(FText::GetEmpty(), FText::GetEmpty(), false)
+						BMaskProperty->CreatePropertyNameWidget()
 					]
 					+ SHorizontalBox::Slot()
 					.HAlign(HAlign_Left)
@@ -407,7 +407,7 @@ void SMaterialParametersOverviewTreeItem::Construct(const FArguments& InArgs, co
 					.Padding(FMargin(10.0f, 0.0f, 0.0f, 0.0f))
 					.AutoWidth()
 					[
-						AMaskProperty->CreatePropertyNameWidget(FText::GetEmpty(), FText::GetEmpty(), false)
+						AMaskProperty->CreatePropertyNameWidget()
 					]
 					+ SHorizontalBox::Slot()
 					.HAlign(HAlign_Left)
diff --git a/Engine/Source/Editor/Persona/Private/AnimNotifyDetails.cpp b/Engine/Source/Editor/Persona/Private/AnimNotifyDetails.cpp
--- a/Engine/Source/Editor/Persona/Private/AnimNotifyDetails.cpp
+++ b/Engine/Source/Editor/Persona/Private/AnimNotifyDetails.cpp
@@ -249,7 +249,7 @@ void FAnimNotifyDetails::ClearInstancedSelectionDropDown(IDetailCategoryBuilder&
 	.CustomWidget(bShowChildren)
 	.NameContent()
 	[
-		PropHandle->CreatePropertyNameWidget(FText::GetEmpty(), FText::GetEmpty(), false)
+		PropHandle->CreatePropertyNameWidget()
 	]
 	.ValueContent()
 	[
diff --git a/Engine/Source/Editor/PropertyEditor/Private/PropertyHandleImpl.cpp b/Engine/Source/Editor/PropertyEditor/Private/PropertyHandleImpl.cpp
--- a/Engine/Source/Editor/PropertyEditor/Private/PropertyHandleImpl.cpp
+++ b/Engine/Source/Editor/PropertyEditor/Private/PropertyHandleImpl.cpp
@@ -2120,16 +2120,21 @@ FString FPropertyHandleBase::GeneratePathToProperty() const
 
 }
 
-TSharedRef<SWidget> FPropertyHandleBase::CreatePropertyNameWidget( const FText& NameOverride, const FText& ToolTipOverride, bool bDisplayResetToDefault, bool bDisplayText, bool bDisplayThumbnail ) const
+TSharedRef<SWidget> FPropertyHandleBase::CreatePropertyNameWidget(const FText& NameOverride, const FText& ToolTipOverride, bool bDisplayResetToDefault, bool bDisplayText, bool bDisplayThumbnail) const
+{
+	return CreatePropertyNameWidget(NameOverride, ToolTipOverride);
+}
+
+TSharedRef<SWidget> FPropertyHandleBase::CreatePropertyNameWidget(const FText& NameOverride, const FText& ToolTipOverride) const
 {
 	if( Implementation.IsValid() && Implementation->GetPropertyNode().IsValid() )
 	{
 		struct FPropertyNodeDisplayNameOverrideHelper
 		{
 			FPropertyNodeDisplayNameOverrideHelper(TSharedPtr<FPropertyValueImpl> InImplementation, const FText& InNameOverride, const FText& InToolTipOverride)
-				:Implementation(InImplementation)
-				,bResetDisplayName(false)
-				,bResetToolTipText(false)
+				: Implementation(InImplementation)
+				, bResetDisplayName(false)
+				, bResetToolTipText(false)
 			{
 				if (!InNameOverride.IsEmpty())
 				{
@@ -3535,7 +3540,7 @@ FPropertyAccess::Result FPropertyHandleDouble::SetValue( const double& NewValue,
 	// Clamp the value from any meta data ranges stored on the property value
 	double FinalValue = ClampValueFromMetaData<double>( NewValue, *Implementation->GetPropertyNode() );
 
-	const FString ValueStr = FString::Printf( TEXT("%f"), FinalValue );
+	const FString ValueStr = FString::Printf( TEXT("%.12f"), FinalValue );
 	Res = Implementation->ImportText( ValueStr, Flags );
 
 	return Res;
diff --git a/Engine/Source/Editor/PropertyEditor/Private/PropertyHandleImpl.h b/Engine/Source/Editor/PropertyEditor/Private/PropertyHandleImpl.h
--- a/Engine/Source/Editor/PropertyEditor/Private/PropertyHandleImpl.h
+++ b/Engine/Source/Editor/PropertyEditor/Private/PropertyHandleImpl.h
@@ -419,7 +419,8 @@ public:
 	virtual bool IsCustomized() const override;
 	virtual bool IsResetToDefaultCustomized() const override;
 	virtual FString GeneratePathToProperty() const override;
-	virtual TSharedRef<SWidget> CreatePropertyNameWidget( const FText& NameOverride = FText::GetEmpty(), const FText& ToolTipOverride = FText::GetEmpty(), bool bDisplayResetToDefault = false, bool bDisplayText = true, bool bDisplayThumbnail = true ) const override;
+	virtual TSharedRef<SWidget> CreatePropertyNameWidget( const FText& NameOverride, const FText& ToolTipOverride, bool bDisplayResetToDefault, bool bDisplayText, bool bDisplayThumbnail ) const override;
+	virtual TSharedRef<SWidget> CreatePropertyNameWidget(const FText& NameOverride, const FText& ToolTipOverride) const override;
 	virtual TSharedRef<SWidget> CreatePropertyValueWidget( bool bDisplayDefaultPropertyButtons = true ) const override;
 	virtual TSharedRef<SWidget> CreateDefaultPropertyButtonWidgets() const override;
 	virtual void CreateDefaultPropertyCopyPasteActions(FUIAction& OutCopyAction, FUIAction& OutPasteAction) const override;
diff --git a/Engine/Source/Editor/PropertyEditor/Public/PropertyCustomizationHelpers.h b/Engine/Source/Editor/PropertyEditor/Public/PropertyCustomizationHelpers.h
--- a/Engine/Source/Editor/PropertyEditor/Public/PropertyCustomizationHelpers.h
+++ b/Engine/Source/Editor/PropertyEditor/Public/PropertyCustomizationHelpers.h
@@ -445,8 +445,6 @@ public:
 	{
 		if (bGenerateHeader)
 		{
-			const bool bDisplayResetToDefaultInNameContent = false;
-
 			TSharedPtr<SHorizontalBox> ContentHorizontalBox;
 			SAssignNew(ContentHorizontalBox, SHorizontalBox);
 			if (bDisplayElementNum)
@@ -465,7 +463,7 @@ public:
 			.FilterString(!DisplayName.IsEmpty() ? DisplayName : BaseProperty->GetPropertyDisplayName())
 			.NameContent()
 			[
-				BaseProperty->CreatePropertyNameWidget(DisplayName, FText::GetEmpty(), bDisplayResetToDefaultInNameContent)
+				BaseProperty->CreatePropertyNameWidget(DisplayName, FText::GetEmpty())
 			]
 			.ValueContent()
 			[
diff --git a/Engine/Source/Editor/PropertyEditor/Public/PropertyHandle.h b/Engine/Source/Editor/PropertyEditor/Public/PropertyHandle.h
--- a/Engine/Source/Editor/PropertyEditor/Public/PropertyHandle.h
+++ b/Engine/Source/Editor/PropertyEditor/Public/PropertyHandle.h
@@ -551,8 +551,15 @@ public:
 	 * @param bDisplayThumbnail			Whether or not to display the thumbnail for the property (if any)
 	 * @return the name widget for this property
 	 */
-	virtual TSharedRef<SWidget> CreatePropertyNameWidget( const FText& NameOverride = FText::GetEmpty(), const FText& ToolTipOverride = FText::GetEmpty(), bool bDisplayResetToDefault = false, bool bDisplayText = true, bool bDisplayThumbnail = true ) const = 0;
+	UE_DEPRECATED(5.0, "CreatePropertyNameWidget no longer supports bDisplayResetToDefault, bDisplayText or bDisplayThumbnail.")
+	virtual TSharedRef<SWidget> CreatePropertyNameWidget( const FText& NameOverride, const FText& ToolTipOverride, bool bDisplayResetToDefault, bool bDisplayText = true, bool bDisplayThumbnail = true ) const = 0;
 
+	/**
+	 * Creates a name widget for this property
+	 * @param NameOverride				The name override to use instead of the property name
+	 * @param ToolTipOverride			The tooltip override to use instead of the property name
+	 */
+	virtual TSharedRef<SWidget> CreatePropertyNameWidget(const FText& NameOverride = FText::GetEmpty(), const FText& ToolTipOverride = FText::GetEmpty()) const = 0;
 	/**
 	 * Creates a value widget for this property
 
