Improve availability of shader platforms in the mat ed (UE-136509).

diff --git a/Engine/Source/Editor/MaterialEditor/Private/MaterialStats.cpp b/Engine/Source/Editor/MaterialEditor/Private/MaterialStats.cpp
--- a/Engine/Source/Editor/MaterialEditor/Private/MaterialStats.cpp
+++ b/Engine/Source/Editor/MaterialEditor/Private/MaterialStats.cpp
@@ -436,14 +436,18 @@ void FMaterialStats::RefreshStatsGrid()
 
 void FMaterialStats::BuildShaderPlatformDB()
 {
+	bool bCanCompileMacDesktopSPs = (PLATFORM_MAC != 0);
+	bool bCanCompileMacMobileSPs = (PLATFORM_MAC != 0) || (PLATFORM_WINDOWS != 0);
+
 #if PLATFORM_WINDOWS
 	// DirectX
 	AddShaderPlatform(EPlatformCategoryType::Desktop, SP_PCD3D_SM5, TEXT("DirectX SM5"), true, true, TEXT("Desktop, DirectX, Shader Model 5"));
+	AddShaderPlatform(EPlatformCategoryType::Desktop, SP_PCD3D_SM6, TEXT("DirectX SM6"), true, true, TEXT("Desktop, DirectX, Shader Model 6"));
 	AddShaderPlatform(EPlatformCategoryType::Desktop, SP_PCD3D_ES3_1, TEXT("DirectX ES 3.1"), true, true, TEXT("Desktop, DirectX, ES 3.1"));
 #endif
 
 	// Vulkan
-	AddShaderPlatform(EPlatformCategoryType::Desktop, SP_VULKAN_SM5, TEXT("Vulkan SM5"), false, true, TEXT("Desktop, Vulkan, Shader Model 5"));
+	AddShaderPlatform(EPlatformCategoryType::Desktop, SP_VULKAN_SM5, TEXT("Vulkan SM5"), true, true, TEXT("Desktop, Vulkan, Shader Model 5"));
 
 	// Android
 	AddShaderPlatform(EPlatformCategoryType::Android, SP_OPENGL_ES3_1_ANDROID, TEXT("Android GLES 3.1"), true, true, TEXT("Android, OpenGLES 3.1"));
@@ -451,9 +455,9 @@ void FMaterialStats::BuildShaderPlatformDB()
 	AddShaderPlatform(EPlatformCategoryType::Android, SP_VULKAN_SM5_ANDROID, TEXT("Android Vulkan SM5"), true, true, TEXT("Android, Vulkan SM5"));
 
 	// Apple
-	AddShaderPlatform(EPlatformCategoryType::Desktop, SP_METAL_SM5, TEXT("Metal SM5"), false, true, TEXT("macOS, Metal, Shader Model 5"));
-	AddShaderPlatform(EPlatformCategoryType::IOS, SP_METAL, TEXT("Metal"), false, true, TEXT("iOS, Metal, Mobile"));
-	AddShaderPlatform(EPlatformCategoryType::IOS, SP_METAL_MRT, TEXT("Metal MRT"), false, true, TEXT("iOS, Metal, Shader Model 5"));
+	AddShaderPlatform(EPlatformCategoryType::Desktop, SP_METAL_SM5, TEXT("Metal SM5"), bCanCompileMacDesktopSPs, true, TEXT("macOS, Metal, Shader Model 5"));
+	AddShaderPlatform(EPlatformCategoryType::IOS, SP_METAL, TEXT("Metal"), bCanCompileMacMobileSPs, true, TEXT("iOS, Metal, Mobile"));
+	AddShaderPlatform(EPlatformCategoryType::IOS, SP_METAL_MRT, TEXT("Metal MRT"), bCanCompileMacMobileSPs, true, TEXT("iOS, Metal, Shader Model 5"));
 
 	ITargetPlatformManagerModule& TPM = GetTargetPlatformManagerRef();
 
