Fix shaders stats not showing up on Mac (UE-134608).

diff --git a/Engine/Source/Editor/MaterialEditor/Private/MaterialStats.cpp b/Engine/Source/Editor/MaterialEditor/Private/MaterialStats.cpp
--- a/Engine/Source/Editor/MaterialEditor/Private/MaterialStats.cpp
+++ b/Engine/Source/Editor/MaterialEditor/Private/MaterialStats.cpp
@@ -456,6 +456,8 @@ void FMaterialStats::BuildShaderPlatformDB()
 
 	// Apple
 	AddShaderPlatform(EPlatformCategoryType::Desktop, SP_METAL_SM5, TEXT("Metal SM5"), bCanCompileMacDesktopSPs, true, TEXT("macOS, Metal, Shader Model 5"));
+	AddShaderPlatform(EPlatformCategoryType::Desktop, SP_METAL_SM5_NOTESS, TEXT("Metal SM5 (no tesselation)"), bCanCompileMacDesktopSPs, true, TEXT("macOS, Metal, Shader Model 5"));
+	AddShaderPlatform(EPlatformCategoryType::Desktop, SP_METAL_MRT_MAC, TEXT("Metal SM5 (MRT)"), bCanCompileMacDesktopSPs, true, TEXT("macOS, Metal, Shader Model 5"));
 	AddShaderPlatform(EPlatformCategoryType::IOS, SP_METAL, TEXT("Metal"), bCanCompileMacMobileSPs, true, TEXT("iOS, Metal, Mobile"));
 	AddShaderPlatform(EPlatformCategoryType::IOS, SP_METAL_MRT, TEXT("Metal MRT"), bCanCompileMacMobileSPs, true, TEXT("iOS, Metal, Shader Model 5"));
 
