Fix blocking compilation when editing instances of special materials.

- Also added a check to surface a potential memory freezing issue earlier (unrelated).

diff --git a/Engine/Source/Editor/MaterialEditor/Public/MaterialStatsCommon.h b/Engine/Source/Editor/MaterialEditor/Public/MaterialStatsCommon.h
--- a/Engine/Source/Editor/MaterialEditor/Public/MaterialStatsCommon.h
+++ b/Engine/Source/Editor/MaterialEditor/Public/MaterialStatsCommon.h
@@ -23,6 +23,9 @@ public:
 	 */
 	virtual bool IsPersistent() const override { return false; }
 
+	/** Material resource stats never requires a synchronous compilation, otherwise opening up empty (newly created) material instance will block compiling default mat's shaders. */
+	virtual bool RequiresSynchronousCompilation() const override { return false; }
+
 	/** this will enable shader source extraction and pass paths to (eventual) offline shader compilers */
 	virtual void SetupExtaCompilationSettings(const EShaderPlatform Platform, FExtraShaderCompilerSettings& Settings) const override;
 };
diff --git a/Engine/Source/Runtime/Engine/Private/Materials/MaterialShader.cpp b/Engine/Source/Runtime/Engine/Private/Materials/MaterialShader.cpp
--- a/Engine/Source/Runtime/Engine/Private/Materials/MaterialShader.cpp
+++ b/Engine/Source/Runtime/Engine/Private/Materials/MaterialShader.cpp
@@ -2816,6 +2816,9 @@ FMaterialShaderMap* FMaterialShaderMap::AcquireFinalizedClone()
 	Clone->bCompiledSuccessfully = bCompiledSuccessfully;
 	Clone->bIsPersistent = bIsPersistent;
 	Clone->AssignCopy(*this);
+	checkf(reinterpret_cast<uint64>(Clone->GetMaterialPath()) > reinterpret_cast<uint64>(Clone->GetContent()), TEXT("String embedded in memory image is located outside of it (image start %p > string %p). Pointer/offset overflow while freezing?"),
+		Clone->GetContent(), Clone->GetMaterialPath()
+		);
 #if WITH_EDITOR
 	Clone->AssociateWithAssets(GetAssociatedAssets());
 #endif
