Hide UMaterialFunction UI options for new HLSL generator/control flow, if the global CVARs are disabled

diff --git a/Engine/Source/Editor/MaterialEditor/Private/MaterialEditor.cpp b/Engine/Source/Editor/MaterialEditor/Private/MaterialEditor.cpp
--- a/Engine/Source/Editor/MaterialEditor/Private/MaterialEditor.cpp
+++ b/Engine/Source/Editor/MaterialEditor/Private/MaterialEditor.cpp
@@ -1482,6 +1482,7 @@ void FMaterialEditor::CreateInternalWidgets()
 	MaterialDetailsView->OnFinishedChangingProperties().AddSP(this, &FMaterialEditor::OnFinishedChangingProperties);
 
 	PropertyEditorModule.RegisterCustomClassLayout( UMaterial::StaticClass()->GetFName(), FOnGetDetailCustomizationInstance::CreateStatic(&FMaterialDetailCustomization::MakeInstance ) );
+	PropertyEditorModule.RegisterCustomClassLayout( UMaterialFunction::StaticClass()->GetFName(), FOnGetDetailCustomizationInstance::CreateStatic(&FMaterialFunctionDetailCustomization::MakeInstance) );
 
 	MaterialEditorInstance = NewObject<UMaterialEditorPreviewParameters>(GetTransientPackage(), NAME_None, RF_Transactional);
 	MaterialEditorInstance->PreviewMaterial = Material;
diff --git a/Engine/Source/Editor/MaterialEditor/Private/MaterialEditorDetailCustomization.cpp b/Engine/Source/Editor/MaterialEditor/Private/MaterialEditorDetailCustomization.cpp
--- a/Engine/Source/Editor/MaterialEditor/Private/MaterialEditorDetailCustomization.cpp
+++ b/Engine/Source/Editor/MaterialEditor/Private/MaterialEditorDetailCustomization.cpp
@@ -837,6 +837,41 @@ void FMaterialDetailCustomization::CustomizeDetails( IDetailLayoutBuilder& Detai
 	}
 }
 
+TSharedRef<class IDetailCustomization> FMaterialFunctionDetailCustomization::MakeInstance()
+{
+	return MakeShareable(new FMaterialFunctionDetailCustomization);
+}
+
+void FMaterialFunctionDetailCustomization::CustomizeDetails(IDetailLayoutBuilder& DetailLayout)
+{
+	static const auto CVarMaterialEnableControlFlow = IConsoleManager::Get().FindTConsoleVariableDataInt(TEXT("r.MaterialEnableControlFlow"));
+	static const auto CVarMaterialEnableNewHLSLGenerator = IConsoleManager::Get().FindTConsoleVariableDataInt(TEXT("r.MaterialEnableNewHLSLGenerator"));
+
+	// MaterialFunction category
+	{
+		IDetailCategoryBuilder& MaterialCategory = DetailLayout.EditCategory(TEXT("MaterialFunction"));
+
+		TArray<TSharedRef<IPropertyHandle>> AllProperties;
+		MaterialCategory.GetDefaultProperties(AllProperties);
+
+		for (TSharedRef<IPropertyHandle>& PropertyHandle : AllProperties)
+		{
+			FProperty* Property = PropertyHandle->GetProperty();
+			FName PropertyName = Property->GetFName();
+
+			if (PropertyName == GET_MEMBER_NAME_CHECKED(UMaterialFunction, bEnableExecWire) && !CVarMaterialEnableControlFlow->GetValueOnAnyThread())
+			{
+				DetailLayout.HideProperty(PropertyHandle);
+			}
+
+			if (PropertyName == GET_MEMBER_NAME_CHECKED(UMaterialFunction, bEnableNewHLSLGenerator) && !CVarMaterialEnableNewHLSLGenerator->GetValueOnAnyThread())
+			{
+				DetailLayout.HideProperty(PropertyHandle);
+			}
+		}
+	}
+}
+
 TSharedRef<class IDetailCustomization> FMaterialExpressionLayersParameterDetails::MakeInstance(FOnCollectParameterGroups InCollectGroupsDelegate)
 {
 	return MakeShareable(new FMaterialExpressionLayersParameterDetails(InCollectGroupsDelegate));
diff --git a/Engine/Source/Editor/MaterialEditor/Private/MaterialEditorDetailCustomization.h b/Engine/Source/Editor/MaterialEditor/Private/MaterialEditorDetailCustomization.h
--- a/Engine/Source/Editor/MaterialEditor/Private/MaterialEditorDetailCustomization.h
+++ b/Engine/Source/Editor/MaterialEditor/Private/MaterialEditorDetailCustomization.h
@@ -147,3 +147,14 @@ public:
 	/** IDetailCustomization interface */
 	virtual void CustomizeDetails( IDetailLayoutBuilder& DetailLayout ) override;
 };
+
+class FMaterialFunctionDetailCustomization : public IDetailCustomization
+{
+public:
+
+	/** Makes a new instance of this detail layout class for a specific detail view requesting it */
+	static TSharedRef<class IDetailCustomization> MakeInstance();
+
+	/** IDetailCustomization interface */
+	virtual void CustomizeDetails(IDetailLayoutBuilder& DetailLayout) override;
+};
