Misc fixes for enabling/disabling control flow in material editor

diff --git a/Engine/Source/Editor/MaterialEditor/Private/MaterialEditor.cpp b/Engine/Source/Editor/MaterialEditor/Private/MaterialEditor.cpp
--- a/Engine/Source/Editor/MaterialEditor/Private/MaterialEditor.cpp
+++ b/Engine/Source/Editor/MaterialEditor/Private/MaterialEditor.cpp
@@ -2387,6 +2387,8 @@ void FMaterialEditor::UpdatePreviewMaterial( bool bForce )
 		ExpressionPreviewMaterial->Expressions = Material->Expressions;
 		ExpressionPreviewMaterial->ExpressionExecBegin = Material->ExpressionExecBegin;
 		ExpressionPreviewMaterial->ExpressionExecEnd = Material->ExpressionExecEnd;
+		ExpressionPreviewMaterial->bEnableExecWire = Material->IsUsingControlFlow();
+		ExpressionPreviewMaterial->bEnableNewHLSLGenerator = Material->IsUsingNewHLSLGenerator();
 
 		if (MaterialFunction)
 		{
@@ -2490,8 +2492,6 @@ bool FMaterialEditor::UpdateOriginalMaterial()
 		// Copy the expressions back from the preview material
 		MaterialFunction->FunctionExpressions = Material->Expressions;
 		MaterialFunction->FunctionEditorComments = Material->EditorComments;
-		MaterialFunction->ExpressionExecBegin = Material->ExpressionExecBegin;
-		MaterialFunction->ExpressionExecEnd = Material->ExpressionExecEnd;
 
 		// Preserve the thumbnail info
 		UThumbnailInfo* OriginalThumbnailInfo = MaterialFunction->ParentFunction->ThumbnailInfo;
diff --git a/Engine/Source/Runtime/Engine/Private/Materials/MaterialExpressions.cpp b/Engine/Source/Runtime/Engine/Private/Materials/MaterialExpressions.cpp
--- a/Engine/Source/Runtime/Engine/Private/Materials/MaterialExpressions.cpp
+++ b/Engine/Source/Runtime/Engine/Private/Materials/MaterialExpressions.cpp
@@ -13362,6 +13362,18 @@ void UMaterialFunction::CreateExecutionFlowExpressions()
 		EditorMaterial->bEnableExecWire = bUsingControlFlow;
 		EditorMaterial->ExpressionExecBegin = ExpressionExecBegin;
 		EditorMaterial->ExpressionExecEnd = ExpressionExecEnd;
+		if (ExpressionExecBegin)
+		{
+			EditorMaterial->Expressions.AddUnique(ExpressionExecBegin);
+			ExpressionExecBegin->Function = nullptr;
+			ExpressionExecBegin->Material = EditorMaterial;
+		}
+		if (ExpressionExecEnd)
+		{
+			EditorMaterial->Expressions.AddUnique(ExpressionExecEnd);
+			ExpressionExecEnd->Function = nullptr;
+			ExpressionExecEnd->Material = EditorMaterial;
+		}
 	}
 }
 #endif // WITH_EDITOR
diff --git a/Engine/Source/Runtime/Engine/Private/Materials/MaterialHLSLGenerator.cpp b/Engine/Source/Runtime/Engine/Private/Materials/MaterialHLSLGenerator.cpp
--- a/Engine/Source/Runtime/Engine/Private/Materials/MaterialHLSLGenerator.cpp
+++ b/Engine/Source/Runtime/Engine/Private/Materials/MaterialHLSLGenerator.cpp
@@ -489,10 +489,14 @@ bool FMaterialHLSLGenerator::GenerateStatements(UE::HLSLTree::FScope& Scope, UMa
 
 	FStatementEntry& Entry = StatementMap.FindOrAdd(MaterialExpression);
 	check(Entry.NumInputs >= 0);
-	check(Entry.NumInputs < MaterialExpression->NumExecutionInputs);
+
+	if (Entry.NumInputs >= MaterialExpression->NumExecutionInputs)
+	{
+		return Errors.AddErrorf(TEXT("Bad control flow, found %d inputs out of %d reported"), Entry.NumInputs, MaterialExpression->NumExecutionInputs);
+	}
 	if (Entry.NumInputs == MaxNumPreviousScopes)
 	{
-		return Errors.AddError(TEXT("Bad control flow"));
+		return Errors.AddErrorf(TEXT("Bad control flow, too many execution inputs"));
 	}
 
 	Entry.PreviousScope[Entry.NumInputs++] = &Scope;
