Fixed: Certain project configurations may result in missing material shader stats for ES31 feature level
Simplified shader name selection for mobile material shader stats

diff --git a/Engine/Source/Editor/MaterialEditor/Private/MaterialStatsCommon.cpp b/Engine/Source/Editor/MaterialEditor/Private/MaterialStatsCommon.cpp
--- a/Engine/Source/Editor/MaterialEditor/Private/MaterialStatsCommon.cpp
+++ b/Engine/Source/Editor/MaterialEditor/Private/MaterialStatsCommon.cpp
@@ -258,6 +258,29 @@ FSlateColor FMaterialStatsUtils::QualitySettingColor(const EMaterialQualityLevel
 	return  FStyleColors::Foreground;
 }
 
+static void MobileBasePassShaderName(bool bVertexShader, const TCHAR* PolicyName, int32 NumPointLigts, bool bHDR, bool bSkyLight, FString& OutName)
+{
+	OutName.Reset();
+	OutName.Append(bVertexShader ? TEXT("TMobileBasePassVS") : TEXT("TMobileBasePassPS"));
+	OutName.Append(PolicyName);
+	if (!bVertexShader)
+	{
+		if (NumPointLigts == INT32_MAX)
+		{
+			OutName.Append(TEXT("INT32_MAX"));
+		}
+		else
+		{
+			OutName.AppendInt(NumPointLigts);
+		}
+	}
+	OutName.Append(bHDR ? TEXT("HDRLinear64") : TEXT("LDRGamma32"));
+	if (!bVertexShader && bSkyLight)
+	{
+		OutName.Append(TEXT("SkyLight"));
+	}
+}
+
 void FMaterialStatsUtils::GetRepresentativeShaderTypesAndDescriptions(TMap<FName, TArray<FRepresentativeShaderInfo>>& ShaderTypeNamesAndDescriptions, const FMaterial* TargetMaterial)
 {
 	static auto* MobileHDR = IConsoleManager::Get().FindTConsoleVariableDataInt(TEXT("r.MobileHDR"));
@@ -324,29 +347,23 @@ void FMaterialStatsUtils::GetRepresentativeShaderTypesAndDescriptions(TMap<FName
 	else
 	{
 		const TCHAR* DescSuffix = bMobileHDR ? TEXT(" (HDR)") : TEXT(" (LDR)");
+		FString ShaderNameStr;
 
 		if (TargetMaterial->GetShadingModels().IsUnlit())
 		{
 			//unlit materials are never lightmapped
-			static FName Name_HDRLinear64 = TEXT("TMobileBasePassPSFNoLightMapPolicy0HDRLinear64");
-			static FName Name_LDRGamma32 = TEXT("TMobileBasePassPSFNoLightMapPolicy0LDRGamma32");
-
-			const FName TBasePassForForwardShadingPSFNoLightMapPolicy0Name = bMobileHDR ? Name_HDRLinear64 : Name_LDRGamma32;
-
+			MobileBasePassShaderName(false, TEXT("FNoLightMapPolicy"), 0, bMobileHDR, false, ShaderNameStr);
 			const FString Description = FString::Printf(TEXT("Mobile base pass shader without light map%s"), DescSuffix);
 			ShaderTypeNamesAndDescriptions.Add(FLocalVertexFactoryName)
-				.Add(FRepresentativeShaderInfo(ERepresentativeShader::StationarySurface, TBasePassForForwardShadingPSFNoLightMapPolicy0Name, Description));
-
-			static const FName Name_NoLM_HDRLinear64 = TEXT("TMobileBasePassVSFNoLightMapPolicyHDRLinear64");
-			static const FName Name_NoLM_LDRGamma32 = TEXT("TMobileBasePassVSFNoLightMapPolicyLDRGamma32");
-			static const FName TBasePassForForwardShadingVSFNoLightMapPolicyName = bMobileHDR ? Name_NoLM_HDRLinear64 : Name_NoLM_LDRGamma32;
+				.Add(FRepresentativeShaderInfo(ERepresentativeShader::StationarySurface, FName(ShaderNameStr), Description));
 
+			MobileBasePassShaderName(true, TEXT("FNoLightMapPolicy"), 0, bMobileHDR, false, ShaderNameStr);
 			ShaderTypeNamesAndDescriptions.FindOrAdd(FLocalVertexFactoryName)
-				.Add(FRepresentativeShaderInfo(ERepresentativeShader::StaticMesh, TBasePassForForwardShadingVSFNoLightMapPolicyName,
+				.Add(FRepresentativeShaderInfo(ERepresentativeShader::StaticMesh, FName(ShaderNameStr),
 					FString::Printf(TEXT("Mobile base pass vertex shader%s"), DescSuffix)));
 
 			ShaderTypeNamesAndDescriptions.FindOrAdd(FGPUFactoryName)
-				.Add(FRepresentativeShaderInfo(ERepresentativeShader::SkeletalMesh, TBasePassForForwardShadingVSFNoLightMapPolicyName,
+				.Add(FRepresentativeShaderInfo(ERepresentativeShader::SkeletalMesh, FName(ShaderNameStr),
 					FString::Printf(TEXT("Mobile base pass vertex shader%s"), DescSuffix)));
 		}
 		else
@@ -371,17 +388,10 @@ void FMaterialStatsUtils::GetRepresentativeShaderTypesAndDescriptions(TMap<FName
 				{
 					// distance field shadows only shaders
 					{
-						static const FName Name_HDRLinear64 = bOnlySkyPermutation ? 
-							TEXT("TMobileBasePassPSFMobileDistanceFieldShadowsAndLQLightMapPolicy0HDRLinear64Skylight") : 
-							TEXT("TMobileBasePassPSFMobileDistanceFieldShadowsAndLQLightMapPolicy0HDRLinear64");
-						static const FName Name_LDRGamma32 = bOnlySkyPermutation ? 
-							TEXT("TMobileBasePassPSFMobileDistanceFieldShadowsAndLQLightMapPolicy0LDRGamma32Skylight") : 
-							TEXT("TMobileBasePassPSFMobileDistanceFieldShadowsAndLQLightMapPolicy0LDRGamma32");
-						static const FName ShaderName = bMobileHDR ? Name_HDRLinear64 : Name_LDRGamma32;
-
+						MobileBasePassShaderName(false, TEXT("FMobileDistanceFieldShadowsAndLQLightMapPolicy"), 0, bMobileHDR, bOnlySkyPermutation, ShaderNameStr);
 						const FString Description = FString::Printf(TEXT("Mobile base pass shader with distance field shadows%s"), DescSuffix);
 						ShaderTypeNamesAndDescriptions.FindOrAdd(FLocalVertexFactoryName)
-							.Add(FRepresentativeShaderInfo(ERepresentativeShader::StationarySurface, ShaderName, Description));
+							.Add(FRepresentativeShaderInfo(ERepresentativeShader::StationarySurface, FName(ShaderNameStr), Description));
 					}
 
 					static auto* CVarAllowDistanceFieldShadowsAndCSM = IConsoleManager::Get().FindTConsoleVariableDataInt(TEXT("r.Mobile.EnableStaticAndCSMShadowReceivers"));
@@ -390,33 +400,18 @@ void FMaterialStatsUtils::GetRepresentativeShaderTypesAndDescriptions(TMap<FName
 					{
 						// distance field shadows & CSM shaders
 						{
-							static const FName Name_HDRLinear64 = bOnlySkyPermutation ? 
-								TEXT("TMobileBasePassPSFMobileDistanceFieldShadowsLightMapAndCSMLightingPolicy0HDRLinear64Skylight") : 
-								TEXT("TMobileBasePassPSFMobileDistanceFieldShadowsLightMapAndCSMLightingPolicy0HDRLinear64");
-							static const FName Name_LDRGamma32 = bOnlySkyPermutation ? 
-								TEXT("TMobileBasePassPSFMobileDistanceFieldShadowsLightMapAndCSMLightingPolicy0LDRGamma32Skylight") : 
-								TEXT("TMobileBasePassPSFMobileDistanceFieldShadowsLightMapAndCSMLightingPolicy0LDRGamma32");
-							static const FName ShaderName = bMobileHDR ? Name_HDRLinear64 : Name_LDRGamma32;
-
+							MobileBasePassShaderName(false, TEXT("FMobileDistanceFieldShadowsLightMapAndCSMLightingPolicy"), 0, bMobileHDR, bOnlySkyPermutation, ShaderNameStr);
 							const FString Description = FString::Printf(TEXT("Mobile base pass shader with distance field shadows and CSM%s"), DescSuffix);
 							ShaderTypeNamesAndDescriptions.FindOrAdd(FLocalVertexFactoryName)
-								.Add(FRepresentativeShaderInfo(ERepresentativeShader::StationarySurfaceCSM, ShaderName, Description));
+								.Add(FRepresentativeShaderInfo(ERepresentativeShader::StationarySurfaceCSM, FName(ShaderNameStr), Description));
 						}
 
 						if (bPointLights) // add point lights shaders + distance field shadows
 						{
-							static const FName Name_HDRLinear64_NLights = bOnlySkyPermutation ? 
-								TEXT("TMobileBasePassPSFMobileDistanceFieldShadowsAndLQLightMapPolicyINT32_MAXHDRLinear64Skylight") : 
-								TEXT("TMobileBasePassPSFMobileDistanceFieldShadowsAndLQLightMapPolicyINT32_MAXHDRLinear64");
-							static const FName Name_LDRGamma32_NLights = bOnlySkyPermutation ? 
-								TEXT("TMobileBasePassPSFMobileDistanceFieldShadowsAndLQLightMapPolicyINT32_MAXLDRGamma32Skylight") :
-								TEXT("TMobileBasePassPSFMobileDistanceFieldShadowsAndLQLightMapPolicyINT32_MAXLDRGamma32");
-
-							static const FName ShaderName = bMobileHDR ? Name_HDRLinear64_NLights : Name_LDRGamma32_NLights;
-
+							MobileBasePassShaderName(false, TEXT("FMobileDistanceFieldShadowsAndLQLightMapPolicy"), INT32_MAX, bMobileHDR, bOnlySkyPermutation, ShaderNameStr);
 							const FString Description = FString::Printf(TEXT("Mobile base pass shader with distance field shadows, CSM and %d point light(s) %s"), NumPointLights, DescSuffix);
 
-							FRepresentativeShaderInfo ShaderInfo = FRepresentativeShaderInfo(ERepresentativeShader::StationarySurfaceNPointLights, ShaderName, Description);
+							FRepresentativeShaderInfo ShaderInfo = FRepresentativeShaderInfo(ERepresentativeShader::StationarySurfaceNPointLights, FName(ShaderNameStr), Description);
 
 							ShaderTypeNamesAndDescriptions.FindOrAdd(FLocalVertexFactoryName).Add(ShaderInfo);
 								
@@ -425,61 +420,53 @@ void FMaterialStatsUtils::GetRepresentativeShaderTypesAndDescriptions(TMap<FName
 				}
 				else //no shadows & lightmapped
 				{
-					static const FName Name_HDRLinear64 = bOnlySkyPermutation ? 
-						TEXT("TMobileBasePassPSTLightMapPolicyLQ0HDRLinear64Skylight") :
-						TEXT("TMobileBasePassPSTLightMapPolicyLQ0HDRLinear64");
-					static const FName Name_LDRGamma32 = bOnlySkyPermutation ? 
-						TEXT("TMobileBasePassPSTLightMapPolicyLQ0LDRGamma32Skylight") : 
-						TEXT("TMobileBasePassPSTLightMapPolicyLQ0LDRGamma32");
-					static const FName TSlateMaterialShaderVStrueName = bMobileHDR ? Name_HDRLinear64 : Name_LDRGamma32;
-
+					MobileBasePassShaderName(false, TEXT("TLightMapPolicyLQ"), 0, bMobileHDR, bOnlySkyPermutation, ShaderNameStr);
 					ShaderTypeNamesAndDescriptions.FindOrAdd(FLocalVertexFactoryName)
-						.Add(FRepresentativeShaderInfo(ERepresentativeShader::StationarySurface, TSlateMaterialShaderVStrueName,
+						.Add(FRepresentativeShaderInfo(ERepresentativeShader::StationarySurface, FName(ShaderNameStr),
 							FString::Printf(TEXT("Mobile base pass shader with static lighting%s"), DescSuffix)));
 
 					if (bPointLights) // add point lights + lightmap
 					{
-						static const FName Name_HDRLinear64_NLights = bOnlySkyPermutation ? 
-							TEXT("TMobileBasePassPSTLightMapPolicyLQINT32_MAXHDRLinear64Skylight") :
-							TEXT("TMobileBasePassPSTLightMapPolicyLQINT32_MAXHDRLinear64");
-						static const FName Name_LDRGamma32_NLights = bOnlySkyPermutation ?  
-							TEXT("TMobileBasePassPSTLightMapPolicyLQINT32_MAXLDRGamma32Skylight") :
-							TEXT("TMobileBasePassPSTLightMapPolicyLQINT32_MAXLDRGamma32");
-
-						static const FName ShaderName = bMobileHDR ? Name_HDRLinear64_NLights : Name_LDRGamma32_NLights;
-
+						MobileBasePassShaderName(false, TEXT("TLightMapPolicyLQ"), INT32_MAX, bMobileHDR, bOnlySkyPermutation, ShaderNameStr);
 						const FString Description = FString::Printf(TEXT("Mobile base pass shader with static lighting and %d point light(s) %s"), NumPointLights, DescSuffix);
 
-						FRepresentativeShaderInfo ShaderInfo = FRepresentativeShaderInfo(ERepresentativeShader::StationarySurfaceNPointLights, ShaderName, Description);
+						FRepresentativeShaderInfo ShaderInfo = FRepresentativeShaderInfo(ERepresentativeShader::StationarySurfaceNPointLights, FName(ShaderNameStr), Description);
 
 						ShaderTypeNamesAndDescriptions.FindOrAdd(FLocalVertexFactoryName).Add(ShaderInfo);
 					}
 				}
 			}
 
-			// dynamically lit shader			
-			static const FName Name_HDRLinear64 = bOnlySkyPermutation ? 
-				TEXT("TMobileBasePassPSFNoLightmapPolicy0HDRLinear64Skylight") : 
-				TEXT("TMobileBasePassPSFNoLightmapPolicy0HDRLinear64");
-			static const FName Name_LDRGamma32 = bOnlySkyPermutation ? 
-				TEXT("TMobileBasePassPSFNoLightmapPolicy0LDRGamma32Skylight") : 
-				TEXT("TMobileBasePassPSFNoLightmapPolicy0LDRGamma32");
-			static const FName TBasePassForForwardShadingPSFSimpleDirectionalLightAndSHIndirectPolicy0Name = bMobileHDR ? Name_HDRLinear64 : Name_LDRGamma32;
+			// only one of these 2 shader types will be displayed
 			
+			// dynamically lit shader NoLightmapPolicy
+			MobileBasePassShaderName(false, TEXT("FNoLightmapPolicy"), 0, bMobileHDR, bOnlySkyPermutation, ShaderNameStr);
 			ShaderTypeNamesAndDescriptions.FindOrAdd(FLocalVertexFactoryName)
-				.Add(FRepresentativeShaderInfo(ERepresentativeShader::DynamicallyLitObject, TBasePassForForwardShadingPSFSimpleDirectionalLightAndSHIndirectPolicy0Name,
+				.Add(FRepresentativeShaderInfo(ERepresentativeShader::DynamicallyLitObject, FName(ShaderNameStr),
 				FString::Printf(TEXT("Mobile base pass shader with only dynamic lighting%s"), DescSuffix)));
 
-			static const FName Name_NoLM_HDRLinear64 = TEXT("TMobileBasePassVSFNoLightMapPolicyHDRLinear64");
-			static const FName Name_NoLM_LDRGamma32 = TEXT("TMobileBasePassVSFNoLightMapPolicyLDRGamma32");
-			static const FName TBasePassForForwardShadingVSFNoLightMapPolicyName = bMobileHDR ? Name_NoLM_HDRLinear64 : Name_NoLM_LDRGamma32;
+			MobileBasePassShaderName(true, TEXT("FNoLightMapPolicy"), 0, bMobileHDR, bOnlySkyPermutation, ShaderNameStr);
+			ShaderTypeNamesAndDescriptions.FindOrAdd(FLocalVertexFactoryName)
+				.Add(FRepresentativeShaderInfo(ERepresentativeShader::StaticMesh, FName(ShaderNameStr),
+				FString::Printf(TEXT("Mobile base pass vertex shader%s"), DescSuffix)));
+
+			ShaderTypeNamesAndDescriptions.FindOrAdd(FGPUFactoryName)
+				.Add(FRepresentativeShaderInfo(ERepresentativeShader::SkeletalMesh, FName(ShaderNameStr),
+				FString::Printf(TEXT("Mobile base pass vertex shader%s"), DescSuffix)));
 
+			// dynamically lit shader FMobileDirectionalLightAndCSMPolicy
+			MobileBasePassShaderName(false, TEXT("FMobileDirectionalLightAndCSMPolicy"), 0, bMobileHDR, bOnlySkyPermutation, ShaderNameStr);
 			ShaderTypeNamesAndDescriptions.FindOrAdd(FLocalVertexFactoryName)
-				.Add(FRepresentativeShaderInfo(ERepresentativeShader::StaticMesh, TBasePassForForwardShadingVSFNoLightMapPolicyName,
+				.Add(FRepresentativeShaderInfo(ERepresentativeShader::DynamicallyLitObject, FName(ShaderNameStr),
+				FString::Printf(TEXT("Mobile base pass shader with only dynamic lighting%s"), DescSuffix)));
+
+			MobileBasePassShaderName(true, TEXT("FMobileDirectionalLightAndCSMPolicy"), 0, bMobileHDR, bOnlySkyPermutation, ShaderNameStr);
+			ShaderTypeNamesAndDescriptions.FindOrAdd(FLocalVertexFactoryName)
+				.Add(FRepresentativeShaderInfo(ERepresentativeShader::StaticMesh, FName(ShaderNameStr),
 				FString::Printf(TEXT("Mobile base pass vertex shader%s"), DescSuffix)));
 
 			ShaderTypeNamesAndDescriptions.FindOrAdd(FGPUFactoryName)
-				.Add(FRepresentativeShaderInfo(ERepresentativeShader::SkeletalMesh, TBasePassForForwardShadingVSFNoLightMapPolicyName,
+				.Add(FRepresentativeShaderInfo(ERepresentativeShader::SkeletalMesh, FName(ShaderNameStr),
 				FString::Printf(TEXT("Mobile base pass vertex shader%s"), DescSuffix)));
 		}
 	}
@@ -500,7 +487,8 @@ void FMaterialStatsUtils::GetRepresentativeInstructionCounts(TArray<FShaderInstr
 	if (MaterialShaderMap && MaterialShaderMap->IsCompilationFinalized())
 	{
 		GetRepresentativeShaderTypesAndDescriptions(ShaderTypeNamesAndDescriptions, Target);
-
+		TStaticArray<bool, (int32)ERepresentativeShader::Num> bShaderTypeAdded(InPlace, false);
+		
 		if (Target->IsUIMaterial())
 		{
 			//for (const TPair<FName, FRepresentativeShaderInfo>& ShaderTypePair : ShaderTypeNamesAndDescriptions)
@@ -510,17 +498,21 @@ void FMaterialStatsUtils::GetRepresentativeInstructionCounts(TArray<FShaderInstr
 				for (int32 i = 0; i < DescriptionArray.Num(); ++i)
 				{
 					const FRepresentativeShaderInfo& ShaderInfo = DescriptionArray[i];
+					if (!bShaderTypeAdded[(int32)ShaderInfo.ShaderType])
+					{
+						FShaderType* ShaderType = FindShaderTypeByName(ShaderInfo.ShaderName);
+						check(ShaderType);
+						const int32 NumInstructions = MaterialShaderMap->GetMaxNumInstructionsForShader(ShaderType);
 
-					FShaderType* ShaderType = FindShaderTypeByName(ShaderInfo.ShaderName);
-					check(ShaderType);
-					const int32 NumInstructions = MaterialShaderMap->GetMaxNumInstructionsForShader(ShaderType);
+						FShaderInstructionsInfo Info;
+						Info.ShaderType = ShaderInfo.ShaderType;
+						Info.ShaderDescription = ShaderInfo.ShaderDescription;
+						Info.InstructionCount = NumInstructions;
 
-					FShaderInstructionsInfo Info;
-					Info.ShaderType = ShaderInfo.ShaderType;
-					Info.ShaderDescription = ShaderInfo.ShaderDescription;
-					Info.InstructionCount = NumInstructions;
+						Results.Push(Info);
 
-					Results.Push(Info);
+						bShaderTypeAdded[(int32)ShaderInfo.ShaderType] = true;
+					}
 				}
 			}
 		}
@@ -540,20 +532,24 @@ void FMaterialStatsUtils::GetRepresentativeInstructionCounts(TArray<FShaderInstr
 					for (int32 i = 0; i < DescriptionArray.Num(); ++i)
 					{
 						const FRepresentativeShaderInfo& ShaderInfo = DescriptionArray[i];
-
-						TShaderRef<FShader>* ShaderEntry = ShaderMap.Find(ShaderInfo.ShaderName);
-						if (ShaderEntry != nullptr)
+						if (!bShaderTypeAdded[(int32)ShaderInfo.ShaderType])
 						{
-							FShaderType* ShaderType = (*ShaderEntry).GetType();
+							TShaderRef<FShader>* ShaderEntry = ShaderMap.Find(ShaderInfo.ShaderName);
+							if (ShaderEntry != nullptr)
 							{
-								const int32 NumInstructions = MeshShaderMap->GetMaxNumInstructionsForShader(*MaterialShaderMap, ShaderType);
+								FShaderType* ShaderType = (*ShaderEntry).GetType();
+								{
+									const int32 NumInstructions = MeshShaderMap->GetMaxNumInstructionsForShader(*MaterialShaderMap, ShaderType);
+
+									FShaderInstructionsInfo Info;
+									Info.ShaderType = ShaderInfo.ShaderType;
+									Info.ShaderDescription = ShaderInfo.ShaderDescription;
+									Info.InstructionCount = NumInstructions;
 
-								FShaderInstructionsInfo Info;
-								Info.ShaderType = ShaderInfo.ShaderType;
-								Info.ShaderDescription = ShaderInfo.ShaderDescription;
-								Info.InstructionCount = NumInstructions;
+									Results.Push(Info);
 
-								Results.Push(Info);
+									bShaderTypeAdded[(int32)ShaderInfo.ShaderType] = true;
+								}
 							}
 						}
 					}
