Fix MaterialFunctionInstance parameter inheritance

diff --git a/Engine/Source/Editor/MaterialEditor/Private/MaterialInstanceEditor.cpp b/Engine/Source/Editor/MaterialEditor/Private/MaterialInstanceEditor.cpp
--- a/Engine/Source/Editor/MaterialEditor/Private/MaterialInstanceEditor.cpp
+++ b/Engine/Source/Editor/MaterialEditor/Private/MaterialInstanceEditor.cpp
@@ -289,8 +289,6 @@ void FMaterialInstanceEditor::InitEditorForMaterialFunction(UMaterialFunctionIns
 	check(InMaterialFunction);
 	MaterialFunctionOriginal = InMaterialFunction;
 
-	UMaterialFunctionInterface* Parent = InMaterialFunction->Parent;
-
 	// Working version of the function instance
 	MaterialFunctionInstance = (UMaterialFunctionInstance*)StaticDuplicateObject(InMaterialFunction, GetTransientPackage(), NAME_None, ~RF_Standalone, UMaterialFunctionInstance::StaticClass()); 
 	MaterialFunctionInstance->Parent = InMaterialFunction;
@@ -330,18 +328,6 @@ void FMaterialInstanceEditor::InitEditorForMaterialFunction(UMaterialFunctionIns
 		Expression->Function = NULL;
 		Expression->Material = FunctionMaterialProxy;
 
-		// The actual MFI we're adding may have multiple layers of inheritance, but for the editor, we only create a single UMaterial which represents all the parents
-		// So if any parent MFI has a parameter override, we need to bake that parameter override into the UMaterial's parameter expression here
-		FMaterialParameterMetadata ParameterMeta;
-		if (Parent && Expression->GetParameterValue(ParameterMeta))
-		{
-			const FName ParameterName = Expression->GetParameterName();
-			if (Parent->GetParameterOverrideValue(ParameterMeta.Value.Type, ParameterName, ParameterMeta))
-			{
-				Expression->SetParameterValue(ParameterName, ParameterMeta);
-			}
-		}
-
 		if (UMaterialExpressionFunctionOutput* FunctionOutput = Cast<UMaterialExpressionFunctionOutput>(Expression))
 		{
 			FirstOutput = FunctionOutput;
@@ -358,14 +344,49 @@ void FMaterialInstanceEditor::InitEditorForMaterialFunction(UMaterialFunctionIns
 		FirstOutput->ConnectToPreviewMaterial(FunctionMaterialProxy, 0);
 	}
 
-	FMaterialUpdateContext UpdateContext(FMaterialUpdateContext::EOptions::SyncWithRenderingThread);
-	UpdateContext.AddMaterial(FunctionMaterialProxy);
-	FunctionMaterialProxy->PreEditChange(NULL);
-	FunctionMaterialProxy->PostEditChange();
+	{
+		FMaterialUpdateContext UpdateContext(FMaterialUpdateContext::EOptions::SyncWithRenderingThread);
+		UpdateContext.AddMaterial(FunctionMaterialProxy);
+		FunctionMaterialProxy->PreEditChange(NULL);
+		FunctionMaterialProxy->PostEditChange();
+	}
+
+	UMaterialFunctionInstance* ParentFunctionInstance = Cast<UMaterialFunctionInstance>(InMaterialFunction->Parent);
+	UMaterialInterface* FunctionInstanceParent = FunctionMaterialProxy;
+	if (ParentFunctionInstance)
+	{
+		// If our FunctionInstance inherits from *another* FunctionInstance, the parent may have overriden certain parameters
+		// These overriden values should be the default value as far as our FunctionInstance is concerned
+		// To model this, we create a MI that will hold these overriden values.  So our editor material hierarchy will look like this
+		// * UMaterial - includes the UMaterialExpressions, copied from our base UMaterialFunction
+		// * UMaterialInstance - we're creating this here, holds all parameter values overriden by our parent UMaterialFunctionInstance(s)
+		// * UMaterialInstance - will be created below, this is the object/proxy we'll be editing and potentially setting parameters on
+		UMaterialInstanceConstant* FunctionMaterialInstanceProxy = NewObject<UMaterialInstanceConstant>(GetTransientPackage(), NAME_None, RF_Transactional);
+		FunctionMaterialInstanceProxy->SetParentEditorOnly(FunctionMaterialProxy);
+		{
+			FMaterialInstanceParameterUpdateContext UpdateContext(FunctionMaterialInstanceProxy);
+			TArray<FMaterialParameterInfo> ParameterInfos;
+			TArray<FGuid> ParameterIds;
+			for (int32 ParameterTypeIndex = 0; ParameterTypeIndex < NumMaterialParameterTypes; ++ParameterTypeIndex)
+			{
+				const EMaterialParameterType ParameterType = (EMaterialParameterType)ParameterTypeIndex;
+				FunctionMaterialProxy->GetAllParameterInfoOfType(ParameterType, ParameterInfos, ParameterIds);
+				for (const FMaterialParameterInfo& ParameterInfo : ParameterInfos)
+				{
+					FMaterialParameterMetadata ParameterMeta;
+					if (ParentFunctionInstance->GetParameterOverrideValue(ParameterType, ParameterInfo.Name, ParameterMeta))
+					{
+						UpdateContext.SetParameterValueEditorOnly(ParameterInfo, ParameterMeta, EMaterialSetParameterValueFlags::SetCurveAtlas);
+					}
+				}
+			}
+		}
+		FunctionInstanceParent = FunctionMaterialInstanceProxy;
+	}
 
 	// Preview instance for function expressions
 	FunctionInstanceProxy = NewObject<UMaterialInstanceConstant>(GetTransientPackage(), NAME_None, RF_Transactional);
-	FunctionInstanceProxy->SetParentEditorOnly(FunctionMaterialProxy);
+	FunctionInstanceProxy->SetParentEditorOnly(FunctionInstanceParent);
 
 	MaterialFunctionInstance->OverrideMaterialInstanceParameterValues(FunctionInstanceProxy);
 	FunctionInstanceProxy->PreEditChange(NULL);
diff --git a/Engine/Source/Runtime/Engine/Classes/Materials/MaterialFunctionInstance.h b/Engine/Source/Runtime/Engine/Classes/Materials/MaterialFunctionInstance.h
--- a/Engine/Source/Runtime/Engine/Classes/Materials/MaterialFunctionInstance.h
+++ b/Engine/Source/Runtime/Engine/Classes/Materials/MaterialFunctionInstance.h
@@ -173,7 +173,7 @@ class UMaterialFunctionInstance : public UMaterialFunctionInterface
 	}
 
 public:
-	virtual bool GetParameterOverrideValue(EMaterialParameterType Type, const FName& ParameterName, FMaterialParameterMetadata& OutValue) const override;
+	virtual bool GetParameterOverrideValue(EMaterialParameterType Type, const FName& ParameterName, FMaterialParameterMetadata& OutValue, FMFRecursionGuard RecursionGuard = FMFRecursionGuard()) const override;
 	
 #endif // WITH_EDITOR
 	//~ End UMaterialFunctionInterface interface
diff --git a/Engine/Source/Runtime/Engine/Classes/Materials/MaterialFunctionInterface.h b/Engine/Source/Runtime/Engine/Classes/Materials/MaterialFunctionInterface.h
--- a/Engine/Source/Runtime/Engine/Classes/Materials/MaterialFunctionInterface.h
+++ b/Engine/Source/Runtime/Engine/Classes/Materials/MaterialFunctionInterface.h
@@ -28,6 +28,8 @@ enum class EMaterialFunctionUsage : uint8
 	MaterialLayerBlend
 };
 
+using FMFRecursionGuard = TMaterialRecursionGuard<class UMaterialFunctionInterface>;
+
 /**
  * A Material Function is a collection of material expressions that can be reused in different materials
  */
@@ -314,7 +316,7 @@ public:
 #endif // WITH_EDITORONLY_DATA
 
 #if WITH_EDITOR
-	virtual bool GetParameterOverrideValue(EMaterialParameterType Type, const FName& ParameterName, FMaterialParameterMetadata& OutValue) const;
+	virtual bool GetParameterOverrideValue(EMaterialParameterType Type, const FName& ParameterName, FMaterialParameterMetadata& OutValue, FMFRecursionGuard RecursionGuard = FMFRecursionGuard()) const;
 
 	ENGINE_API bool OverrideNamedScalarParameter(const FHashedMaterialParameterInfo& ParameterInfo, float& OutValue);
 	ENGINE_API bool OverrideNamedVectorParameter(const FHashedMaterialParameterInfo& ParameterInfo, FLinearColor& OutValue);
diff --git a/Engine/Source/Runtime/Engine/Classes/Materials/MaterialInterface.h b/Engine/Source/Runtime/Engine/Classes/Materials/MaterialInterface.h
--- a/Engine/Source/Runtime/Engine/Classes/Materials/MaterialInterface.h
+++ b/Engine/Source/Runtime/Engine/Classes/Materials/MaterialInterface.h
@@ -173,39 +173,7 @@ struct FMaterialTextureInfo
 	ENGINE_API bool IsValid(bool bCheckTextureIndex = false) const; 
 };
 
-struct TMicRecursionGuard
-{
-	inline TMicRecursionGuard() = default;
-	inline TMicRecursionGuard(const TMicRecursionGuard& Parent)
-		: MaterialInterface(nullptr)
-		, PreviousLink(&Parent)
-	{
-	}
-
-	inline void Set(class UMaterialInterface const* InMaterialInterface)
-	{
-		check(MaterialInterface == nullptr);
-		MaterialInterface = InMaterialInterface;
-	}
-
-	inline bool Contains(class UMaterialInterface const* InMaterialInterface)
-	{
-		TMicRecursionGuard const* Link = this;
-		do
-		{
-			if (Link->MaterialInterface == InMaterialInterface)
-			{
-				return true;
-			}
-			Link = Link->PreviousLink;
-		} while (Link);
-		return false;
-	}
-
-private:
-	void const* MaterialInterface = nullptr;
-	TMicRecursionGuard const* PreviousLink = nullptr;
-};
+using TMicRecursionGuard = TMaterialRecursionGuard<class UMaterialInterface>;
 
 /** Holds information about a hierarchy of materials */
 struct FMaterialInheritanceChain
diff --git a/Engine/Source/Runtime/Engine/Private/Materials/HLSLMaterialTranslator.cpp b/Engine/Source/Runtime/Engine/Private/Materials/HLSLMaterialTranslator.cpp
--- a/Engine/Source/Runtime/Engine/Private/Materials/HLSLMaterialTranslator.cpp
+++ b/Engine/Source/Runtime/Engine/Private/Materials/HLSLMaterialTranslator.cpp
@@ -4073,14 +4073,22 @@ int32 FHLSLMaterialTranslator::AccessCollectionParameter(UMaterialParameterColle
 
 bool FHLSLMaterialTranslator::GetParameterOverrideValueForCurrentFunction(EMaterialParameterType ParameterType, FName ParameterName, FMaterialParameterMetadata& OutResult) const
 {
-	const FMaterialFunctionCompileState* CurrentFunctionState = FunctionStacks[ShaderFrequency].Last();
+	const TArray<FMaterialFunctionCompileState*>& FunctionStack = FunctionStacks[ShaderFrequency];
+
+	// Give every function in the callstack on opportunity to override the parameter value
+	// Parameters in outer functions take priority
+	// For example, if a layer instance calls a function instance that includes an overriden parameter, we want to use the value from the layer instance rather than the function instance
 	bool bResult = false;
-	if (CurrentFunctionState->FunctionCall)
+	for (const FMaterialFunctionCompileState* FunctionState : FunctionStack)
 	{
-		const UMaterialFunctionInterface* CurrentFunction = CurrentFunctionState->FunctionCall->MaterialFunction;
+		const UMaterialFunctionInterface* CurrentFunction = (FunctionState && FunctionState->FunctionCall) ? FunctionState->FunctionCall->MaterialFunction : nullptr;
 		if (CurrentFunction)
 		{
-			bResult = CurrentFunction->GetParameterOverrideValue(ParameterType, ParameterName, OutResult);
+			if (CurrentFunction->GetParameterOverrideValue(ParameterType, ParameterName, OutResult))
+			{
+				bResult = true;
+				break;
+			}
 		}
 	}
 
diff --git a/Engine/Source/Runtime/Engine/Private/Materials/MaterialExpressions.cpp b/Engine/Source/Runtime/Engine/Private/Materials/MaterialExpressions.cpp
--- a/Engine/Source/Runtime/Engine/Private/Materials/MaterialExpressions.cpp
+++ b/Engine/Source/Runtime/Engine/Private/Materials/MaterialExpressions.cpp
@@ -12468,7 +12468,7 @@ void UMaterialFunctionInterface::ForceRecompileForRendering(FMaterialUpdateConte
 	}
 }
 
-bool UMaterialFunctionInterface::GetParameterOverrideValue(EMaterialParameterType Type, const FName& ParameterName, FMaterialParameterMetadata& OutValue) const
+bool UMaterialFunctionInterface::GetParameterOverrideValue(EMaterialParameterType Type, const FName& ParameterName, FMaterialParameterMetadata& OutValue, FMFRecursionGuard RecursionGuard) const
 {
 	return false;
 }
@@ -13643,7 +13643,7 @@ bool UMaterialFunctionInstance::HasFlippedCoordinates() const
 	return Parent ? Parent->HasFlippedCoordinates() : false;
 }
 
-bool UMaterialFunctionInstance::GetParameterOverrideValue(EMaterialParameterType Type, const FName& ParameterName, FMaterialParameterMetadata& OutResult) const
+bool UMaterialFunctionInstance::GetParameterOverrideValue(EMaterialParameterType Type, const FName& ParameterName, FMaterialParameterMetadata& OutResult, FMFRecursionGuard RecursionGuard) const
 {
 	const FMemoryImageMaterialParameterInfo ParameterInfo(ParameterName);
 
@@ -13659,6 +13659,13 @@ bool UMaterialFunctionInstance::GetParameterOverrideValue(EMaterialParameterType
 	case EMaterialParameterType::StaticComponentMask: bResult = GameThread_GetParameterValue(StaticComponentMaskParameterValues, ParameterInfo, OutResult); break;
 	default: checkNoEntry(); break;
 	}
+
+	if (!bResult && Parent && !RecursionGuard.Contains(this))
+	{
+		RecursionGuard.Set(this);
+		bResult = Parent->GetParameterOverrideValue(Type, ParameterName, OutResult, RecursionGuard);
+	}
+
 	return bResult;
 }
 
diff --git a/Engine/Source/Runtime/Engine/Private/ShaderDerivedDataVersion.h b/Engine/Source/Runtime/Engine/Private/ShaderDerivedDataVersion.h
--- a/Engine/Source/Runtime/Engine/Private/ShaderDerivedDataVersion.h
+++ b/Engine/Source/Runtime/Engine/Private/ShaderDerivedDataVersion.h
@@ -11,4 +11,4 @@
 // In case of merge conflicts with DDC versions, you MUST generate a new GUID and set this new
 // guid as version
 #define GLOBALSHADERMAP_DERIVEDDATA_VER			TEXT("1D2DE12F326C4BB38DBC0B2CBE348802")
-#define MATERIALSHADERMAP_DERIVEDDATA_VER		TEXT("DB3996BDDA7D410FA8F3F4B74B4A4404")
+#define MATERIALSHADERMAP_DERIVEDDATA_VER		TEXT("01DF57ADAE184DB89CB6C43A8A421B57")
diff --git a/Engine/Source/Runtime/Engine/Public/MaterialShared.h b/Engine/Source/Runtime/Engine/Public/MaterialShared.h
--- a/Engine/Source/Runtime/Engine/Public/MaterialShared.h
+++ b/Engine/Source/Runtime/Engine/Public/MaterialShared.h
@@ -3414,3 +3414,38 @@ inline bool ShouldIncludeMaterialInDefaultOpaquePass(const FMaterial& Material)
 	return !Material.IsSky()
 		&& !Material.GetShadingModels().HasShadingModel(MSM_SingleLayerWater);
 }
+
+template<typename T>
+struct TMaterialRecursionGuard
+{
+	inline TMaterialRecursionGuard() = default;
+	inline TMaterialRecursionGuard(const TMaterialRecursionGuard& Parent)
+		: Value(nullptr)
+		, PreviousLink(&Parent)
+	{
+	}
+
+	inline void Set(const T* InValue)
+	{
+		check(Value == nullptr);
+		Value = InValue;
+	}
+
+	inline bool Contains(const T* InValue)
+	{
+		TMaterialRecursionGuard const* Link = this;
+		do
+		{
+			if (Link->Value == InValue)
+			{
+				return true;
+			}
+			Link = Link->PreviousLink;
+		} while (Link);
+		return false;
+	}
+
+private:
+	const T* Value = nullptr;
+	TMaterialRecursionGuard const* PreviousLink = nullptr;
+};
