Don't crash if layer function outputs haven't been initialized yet

diff --git a/Engine/Source/Editor/MaterialEditor/Private/MaterialEditorUtilities.cpp b/Engine/Source/Editor/MaterialEditor/Private/MaterialEditorUtilities.cpp
--- a/Engine/Source/Editor/MaterialEditor/Private/MaterialEditorUtilities.cpp
+++ b/Engine/Source/Editor/MaterialEditor/Private/MaterialEditorUtilities.cpp
@@ -500,7 +500,7 @@ void FMaterialEditorUtilities::GetVisibleMaterialParametersFromExpression(
 	}
 	else if (FunctionCallExpression)
 	{
-		if (FunctionCallExpression->MaterialFunction)
+		if (FunctionCallExpression->MaterialFunction && FunctionCallExpression->FunctionOutputs.IsValidIndex(MaterialExpressionKey.OutputIndex))
 		{			
 			for (int32 FunctionCallIndex = 0; FunctionCallIndex < FunctionStack.Num(); FunctionCallIndex++)
 			{
@@ -531,7 +531,8 @@ void FMaterialEditorUtilities::GetVisibleMaterialParametersFromExpression(
 			{
 				for (auto& Layer : LayersExpression->LayerCallers)
 				{
-					if (Layer && Layer->MaterialFunction)
+					// Possible that Layer->FunctionOutputs will be empty if this is a newly create layer
+					if (Layer && Layer->MaterialFunction && Layer->FunctionOutputs.IsValidIndex(MaterialExpressionKey.OutputIndex))
 					{
 						TUniquePtr<FGetVisibleMaterialParametersFunctionState> NewFunctionState = MakeUnique<FGetVisibleMaterialParametersFunctionState>(Layer);
 						FunctionStack.Push(NewFunctionState.Get());
@@ -544,7 +545,7 @@ void FMaterialEditorUtilities::GetVisibleMaterialParametersFromExpression(
 
 				for (auto& Blend : LayersExpression->BlendCallers)
 				{
-					if (Blend && Blend->MaterialFunction)
+					if (Blend && Blend->MaterialFunction && Blend->FunctionOutputs.IsValidIndex(MaterialExpressionKey.OutputIndex))
 					{
 						TUniquePtr<FGetVisibleMaterialParametersFunctionState> NewFunctionState = MakeUnique<FGetVisibleMaterialParametersFunctionState>(Blend);
 						FunctionStack.Push(NewFunctionState.Get());
