Fix uninitialized pointers in material editor, that can happen if base material is reloaded while MI is open

diff --git a/Engine/Source/Editor/MaterialEditor/Private/SMaterialLayersFunctionsTree.cpp b/Engine/Source/Editor/MaterialEditor/Private/SMaterialLayersFunctionsTree.cpp
--- a/Engine/Source/Editor/MaterialEditor/Private/SMaterialLayersFunctionsTree.cpp
+++ b/Engine/Source/Editor/MaterialEditor/Private/SMaterialLayersFunctionsTree.cpp
@@ -289,7 +289,7 @@ void SMaterialLayersFunctionsInstanceTreeItem::Construct(const FArguments& InArg
 		for (TSharedPtr<FSortedParamData> AssetChild : AssetChildren)
 		{
 			TSharedPtr<SBox> ThumbnailBox;
-			UObject* AssetObject;
+			UObject* AssetObject = nullptr;
 			AssetChild->ParameterHandle->GetValue(AssetObject);
 			int32 PreviewIndex = INDEX_NONE;
 			int32 ThumbnailIndex = INDEX_NONE;
@@ -1473,18 +1473,20 @@ void SMaterialLayersFunctionsInstanceTree::CreateGroupsWidget()
 					ChildProperty->ParameterInfo.Association = EMaterialParameterAssociation::LayerParameter;
 					ChildProperty->NodeKey = FString::FromInt(ChildProperty->ParameterInfo.Index) + FString::FromInt(ChildProperty->ParameterInfo.Association);
 
-					UObject* AssetObject;
-					ChildProperty->ParameterHandle->GetValue(AssetObject);
-					if (AssetObject)
 					{
-						if (MaterialEditorInstance->StoredLayerPreviews[LayerChildren - 1] == nullptr)
+						UObject* AssetObject = nullptr;
+						ChildProperty->ParameterHandle->GetValue(AssetObject);
+						if (AssetObject)
 						{
-							MaterialEditorInstance->StoredLayerPreviews[LayerChildren - 1] = (NewObject<UMaterialInstanceConstant>(MaterialEditorInstance, NAME_None));
-						}
-						UMaterialInterface* EditedMaterial = Cast<UMaterialFunctionInterface>(AssetObject)->GetPreviewMaterial();
-						if (MaterialEditorInstance->StoredLayerPreviews[LayerChildren - 1] && MaterialEditorInstance->StoredLayerPreviews[LayerChildren - 1]->Parent != EditedMaterial)
-						{
-							MaterialEditorInstance->StoredLayerPreviews[LayerChildren - 1]->SetParentEditorOnly(EditedMaterial);
+							if (MaterialEditorInstance->StoredLayerPreviews[LayerChildren - 1] == nullptr)
+							{
+								MaterialEditorInstance->StoredLayerPreviews[LayerChildren - 1] = (NewObject<UMaterialInstanceConstant>(MaterialEditorInstance, NAME_None));
+							}
+							UMaterialInterface* EditedMaterial = Cast<UMaterialFunctionInterface>(AssetObject)->GetPreviewMaterial();
+							if (MaterialEditorInstance->StoredLayerPreviews[LayerChildren - 1] && MaterialEditorInstance->StoredLayerPreviews[LayerChildren - 1]->Parent != EditedMaterial)
+							{
+								MaterialEditorInstance->StoredLayerPreviews[LayerChildren - 1]->SetParentEditorOnly(EditedMaterial);
+							}
 						}
 					}
 
@@ -1503,17 +1505,20 @@ void SMaterialLayersFunctionsInstanceTree::CreateGroupsWidget()
 							ChildProperty->ParameterInfo.Index = Counter;
 							ChildProperty->ParameterInfo.Association = EMaterialParameterAssociation::BlendParameter;
 							ChildProperty->NodeKey = FString::FromInt(ChildProperty->ParameterInfo.Index) + FString::FromInt(ChildProperty->ParameterInfo.Association);
-							ChildProperty->ParameterHandle->GetValue(AssetObject);
-							if (AssetObject)
 							{
-								if (MaterialEditorInstance->StoredBlendPreviews[Counter] == nullptr)
-								{
-									MaterialEditorInstance->StoredBlendPreviews[Counter] = (NewObject<UMaterialInstanceConstant>(MaterialEditorInstance, NAME_None));
-								}
-								UMaterialInterface* EditedMaterial = Cast<UMaterialFunctionInterface>(AssetObject)->GetPreviewMaterial();
-								if (MaterialEditorInstance->StoredBlendPreviews[Counter] && MaterialEditorInstance->StoredBlendPreviews[Counter]->Parent != EditedMaterial)
+								UObject* AssetObject = nullptr;
+								ChildProperty->ParameterHandle->GetValue(AssetObject);
+								if (AssetObject)
 								{
-									MaterialEditorInstance->StoredBlendPreviews[Counter]->SetParentEditorOnly(EditedMaterial);
+									if (MaterialEditorInstance->StoredBlendPreviews[Counter] == nullptr)
+									{
+										MaterialEditorInstance->StoredBlendPreviews[Counter] = (NewObject<UMaterialInstanceConstant>(MaterialEditorInstance, NAME_None));
+									}
+									UMaterialInterface* EditedMaterial = Cast<UMaterialFunctionInterface>(AssetObject)->GetPreviewMaterial();
+									if (MaterialEditorInstance->StoredBlendPreviews[Counter] && MaterialEditorInstance->StoredBlendPreviews[Counter]->Parent != EditedMaterial)
+									{
+										MaterialEditorInstance->StoredBlendPreviews[Counter]->SetParentEditorOnly(EditedMaterial);
+									}
 								}
 							}
 							LayerProperties.Last()->Children.Add(ChildProperty);
@@ -1533,17 +1538,20 @@ void SMaterialLayersFunctionsInstanceTree::CreateGroupsWidget()
 							ChildProperty->ParameterInfo.Index = Counter;
 							ChildProperty->ParameterInfo.Association = EMaterialParameterAssociation::LayerParameter;
 							ChildProperty->NodeKey = FString::FromInt(ChildProperty->ParameterInfo.Index) + FString::FromInt(ChildProperty->ParameterInfo.Association);
-							ChildProperty->ParameterHandle->GetValue(AssetObject);
-							if (AssetObject)
 							{
-								if (MaterialEditorInstance->StoredLayerPreviews[Counter] == nullptr)
-								{
-									MaterialEditorInstance->StoredLayerPreviews[Counter] = (NewObject<UMaterialInstanceConstant>(MaterialEditorInstance, NAME_None));
-								}
-								UMaterialInterface* EditedMaterial = Cast<UMaterialFunctionInterface>(AssetObject)->GetPreviewMaterial();
-								if (MaterialEditorInstance->StoredLayerPreviews[Counter] && MaterialEditorInstance->StoredLayerPreviews[Counter]->Parent != EditedMaterial)
+								UObject* AssetObject = nullptr;
+								ChildProperty->ParameterHandle->GetValue(AssetObject);
+								if (AssetObject)
 								{
-									MaterialEditorInstance->StoredLayerPreviews[Counter]->SetParentEditorOnly(EditedMaterial);
+									if (MaterialEditorInstance->StoredLayerPreviews[Counter] == nullptr)
+									{
+										MaterialEditorInstance->StoredLayerPreviews[Counter] = (NewObject<UMaterialInstanceConstant>(MaterialEditorInstance, NAME_None));
+									}
+									UMaterialInterface* EditedMaterial = Cast<UMaterialFunctionInterface>(AssetObject)->GetPreviewMaterial();
+									if (MaterialEditorInstance->StoredLayerPreviews[Counter] && MaterialEditorInstance->StoredLayerPreviews[Counter]->Parent != EditedMaterial)
+									{
+										MaterialEditorInstance->StoredLayerPreviews[Counter]->SetParentEditorOnly(EditedMaterial);
+									}
 								}
 							}
 							LayerProperties.Last()->Children.Add(ChildProperty);
@@ -2804,18 +2812,20 @@ void SMaterialLayersFunctionsMaterialTree::CreateGroupsWidget()
 					ChildProperty->ParameterInfo.Association = EMaterialParameterAssociation::LayerParameter;
 					ChildProperty->NodeKey = FString::FromInt(ChildProperty->ParameterInfo.Index) + FString::FromInt(ChildProperty->ParameterInfo.Association);
 
-					UObject* AssetObject;
-					ChildProperty->ParameterHandle->GetValue(AssetObject);
-					if (AssetObject)
 					{
-						if (MaterialEditorInstance->StoredLayerPreviews[LayerChildren - 1] == nullptr)
+						UObject* AssetObject = nullptr;
+						ChildProperty->ParameterHandle->GetValue(AssetObject);
+						if (AssetObject)
 						{
-							MaterialEditorInstance->StoredLayerPreviews[LayerChildren - 1] = (NewObject<UMaterialInstanceConstant>(MaterialEditorInstance, NAME_None));
-						}
-						UMaterialInterface* EditedMaterial = Cast<UMaterialFunctionInterface>(AssetObject)->GetPreviewMaterial();
-						if (MaterialEditorInstance->StoredLayerPreviews[LayerChildren - 1] && MaterialEditorInstance->StoredLayerPreviews[LayerChildren - 1]->Parent != EditedMaterial)
-						{
-							MaterialEditorInstance->StoredLayerPreviews[LayerChildren - 1]->SetParentEditorOnly(EditedMaterial);
+							if (MaterialEditorInstance->StoredLayerPreviews[LayerChildren - 1] == nullptr)
+							{
+								MaterialEditorInstance->StoredLayerPreviews[LayerChildren - 1] = (NewObject<UMaterialInstanceConstant>(MaterialEditorInstance, NAME_None));
+							}
+							UMaterialInterface* EditedMaterial = Cast<UMaterialFunctionInterface>(AssetObject)->GetPreviewMaterial();
+							if (MaterialEditorInstance->StoredLayerPreviews[LayerChildren - 1] && MaterialEditorInstance->StoredLayerPreviews[LayerChildren - 1]->Parent != EditedMaterial)
+							{
+								MaterialEditorInstance->StoredLayerPreviews[LayerChildren - 1]->SetParentEditorOnly(EditedMaterial);
+							}
 						}
 					}
 
@@ -2834,17 +2844,20 @@ void SMaterialLayersFunctionsMaterialTree::CreateGroupsWidget()
 							ChildProperty->ParameterInfo.Index = Counter;
 							ChildProperty->ParameterInfo.Association = EMaterialParameterAssociation::BlendParameter;
 							ChildProperty->NodeKey = FString::FromInt(ChildProperty->ParameterInfo.Index) + FString::FromInt(ChildProperty->ParameterInfo.Association);
-							ChildProperty->ParameterHandle->GetValue(AssetObject);
-							if (AssetObject)
 							{
-								if (MaterialEditorInstance->StoredBlendPreviews[Counter] == nullptr)
-								{
-									MaterialEditorInstance->StoredBlendPreviews[Counter] = (NewObject<UMaterialInstanceConstant>(MaterialEditorInstance, NAME_None));
-								}
-								UMaterialInterface* EditedMaterial = Cast<UMaterialFunctionInterface>(AssetObject)->GetPreviewMaterial();
-								if (MaterialEditorInstance->StoredBlendPreviews[Counter] && MaterialEditorInstance->StoredBlendPreviews[Counter]->Parent != EditedMaterial)
+								UObject* AssetObject = nullptr;
+								ChildProperty->ParameterHandle->GetValue(AssetObject);
+								if (AssetObject)
 								{
-									MaterialEditorInstance->StoredBlendPreviews[Counter]->SetParentEditorOnly(EditedMaterial);
+									if (MaterialEditorInstance->StoredBlendPreviews[Counter] == nullptr)
+									{
+										MaterialEditorInstance->StoredBlendPreviews[Counter] = (NewObject<UMaterialInstanceConstant>(MaterialEditorInstance, NAME_None));
+									}
+									UMaterialInterface* EditedMaterial = Cast<UMaterialFunctionInterface>(AssetObject)->GetPreviewMaterial();
+									if (MaterialEditorInstance->StoredBlendPreviews[Counter] && MaterialEditorInstance->StoredBlendPreviews[Counter]->Parent != EditedMaterial)
+									{
+										MaterialEditorInstance->StoredBlendPreviews[Counter]->SetParentEditorOnly(EditedMaterial);
+									}
 								}
 							}
 							LayerProperties.Last()->Children.Add(ChildProperty);
@@ -2864,17 +2877,20 @@ void SMaterialLayersFunctionsMaterialTree::CreateGroupsWidget()
 							ChildProperty->ParameterInfo.Index = Counter;
 							ChildProperty->ParameterInfo.Association = EMaterialParameterAssociation::LayerParameter;
 							ChildProperty->NodeKey = FString::FromInt(ChildProperty->ParameterInfo.Index) + FString::FromInt(ChildProperty->ParameterInfo.Association);
-							ChildProperty->ParameterHandle->GetValue(AssetObject);
-							if (AssetObject)
 							{
-								if (MaterialEditorInstance->StoredLayerPreviews[Counter] == nullptr)
-								{
-									MaterialEditorInstance->StoredLayerPreviews[Counter] = (NewObject<UMaterialInstanceConstant>(MaterialEditorInstance, NAME_None));
-								}
-								UMaterialInterface* EditedMaterial = Cast<UMaterialFunctionInterface>(AssetObject)->GetPreviewMaterial();
-								if (MaterialEditorInstance->StoredLayerPreviews[Counter] && MaterialEditorInstance->StoredLayerPreviews[Counter]->Parent != EditedMaterial)
+								UObject* AssetObject = nullptr;
+								ChildProperty->ParameterHandle->GetValue(AssetObject);
+								if (AssetObject)
 								{
-									MaterialEditorInstance->StoredLayerPreviews[Counter]->SetParentEditorOnly(EditedMaterial);
+									if (MaterialEditorInstance->StoredLayerPreviews[Counter] == nullptr)
+									{
+										MaterialEditorInstance->StoredLayerPreviews[Counter] = (NewObject<UMaterialInstanceConstant>(MaterialEditorInstance, NAME_None));
+									}
+									UMaterialInterface* EditedMaterial = Cast<UMaterialFunctionInterface>(AssetObject)->GetPreviewMaterial();
+									if (MaterialEditorInstance->StoredLayerPreviews[Counter] && MaterialEditorInstance->StoredLayerPreviews[Counter]->Parent != EditedMaterial)
+									{
+										MaterialEditorInstance->StoredLayerPreviews[Counter]->SetParentEditorOnly(EditedMaterial);
+									}
 								}
 							}
 							LayerProperties.Last()->Children.Add(ChildProperty);
