Ensure that all compilation has been submitted when querying for a material's statistics.

- This is important when running with "compile what you use" ODSC in the editor.

diff --git a/Engine/Source/Editor/MaterialEditor/Private/MaterialEditingLibrary.cpp b/Engine/Source/Editor/MaterialEditor/Private/MaterialEditingLibrary.cpp
--- a/Engine/Source/Editor/MaterialEditor/Private/MaterialEditingLibrary.cpp
+++ b/Engine/Source/Editor/MaterialEditor/Private/MaterialEditingLibrary.cpp
@@ -30,6 +30,7 @@
 #include "AssetRegistryModule.h"
 #include "DebugViewModeHelpers.h"
 #include "Subsystems/AssetEditorSubsystem.h"
+#include "ShaderCompiler.h"
 
 DEFINE_LOG_CATEGORY_STATIC(LogMaterialEditingLibrary, Warning, All);
 
@@ -1211,6 +1212,10 @@ FMaterialStatistics UMaterialEditingLibrary::GetStatistics(class UMaterialInterf
 	FMaterialResource* Resource = Material ? Material->GetMaterialResource(GMaxRHIFeatureLevel) : nullptr;
 	if (Resource)
 	{
+		if (!Resource->IsGameThreadShaderMapComplete())
+		{
+			Resource->SubmitCompileJobs(EShaderCompileJobPriority::High);
+		}
 		Resource->FinishCompilation();
 
 		TArray<FMaterialStatsUtils::FShaderInstructionsInfo> InstructionInfos;
