Use TWeakObjectPtr<> on MI layer parameter editor code, so we can check for live objects during various Slate callbacks
Otherwise we can hit callbacks before editor has a chance to refresh, and crash on stale pointer access

diff --git a/Engine/Source/Editor/MaterialEditor/Private/MaterialPropertyHelpers.cpp b/Engine/Source/Editor/MaterialEditor/Private/MaterialPropertyHelpers.cpp
--- a/Engine/Source/Editor/MaterialEditor/Private/MaterialPropertyHelpers.cpp
+++ b/Engine/Source/Editor/MaterialEditor/Private/MaterialPropertyHelpers.cpp
@@ -530,7 +530,7 @@ FReply FMaterialPropertyHelpers::OnClickedSaveNewLayerInstance(class UMaterialFu
 
 bool FMaterialPropertyHelpers::IsOverriddenExpression(UDEditorParameterValue* Parameter)
 {
-	return Parameter->bOverride != 0;
+	return Parameter && Parameter->bOverride != 0;
 }
 
 ECheckBoxState FMaterialPropertyHelpers::IsOverriddenExpressionCheckbox(UDEditorParameterValue* Parameter)
diff --git a/Engine/Source/Editor/MaterialEditor/Private/SMaterialLayersFunctionsTree.cpp b/Engine/Source/Editor/MaterialEditor/Private/SMaterialLayersFunctionsTree.cpp
--- a/Engine/Source/Editor/MaterialEditor/Private/SMaterialLayersFunctionsTree.cpp
+++ b/Engine/Source/Editor/MaterialEditor/Private/SMaterialLayersFunctionsTree.cpp
@@ -1608,7 +1608,11 @@ void SMaterialLayersFunctionsInstanceTree::CreateGroupsWidget()
 
 bool SMaterialLayersFunctionsInstanceTree::IsLayerVisible(int32 Index) const
 {
-	return FunctionInstance->GetLayerVisibility(Index);
+	if (FunctionParameter.IsValid())
+	{
+		return FunctionInstance->GetLayerVisibility(Index);
+	}
+	return false;
 }
 
 TSharedRef<SWidget> SMaterialLayersFunctionsInstanceTree::CreateThumbnailWidget(EMaterialParameterAssociation InAssociation, int32 InIndex, float InThumbnailSize)
@@ -1768,7 +1772,7 @@ void SMaterialLayersFunctionsInstanceTree::ShowSubParameters(TSharedPtr<FSortedP
 
 void SMaterialLayersFunctionsInstanceWrapper::Refresh()
 {
-	LayerParameter = nullptr;
+	LayerParameter.Reset();
 	TSharedPtr<SHorizontalBox> HeaderBox;
 	NestedTree->CreateGroupsWidget();
 	LayerParameter = NestedTree->FunctionParameter;
@@ -1803,7 +1807,7 @@ void SMaterialLayersFunctionsInstanceWrapper::Refresh()
 					NestedTree.ToSharedRef()
 				]
 		];
-		if (NestedTree->FunctionParameter != nullptr && FMaterialPropertyHelpers::IsOverriddenExpression(NestedTree->FunctionParameter))
+		if (FMaterialPropertyHelpers::IsOverriddenExpression(NestedTree->FunctionParameter))
 		{
 			HeaderBox->AddSlot()
 				.HAlign(HAlign_Left)
@@ -1935,7 +1939,11 @@ void SMaterialLayersFunctionsMaterialTreeItem::RefreshOnRowChange(const FAssetDa
 
 FText SMaterialLayersFunctionsMaterialTreeItem::GetLayerName(SMaterialLayersFunctionsMaterialTree* InTree, int32 Counter) const
 {
-	return InTree->FunctionInstance->GetLayerName(Counter);
+	if (InTree->FunctionParameter.IsValid())
+	{
+		return InTree->FunctionInstance->GetLayerName(Counter);
+	}
+	return FText();
 }
 
 void SMaterialLayersFunctionsMaterialTreeItem::Construct(const FArguments& InArgs, const TSharedRef<STableViewBase>& InOwnerTableView)
diff --git a/Engine/Source/Editor/MaterialEditor/Private/SMaterialLayersFunctionsTree.h b/Engine/Source/Editor/MaterialEditor/Private/SMaterialLayersFunctionsTree.h
--- a/Engine/Source/Editor/MaterialEditor/Private/SMaterialLayersFunctionsTree.h
+++ b/Engine/Source/Editor/MaterialEditor/Private/SMaterialLayersFunctionsTree.h
@@ -13,6 +13,7 @@
 #include "PropertyCustomizationHelpers.h"
 #include "UObject/Object.h"
 #include "UObject/ObjectMacros.h"
+#include "UObject/WeakObjectPtrTemplates.h"
 #include "Widgets/DeclarativeSyntaxSupport.h"
 #include "Widgets/SCompoundWidget.h"
 #include "Widgets/Views/STreeView.h"
@@ -112,7 +113,7 @@ public:
 	void SetEditorInstance(UMaterialEditorInstanceConstant* InMaterialEditorInstance);
 
 	TAttribute<ECheckBoxState> IsParamChecked;
-	class UDEditorParameterValue* LayerParameter;
+	TWeakObjectPtr<class UDEditorParameterValue> LayerParameter;
 	class UMaterialEditorInstanceConstant* MaterialEditorInstance;
 	TSharedPtr<class SMaterialLayersFunctionsInstanceTree> NestedTree;
 	FSimpleDelegate OnLayerPropertyChanged;
@@ -139,7 +140,7 @@ public:
 	void SetParentsExpansionState();
 
 	void ShowHiddenValues(bool& bShowHiddenParameters) { bShowHiddenParameters = true; }
-	class UDEditorParameterValue* FunctionParameter;
+	TWeakObjectPtr<class UDEditorParameterValue> FunctionParameter;
 	struct FMaterialLayersFunctions* FunctionInstance;
 	TSharedPtr<IPropertyHandle> FunctionInstanceHandle;
 	void RefreshOnAssetChange(const struct FAssetData& InAssetData, int32 Index, EMaterialParameterAssociation MaterialType);
@@ -259,7 +260,7 @@ public:
 	void Construct(const FArguments& InArgs);
 	void SetEditorInstance(UMaterialEditorPreviewParameters* InMaterialEditorInstance);
 	TSharedPtr<class IPropertyRowGenerator> GetGenerator();
-	class UDEditorParameterValue* LayerParameter;
+	TWeakObjectPtr<class UDEditorParameterValue> LayerParameter;
 	class UMaterialEditorPreviewParameters* MaterialEditorInstance;
 	TSharedPtr<class SMaterialLayersFunctionsMaterialTree> NestedTree;
 
@@ -286,7 +287,7 @@ public:
 		void OnExpansionChanged(TSharedPtr<FSortedParamData> Item, bool bIsExpanded);
 		void SetParentsExpansionState();
 		void ShowHiddenValues(bool& bShowHiddenParameters) { bShowHiddenParameters = true; }
-		class UDEditorParameterValue* FunctionParameter;
+		TWeakObjectPtr<class UDEditorParameterValue> FunctionParameter;
 		struct FMaterialLayersFunctions* FunctionInstance;
 		TSharedPtr<IPropertyHandle> FunctionInstanceHandle;
 		TSharedPtr<class FAssetThumbnailPool> GetTreeThumbnailPool();
diff --git a/Engine/Source/Editor/MaterialEditor/Public/MaterialPropertyHelpers.h b/Engine/Source/Editor/MaterialEditor/Public/MaterialPropertyHelpers.h
--- a/Engine/Source/Editor/MaterialEditor/Public/MaterialPropertyHelpers.h
+++ b/Engine/Source/Editor/MaterialEditor/Public/MaterialPropertyHelpers.h
@@ -142,6 +142,7 @@ public:
 	/** Returns true if the parameter is being overridden */
 	static bool IsOverriddenExpression(UDEditorParameterValue* Parameter);
 	static bool IsOverriddenExpression(TObjectPtr<UDEditorParameterValue> Parameter) { return IsOverriddenExpression(Parameter.Get()); }
+	static bool IsOverriddenExpression(TWeakObjectPtr<UDEditorParameterValue> Parameter) { return IsOverriddenExpression(Parameter.Get()); }
 	static ECheckBoxState IsOverriddenExpressionCheckbox(UDEditorParameterValue* Parameter);
 
 	/** Gets the expression description of this parameter from the base material */
