Fix crash, support no preview material for the preview material editor.

diff --git a/Engine/Source/Editor/MaterialEditor/Private/SMaterialEditorViewport.cpp b/Engine/Source/Editor/MaterialEditor/Private/SMaterialEditorViewport.cpp
--- a/Engine/Source/Editor/MaterialEditor/Private/SMaterialEditorViewport.cpp
+++ b/Engine/Source/Editor/MaterialEditor/Private/SMaterialEditorViewport.cpp
@@ -885,7 +885,14 @@ void SMaterialEditorUIPreviewZoomer::Construct( const FArguments& InArgs, UMater
 	BackgroundSettings = InArgs._BackgroundSettings;
 	ModifyCheckerboardTextureColors(BackgroundSettings.Checkerboard);
 
-	PreviewBrush = MakeShared<FSlateMaterialBrush>(*InPreviewMaterial, InArgs._InitialPreviewSize);
+	if (InPreviewMaterial)
+	{
+		PreviewBrush = MakeShared<FSlateMaterialBrush>(*InPreviewMaterial, InArgs._InitialPreviewSize);
+	}
+	else
+	{
+		PreviewBrush = MakeShared<FSlateMaterialBrush>(InArgs._InitialPreviewSize);
+	}
 
 	ChildSlot
 	[
@@ -1042,7 +1049,14 @@ void SMaterialEditorUIPreviewZoomer::SetPreviewSize( const FVector2D PreviewSize
 void SMaterialEditorUIPreviewZoomer::SetPreviewMaterial(UMaterialInterface* InPreviewMaterial)
 {
 	// Just create a new brush to avoid possible invalidation issues from only the resource changing
-	PreviewBrush = MakeShared<FSlateMaterialBrush>(*InPreviewMaterial, PreviewBrush->ImageSize);
+	if (InPreviewMaterial)
+	{
+		PreviewBrush = MakeShared<FSlateMaterialBrush>(*InPreviewMaterial, PreviewBrush->ImageSize);
+	}
+	else
+	{
+		PreviewBrush = MakeShared<FSlateMaterialBrush>(PreviewBrush->ImageSize);
+	}
 	ImageWidget->SetImage(PreviewBrush.Get());
 }
 
