Don't create material stats in editor for invalid shader platforms

diff --git a/Engine/Source/Editor/MaterialEditor/Private/MaterialStats.cpp b/Engine/Source/Editor/MaterialEditor/Private/MaterialStats.cpp
--- a/Engine/Source/Editor/MaterialEditor/Private/MaterialStats.cpp
+++ b/Engine/Source/Editor/MaterialEditor/Private/MaterialStats.cpp
@@ -487,6 +487,11 @@ void FMaterialStats::BuildShaderPlatformDB()
 TSharedPtr<FShaderPlatformSettings> FMaterialStats::AddShaderPlatform(const EPlatformCategoryType PlatformType, const EShaderPlatform PlatformID, const FName PlatformName,
 	const bool bAllowPresenceInGrid, const bool bAllowCodeView, const FString& Description)
 {
+	if (!FDataDrivenShaderPlatformInfo::IsValid(PlatformID))
+	{
+		return TSharedPtr<FShaderPlatformSettings>();
+	}
+
 	TSharedPtr<FShaderPlatformSettings> PlatformPtr = MakeShareable(new FShaderPlatformSettings(PlatformType, PlatformID, PlatformName, bAllowPresenceInGrid, bAllowCodeView, Description));
 	ShaderPlatformStatsDB.Add(PlatformID, PlatformPtr);
 
