Only compile the shaders needed to gather shader platform stats in the Material Editor.

# Results
Modfiying WorldGridMaterial with two platforms open in platform stats.
- We compile a fixed number of shaders (less than 10 per platform) now instead of the 1078 shaders (per platform) w/ WorldGridMaterial.
- On average ProcessCompiledShaderMaps is 40x faster. (1727ms vs. 43ms)
- In the worst cast ProcessCompiledShaderMaps is 10.9x faster. (4865ms vs 445ms)
- The material editor in this scenario goes from unusable to useable.

# Changes
- Add `FMaterial::CacheGivenTypes` to compile just the shader types given to it.
- Call `GetRepresentativeShaderTypesAndDescriptions` to gather the shader types we care about and only submit jobs to compile those shaders when generating platform stats.
- Since we are no longer compile a complete shader map the ShaderCount is incorrect.  Now we call `GetShaderTypes` of the `FMaterialResource` to gather the number of shaders in the material.  This function doesn't trigger shader compilation which is important.
- Adding const to the following functions:
    - FMaterial::GetShaderTypesForLayout
    - FMaterial::GetShaderTypes
- FMaterial::GetShaderTypes is now ENGINE_API so we can call it from the Material Editor.
- Add `TBasePassPSFCachedVolumeIndirectLightingPolicy` shader type to FMaterialStatsUtils::GetRepresentativeShaderTypesAndDescriptions.  This shader has the worst case sampler count.  This ensures we compile this shader and are able to query the worst case sampler count for the platform stats.  This is very ad-hoc, and could be improved in the future.

diff --git a/Engine/Source/Editor/MaterialEditor/Private/MaterialStats.cpp b/Engine/Source/Editor/MaterialEditor/Private/MaterialStats.cpp
--- a/Engine/Source/Editor/MaterialEditor/Private/MaterialStats.cpp
+++ b/Engine/Source/Editor/MaterialEditor/Private/MaterialStats.cpp
@@ -15,6 +15,7 @@
 #include "IMaterialEditor.h"
 #include "Preferences/MaterialStatsOptions.h"
 #include "MaterialEditorSettings.h"
+#include "ShaderCompiler.h"
 
 #include "Modules/ModuleManager.h"
 #include "Developer/MessageLog/Public/MessageLogModule.h"
@@ -202,7 +203,34 @@ bool FShaderPlatformSettings::CheckShaders()
 				}
 
 				Data.LastTimeCompilationRequested = CurrentTime;
-				Data.MaterialResourcesStats->CacheShaders(PlatformShaderID);
+
+				TMap<FName, TArray<FMaterialStatsUtils::FRepresentativeShaderInfo>> ShaderTypeNamesAndDescriptions;
+				FMaterialStatsUtils::GetRepresentativeShaderTypesAndDescriptions(ShaderTypeNamesAndDescriptions, Data.MaterialResourcesStats);
+
+				TArray<FVertexFactoryType*> VFTypes;
+				TArray<FShaderType*> ShaderTypes;
+				for (auto& DescriptionPair : ShaderTypeNamesAndDescriptions)
+				{
+					FVertexFactoryType* VFType = FindVertexFactoryType(DescriptionPair.Key);
+					check(VFType);
+
+					auto& DescriptionArray = DescriptionPair.Value;
+					for (const FMaterialStatsUtils::FRepresentativeShaderInfo& ShaderInfo : DescriptionArray)
+					{
+						FShaderType* ShaderType = FindShaderTypeByName(ShaderInfo.ShaderName);
+						if (ShaderType && VFType)
+						{
+							VFTypes.Add(VFType);
+							ShaderTypes.Add(ShaderType);
+						}
+					}
+				}
+
+				// Prepare the resource for compilation, but don't compile the completed shader map.
+				Data.MaterialResourcesStats->CacheShaders(PlatformShaderID, EMaterialShaderPrecompileMode::None);
+
+				// Compile just the types we want.
+				Data.MaterialResourcesStats->CacheGivenTypes(PlatformShaderID, VFTypes, ShaderTypes);
 
 				Data.bCompilingShaders = true;
 				Data.bUpdateShaderCode = true;
@@ -263,7 +291,7 @@ bool FShaderPlatformSettings::Update()
 					}
 				}
 
-				FMaterialStatsUtils::ExtractMatertialStatsInfo(QualityItem.ShaderStatsInfo, Resource);
+				FMaterialStatsUtils::ExtractMatertialStatsInfo(PlatformShaderID, QualityItem.ShaderStatsInfo, Resource);
 
 				bRetValue = true;
 			}
diff --git a/Engine/Source/Editor/MaterialEditor/Private/MaterialStatsCommon.cpp b/Engine/Source/Editor/MaterialEditor/Private/MaterialStatsCommon.cpp
--- a/Engine/Source/Editor/MaterialEditor/Private/MaterialStatsCommon.cpp
+++ b/Engine/Source/Editor/MaterialEditor/Private/MaterialStatsCommon.cpp
@@ -343,6 +343,11 @@ void FMaterialStatsUtils::GetRepresentativeShaderTypesAndDescriptions(TMap<FName
 
 		ShaderTypeNamesAndDescriptions.FindOrAdd(FGPUFactoryName)
 			.Add(FRepresentativeShaderInfo(ERepresentativeShader::SkeletalMesh, TBasePassVSFNoLightMapPolicyName, TEXT("Base pass vertex shader")));
+
+		// Add the shader type with the most sampler usages so we can accurately report the worst case scenario.
+		// This is ad-hoc, and ideally we have a better way for finding this shader type in the future.
+		ShaderTypeNamesAndDescriptions.FindOrAdd(FLocalVertexFactoryName)
+			.Add(FRepresentativeShaderInfo(ERepresentativeShader::StationarySurface, FName(TEXT("TBasePassPSFCachedVolumeIndirectLightingPolicy")), TEXT("MaxSampler")));
 	}
 	else
 	{
@@ -484,7 +489,7 @@ void FMaterialStatsUtils::GetRepresentativeInstructionCounts(TArray<FShaderInstr
 	//when adding a shader type here be sure to update FPreviewMaterial::ShouldCache()
 	//so the shader type will get compiled with preview materials
 	const FMaterialShaderMap* MaterialShaderMap = Target->GetGameThreadShaderMap();
-	if (MaterialShaderMap && MaterialShaderMap->IsCompilationFinalized())
+	if (MaterialShaderMap)
 	{
 		GetRepresentativeShaderTypesAndDescriptions(ShaderTypeNamesAndDescriptions, Target);
 		TStaticArray<bool, (int32)ERepresentativeShader::Num> bShaderTypeAdded(InPlace, false);
@@ -559,7 +564,7 @@ void FMaterialStatsUtils::GetRepresentativeInstructionCounts(TArray<FShaderInstr
 	}
 }
 
-void FMaterialStatsUtils::ExtractMatertialStatsInfo(FShaderStatsInfo& OutInfo, const FMaterialResource* MaterialResource)
+void FMaterialStatsUtils::ExtractMatertialStatsInfo(EShaderPlatform ShaderPlatform, FShaderStatsInfo& OutInfo, const FMaterialResource* MaterialResource)
 {
 	// extract potential errors
 	const ERHIFeatureLevel::Type MaterialFeatureLevel = MaterialResource->GetFeatureLevel();
@@ -599,7 +604,7 @@ void FMaterialStatsUtils::ExtractMatertialStatsInfo(FShaderStatsInfo& OutInfo, c
 		OutInfo.SamplersCount.StrDescription = FString::Printf(TEXT("%u/%u"), SamplersUsed, MaxSamplers);
 		OutInfo.SamplersCount.StrDescriptionLong = FString::Printf(TEXT("%s samplers: %u/%u"), TEXT("Texture"), SamplersUsed, MaxSamplers);
 
-		// extract esimated sample info
+		// extract estimated sample info
 		uint32 NumVSTextureSamples = 0, NumPSTextureSamples = 0;
 		MaterialResource->GetEstimatedNumTextureSamples(NumVSTextureSamples, NumPSTextureSamples);
 
@@ -622,11 +627,26 @@ void FMaterialStatsUtils::ExtractMatertialStatsInfo(FShaderStatsInfo& OutInfo, c
 		OutInfo.InterpolatorsCount.StrDescriptionLong = FString::Printf(TEXT("User interpolators: %u/%u Scalars (%u/4 Vectors) (TexCoords: %i, Custom: %i)"),
 			TotalScalars, MaxScalars, MaxScalars / 4, UVScalarsUsed, CustomInterpolatorScalarsUsed);
 
-		// extract total shader count info
-		const uint32 ShaderCount = MaterialResource->GetGameThreadShaderMap()->GetShaderNum();
+		// Extract total shader count w/o having to compile shaders.
+		FPlatformTypeLayoutParameters LayoutParams;
+		LayoutParams.InitializeForPlatform(nullptr);
+
+		TArray<FDebugShaderTypeInfo> OutShaderInfo;
+		MaterialResource->GetShaderTypes(ShaderPlatform, LayoutParams, OutShaderInfo);
+
+		int TotalShadersForMaterial = 0;
+		for (const FDebugShaderTypeInfo& ShaderInfo : OutShaderInfo)
+		{
+			TotalShadersForMaterial += ShaderInfo.ShaderTypes.Num();
+
+			for (const FDebugShaderPipelineInfo& PipelineInfo : ShaderInfo.Pipelines)
+			{
+				TotalShadersForMaterial += PipelineInfo.ShaderTypes.Num();
+			}
+		}
 
-		OutInfo.ShaderCount.StrDescription = FString::Printf(TEXT("%u"), ShaderCount);
-		OutInfo.ShaderCount.StrDescriptionLong = FString::Printf(TEXT("Total Shaders: %u"), ShaderCount);
+		OutInfo.ShaderCount.StrDescription = FString::Printf(TEXT("%u"), TotalShadersForMaterial);
+		OutInfo.ShaderCount.StrDescriptionLong = FString::Printf(TEXT("Total Shaders: %u"), TotalShadersForMaterial);
 	}
 }
 
diff --git a/Engine/Source/Editor/MaterialEditor/Public/MaterialStatsCommon.h b/Engine/Source/Editor/MaterialEditor/Public/MaterialStatsCommon.h
--- a/Engine/Source/Editor/MaterialEditor/Public/MaterialStatsCommon.h
+++ b/Engine/Source/Editor/MaterialEditor/Public/MaterialStatsCommon.h
@@ -104,7 +104,7 @@ public:
 	static void GetRepresentativeInstructionCounts(TArray<FShaderInstructionsInfo>& Results, const class FMaterialResource* Target);
 
 	MATERIALEDITOR_API static void GetRepresentativeShaderTypesAndDescriptions(TMap<FName, TArray<FRepresentativeShaderInfo>>& OutShaderTypeNameAndDescriptions, const class FMaterial* TargetMaterial);
-	MATERIALEDITOR_API static void ExtractMatertialStatsInfo(struct FShaderStatsInfo& OutInfo, const FMaterialResource* Target);
+	MATERIALEDITOR_API static void ExtractMatertialStatsInfo(EShaderPlatform ShaderPlatform, struct FShaderStatsInfo& OutInfo, const FMaterialResource* Target);
 
 	static FString RepresentativeShaderTypeToString(const ERepresentativeShader ShaderType);
 
diff --git a/Engine/Source/Runtime/Engine/Private/Materials/MaterialShared.cpp b/Engine/Source/Runtime/Engine/Private/Materials/MaterialShared.cpp
--- a/Engine/Source/Runtime/Engine/Private/Materials/MaterialShared.cpp
+++ b/Engine/Source/Runtime/Engine/Private/Materials/MaterialShared.cpp
@@ -2526,6 +2526,42 @@ bool FMaterial::CacheShaders(const FMaterialShaderMapId& ShaderMapId, EShaderPla
 	return bShaderMapValid;
 }
 
+#if WITH_EDITOR
+void FMaterial::CacheGivenTypes(EShaderPlatform Platform, const TArray<FVertexFactoryType*>& VFTypes, const TArray<FShaderType*>& ShaderTypes, const ITargetPlatform* TargetPlatform)
+{
+	TRACE_CPUPROFILER_EVENT_SCOPE(FMaterial::CacheGivenTypes);
+	check(IsInGameThread());
+	checkf(ShaderTypes.Num() == VFTypes.Num(), TEXT("The size of the shader type array and vertex factory type array must match."));
+	checkf(GameThreadShaderMap, TEXT("Shader map is not initialized.  Please call CacheShaders first."));
+	checkf(GetGameThreadCompilingShaderMapId() != 0, TEXT("Material is not prepared to compile yet.  Please call CacheShaders first."));
+
+	TArray<FShaderCommonCompileJobPtr> CompileJobs;
+	for (int i = 0; i < VFTypes.Num(); ++i)
+	{
+		FVertexFactoryType* VFType = VFTypes[i];
+		check(VFType);
+
+		FShaderType* ShaderType = ShaderTypes[i];
+		check(ShaderType);
+
+		ShaderType->AsMeshMaterialShaderType()->BeginCompileShader(
+			EShaderCompileJobPriority::ForceLocal,
+			GetGameThreadCompilingShaderMapId(),
+			0,
+			Platform,
+			GameThreadShaderMap->GetPermutationFlags(),
+			this,
+			GameThreadPendingCompilerEnvironment,
+			VFType,
+			CompileJobs,
+			nullptr,
+			nullptr);
+	}
+
+	GShaderCompilingManager->SubmitJobs(CompileJobs, GetBaseMaterialPathName(), GameThreadShaderMap->GetDebugDescription());
+}
+#endif // WITH_EDITOR
+
 bool FMaterial::Translate_Legacy(const FMaterialShaderMapId& ShaderMapId,
 	const FStaticParameterSet& InStaticParameters,
 	EShaderPlatform InPlatform,
@@ -4115,7 +4151,7 @@ void FMaterial::SaveShaderStableKeys(EShaderPlatform TargetShaderPlatform, FStab
 }
 
 #if WITH_EDITOR
-void FMaterial::GetShaderTypesForLayout(EShaderPlatform Platform, const FShaderMapLayout& Layout, FVertexFactoryType* VertexFactory, TArray<FDebugShaderTypeInfo>& OutShaderInfo)
+void FMaterial::GetShaderTypesForLayout(EShaderPlatform Platform, const FShaderMapLayout& Layout, FVertexFactoryType* VertexFactory, TArray<FDebugShaderTypeInfo>& OutShaderInfo) const
 {
 	FDebugShaderTypeInfo DebugInfo;
 	DebugInfo.VFType = VertexFactory;
@@ -4147,7 +4183,7 @@ void FMaterial::GetShaderTypesForLayout(EShaderPlatform Platform, const FShaderM
 	OutShaderInfo.Add(DebugInfo);
 }
 
-void FMaterial::GetShaderTypes(EShaderPlatform Platform, const FPlatformTypeLayoutParameters& LayoutParams, TArray<FDebugShaderTypeInfo>& OutShaderInfo)
+void FMaterial::GetShaderTypes(EShaderPlatform Platform, const FPlatformTypeLayoutParameters& LayoutParams, TArray<FDebugShaderTypeInfo>& OutShaderInfo) const
 {
 	const FMaterialShaderParameters MaterialParameters(this);
 	const FMaterialShaderMapLayout& Layout = AcquireMaterialShaderMapLayout(Platform, GetShaderPermutationFlags(LayoutParams), MaterialParameters);
diff --git a/Engine/Source/Runtime/Engine/Public/MaterialShared.h b/Engine/Source/Runtime/Engine/Public/MaterialShared.h
--- a/Engine/Source/Runtime/Engine/Public/MaterialShared.h
+++ b/Engine/Source/Runtime/Engine/Public/MaterialShared.h
@@ -1702,6 +1702,16 @@ public:
 	 */
 	ENGINE_API bool CacheShaders(const FMaterialShaderMapId& ShaderMapId, EShaderPlatform Platform, EMaterialShaderPrecompileMode PrecompileMode = EMaterialShaderPrecompileMode::Default, const ITargetPlatform* TargetPlatform = nullptr);
 
+#if WITH_EDITOR
+	/**
+	 * Submits local compile jobs for the exact given shader types and vertex factory type combination.
+	 * @note CacheShaders() should be called first to prepare the resource for compilation.
+	 * @note The arrays of shader types and vertex factory types must match.
+	 * @note These compile jobs are submitted async and it is up to client code to block on results if needed.
+	 */
+	ENGINE_API void CacheGivenTypes(EShaderPlatform ShaderPlatform, const TArray<FVertexFactoryType*>& VFTypes, const TArray<FShaderType*>& ShaderTypes, const ITargetPlatform* TargetPlatform = nullptr);
+#endif
+
 	/**
 	 * Should the shader for this material with the given platform, shader type and vertex 
 	 * factory type combination be compiled
@@ -2076,17 +2086,17 @@ public:
 	void SaveShaderStableKeys(EShaderPlatform TargetShaderPlatform, struct FStableShaderKeyAndValue& SaveKeyVal); // arg is non-const, we modify it as we go
 
 #if WITH_EDITOR
-/**
-*	Gathers a list of shader types sorted by vertex factory types that should be cached for this material.  Avoids doing expensive material
-*	and shader compilation to acquire this information.
-*
-*	@param	Platform		The shader platform to get info for.
-*	@param	OutShaderInfo	Array of results sorted by vertex factory type, and shader type.
-*
-*/
-	void GetShaderTypes(EShaderPlatform Platform, const FPlatformTypeLayoutParameters& LayoutParams, TArray<FDebugShaderTypeInfo>& OutShaderInfo);
+	/**
+	*	Gathers a list of shader types sorted by vertex factory types that should be cached for this material.  Avoids doing expensive material
+	*	and shader compilation to acquire this information.
+	*
+	*	@param	Platform		The shader platform to get info for.
+	*	@param	OutShaderInfo	Array of results sorted by vertex factory type, and shader type.
+	*
+	*/
+	ENGINE_API void GetShaderTypes(EShaderPlatform Platform, const FPlatformTypeLayoutParameters& LayoutParams, TArray<FDebugShaderTypeInfo>& OutShaderInfo) const;
 
-	void GetShaderTypesForLayout(EShaderPlatform Platform, const FShaderMapLayout& Layout, FVertexFactoryType* VertexFactory, TArray<FDebugShaderTypeInfo>& OutShaderInfo);
+	void GetShaderTypesForLayout(EShaderPlatform Platform, const FShaderMapLayout& Layout, FVertexFactoryType* VertexFactory, TArray<FDebugShaderTypeInfo>& OutShaderInfo) const;
 #endif // WITH_EDITOR
 
 #if WITH_EDITOR
