Adding SM6 preview to the Material Editor.

diff --git a/Engine/Source/Editor/MaterialEditor/Private/MaterialEditor.cpp b/Engine/Source/Editor/MaterialEditor/Private/MaterialEditor.cpp
--- a/Engine/Source/Editor/MaterialEditor/Private/MaterialEditor.cpp
+++ b/Engine/Source/Editor/MaterialEditor/Private/MaterialEditor.cpp
@@ -1501,11 +1501,16 @@ void FMaterialEditor::SetFeaturePreview(ERHIFeatureLevel::Type NewFeatureLevel)
 	bPreviewFeaturesChanged = true;
 }
 
-bool FMaterialEditor::IsFeaturePreviewChecked(ERHIFeatureLevel::Type TestFeatureLevel)
+bool FMaterialEditor::IsFeaturePreviewChecked(ERHIFeatureLevel::Type TestFeatureLevel) const
 {
 	return NodeFeatureLevel == TestFeatureLevel;
 }
 
+bool FMaterialEditor::IsFeaturePreviewAvailable(ERHIFeatureLevel::Type TestFeatureLevel) const
+{
+	return GMaxRHIFeatureLevel >= TestFeatureLevel;
+}
+
 BEGIN_SLATE_FUNCTION_BUILD_OPTIMIZATION
 void FMaterialEditor::CreateInternalWidgets()
 {
@@ -1980,6 +1985,7 @@ void FMaterialEditor::GeneratePreviewMenuContent(UToolMenu* Menu)
 		FeatureLevelSection.AddMenuEntry(FMaterialEditorCommands::Get().FeatureLevel_All);
 		FeatureLevelSection.AddMenuEntry(FMaterialEditorCommands::Get().FeatureLevel_ES31);
 		FeatureLevelSection.AddMenuEntry(FMaterialEditorCommands::Get().FeatureLevel_SM5);
+		FeatureLevelSection.AddMenuEntry(FMaterialEditorCommands::Get().FeatureLevel_SM6);
 	}
 
 	FToolMenuSection& StaticSwitchSection = Menu->AddSection("StaticSwitchPreview", LOCTEXT("StaticSwitchHeading", "Switch Params"));
@@ -3379,13 +3385,18 @@ void FMaterialEditor::BindCommands()
 	ToolkitCommands->MapAction(
 		Commands.FeatureLevel_ES31,
 		FExecuteAction::CreateSP(this, &FMaterialEditor::SetFeaturePreview, ERHIFeatureLevel::ES3_1),
-		FCanExecuteAction(),
+		FCanExecuteAction::CreateSP(this, &FMaterialEditor::IsFeaturePreviewAvailable, ERHIFeatureLevel::ES3_1),
 		FIsActionChecked::CreateSP(this, &FMaterialEditor::IsFeaturePreviewChecked, ERHIFeatureLevel::ES3_1));
 	ToolkitCommands->MapAction(
 		Commands.FeatureLevel_SM5,
 		FExecuteAction::CreateSP(this, &FMaterialEditor::SetFeaturePreview, ERHIFeatureLevel::SM5),
-		FCanExecuteAction(),
+		FCanExecuteAction::CreateSP(this, &FMaterialEditor::IsFeaturePreviewAvailable, ERHIFeatureLevel::SM5),
 		FIsActionChecked::CreateSP(this, &FMaterialEditor::IsFeaturePreviewChecked, ERHIFeatureLevel::SM5));
+	ToolkitCommands->MapAction(
+		Commands.FeatureLevel_SM6,
+		FExecuteAction::CreateSP(this, &FMaterialEditor::SetFeaturePreview, ERHIFeatureLevel::SM6),
+		FCanExecuteAction::CreateSP(this, &FMaterialEditor::IsFeaturePreviewAvailable, ERHIFeatureLevel::SM6),
+		FIsActionChecked::CreateSP(this, &FMaterialEditor::IsFeaturePreviewChecked, ERHIFeatureLevel::SM6));
 }
 
 void FMaterialEditor::OnApply()
diff --git a/Engine/Source/Editor/MaterialEditor/Private/MaterialEditor.h b/Engine/Source/Editor/MaterialEditor/Private/MaterialEditor.h
--- a/Engine/Source/Editor/MaterialEditor/Private/MaterialEditor.h
+++ b/Engine/Source/Editor/MaterialEditor/Private/MaterialEditor.h
@@ -670,7 +670,8 @@ private:
 
 	/** Helper functions for the feature level node display toggling */
 	void SetFeaturePreview(ERHIFeatureLevel::Type NewFeatureLevel);
-	bool IsFeaturePreviewChecked(ERHIFeatureLevel::Type TestFeatureLevel);
+	bool IsFeaturePreviewChecked(ERHIFeatureLevel::Type TestFeatureLevel) const;
+	bool IsFeaturePreviewAvailable(ERHIFeatureLevel::Type TestFeatureLevel) const;
 
 
 public:
diff --git a/Engine/Source/Editor/MaterialEditor/Private/MaterialEditorActions.cpp b/Engine/Source/Editor/MaterialEditor/Private/MaterialEditorActions.cpp
--- a/Engine/Source/Editor/MaterialEditor/Private/MaterialEditorActions.cpp
+++ b/Engine/Source/Editor/MaterialEditor/Private/MaterialEditorActions.cpp
@@ -71,6 +71,7 @@ void FMaterialEditorCommands::RegisterCommands()
 	UI_COMMAND(FeatureLevel_All, "All", "Sets node preview to show all feature levels.", EUserInterfaceActionType::RadioButton, FInputChord());
 	UI_COMMAND(FeatureLevel_ES31, "ES3.1", "Sets node preview to show the ES3.1 feature level.", EUserInterfaceActionType::RadioButton, FInputChord());
 	UI_COMMAND(FeatureLevel_SM5, "SM5", "Sets node preview to show the SM5 feature level.", EUserInterfaceActionType::RadioButton, FInputChord());
+	UI_COMMAND(FeatureLevel_SM6, "SM6", "Sets node preview to show the SM6 feature level.", EUserInterfaceActionType::RadioButton, FInputChord());
 }
 
 //////////////////////////////////////////////////////////////////////////
diff --git a/Engine/Source/Editor/MaterialEditor/Public/MaterialEditorActions.h b/Engine/Source/Editor/MaterialEditor/Public/MaterialEditorActions.h
--- a/Engine/Source/Editor/MaterialEditor/Public/MaterialEditorActions.h
+++ b/Engine/Source/Editor/MaterialEditor/Public/MaterialEditorActions.h
@@ -192,6 +192,7 @@ public:
 	TSharedPtr< FUICommandInfo > FeatureLevel_All;
 	TSharedPtr< FUICommandInfo > FeatureLevel_ES31;
 	TSharedPtr< FUICommandInfo > FeatureLevel_SM5;
+	TSharedPtr< FUICommandInfo > FeatureLevel_SM6;
 
 	/**
 	 * Initialize commands
