Warn the user if there are material errors when calling CacheGivenTypes().

- Prevent CacheGivenTypes from being called if material compilation fails when generating platform stats.

diff --git a/Engine/Source/Editor/MaterialEditor/Private/MaterialStats.cpp b/Engine/Source/Editor/MaterialEditor/Private/MaterialStats.cpp
--- a/Engine/Source/Editor/MaterialEditor/Private/MaterialStats.cpp
+++ b/Engine/Source/Editor/MaterialEditor/Private/MaterialStats.cpp
@@ -227,10 +227,13 @@ bool FShaderPlatformSettings::CheckShaders()
 				}
 
 				// Prepare the resource for compilation, but don't compile the completed shader map.
-				Data.MaterialResourcesStats->CacheShaders(PlatformShaderID, EMaterialShaderPrecompileMode::None);
+				const bool bSuccess = Data.MaterialResourcesStats->CacheShaders(PlatformShaderID, EMaterialShaderPrecompileMode::None);
 
-				// Compile just the types we want.
-				Data.MaterialResourcesStats->CacheGivenTypes(PlatformShaderID, VFTypes, ShaderTypes);
+				if (bSuccess)
+				{
+					// Compile just the types we want.
+					Data.MaterialResourcesStats->CacheGivenTypes(PlatformShaderID, VFTypes, ShaderTypes);
+				}
 
 				Data.bCompilingShaders = true;
 				Data.bUpdateShaderCode = true;
diff --git a/Engine/Source/Runtime/Engine/Private/Materials/MaterialShared.cpp b/Engine/Source/Runtime/Engine/Private/Materials/MaterialShared.cpp
--- a/Engine/Source/Runtime/Engine/Private/Materials/MaterialShared.cpp
+++ b/Engine/Source/Runtime/Engine/Private/Materials/MaterialShared.cpp
@@ -2529,6 +2529,17 @@ bool FMaterial::CacheShaders(const FMaterialShaderMapId& ShaderMapId, EShaderPla
 #if WITH_EDITOR
 void FMaterial::CacheGivenTypes(EShaderPlatform Platform, const TArray<FVertexFactoryType*>& VFTypes, const TArray<FShaderType*>& ShaderTypes, const ITargetPlatform* TargetPlatform)
 {
+	if (CompileErrors.Num())
+	{
+		UE_LOG(LogMaterial, Warning, TEXT("Material failed to compile."));
+		for (const FString& CompileError : CompileErrors)
+		{
+			UE_LOG(LogMaterial, Warning, TEXT("%s"), *CompileError);
+		}
+
+		return;
+	}
+
 	if (GameThreadShaderMap)
 	{
 		TRACE_CPUPROFILER_EVENT_SCOPE(FMaterial::CacheGivenTypes);
