Fxied: Missing material stats for mobile shaders

diff --git a/Engine/Source/Editor/MaterialEditor/Private/MaterialStatsCommon.cpp b/Engine/Source/Editor/MaterialEditor/Private/MaterialStatsCommon.cpp
--- a/Engine/Source/Editor/MaterialEditor/Private/MaterialStatsCommon.cpp
+++ b/Engine/Source/Editor/MaterialEditor/Private/MaterialStatsCommon.cpp
@@ -258,27 +258,21 @@ FSlateColor FMaterialStatsUtils::QualitySettingColor(const EMaterialQualityLevel
 	return  FStyleColors::Foreground;
 }
 
-static void MobileBasePassShaderName(bool bVertexShader, const TCHAR* PolicyName, int32 NumPointLigts, bool bHDR, bool bSkyLight, FString& OutName)
+static void MobileBasePassShaderName(bool bVertexShader, const TCHAR* PolicyName, bool bEnableLocalLights, bool bHDR, bool bSkyLight, FString& OutName)
 {
 	OutName.Reset();
 	OutName.Append(bVertexShader ? TEXT("TMobileBasePassVS") : TEXT("TMobileBasePassPS"));
 	OutName.Append(PolicyName);
-	if (!bVertexShader)
-	{
-		if (NumPointLigts == INT32_MAX)
-		{
-			OutName.Append(TEXT("INT32_MAX"));
-		}
-		else
-		{
-			OutName.AppendInt(NumPointLigts);
-		}
-	}
 	OutName.Append(bHDR ? TEXT("HDRLinear64") : TEXT("LDRGamma32"));
 	if (!bVertexShader && bSkyLight)
 	{
 		OutName.Append(TEXT("SkyLight"));
 	}
+	if (!bVertexShader && bEnableLocalLights)
+	{
+		OutName.Append(TEXT("EnableLocalLights"));
+	}
+	
 }
 
 void FMaterialStatsUtils::GetRepresentativeShaderTypesAndDescriptions(TMap<FName, TArray<FRepresentativeShaderInfo>>& ShaderTypeNamesAndDescriptions, const FMaterial* TargetMaterial)
@@ -356,13 +350,15 @@ void FMaterialStatsUtils::GetRepresentativeShaderTypesAndDescriptions(TMap<FName
 
 		if (TargetMaterial->GetShadingModels().IsUnlit())
 		{
+			bool bEnableLocalLights = false;
+			
 			//unlit materials are never lightmapped
-			MobileBasePassShaderName(false, TEXT("FNoLightMapPolicy"), 0, bMobileHDR, false, ShaderNameStr);
+			MobileBasePassShaderName(false, TEXT("FNoLightMapPolicy"), bEnableLocalLights, bMobileHDR, false, ShaderNameStr);
 			const FString Description = FString::Printf(TEXT("Mobile base pass shader without light map%s"), DescSuffix);
 			ShaderTypeNamesAndDescriptions.Add(FLocalVertexFactoryName)
 				.Add(FRepresentativeShaderInfo(ERepresentativeShader::StationarySurface, FName(ShaderNameStr), Description));
 
-			MobileBasePassShaderName(true, TEXT("FNoLightMapPolicy"), 0, bMobileHDR, false, ShaderNameStr);
+			MobileBasePassShaderName(true, TEXT("FNoLightMapPolicy"), bEnableLocalLights, bMobileHDR, false, ShaderNameStr);
 			ShaderTypeNamesAndDescriptions.FindOrAdd(FLocalVertexFactoryName)
 				.Add(FRepresentativeShaderInfo(ERepresentativeShader::StaticMesh, FName(ShaderNameStr),
 					FString::Printf(TEXT("Mobile base pass vertex shader%s"), DescSuffix)));
@@ -388,7 +384,8 @@ void FMaterialStatsUtils::GetRepresentativeShaderTypesAndDescriptions(TMap<FName
 				{
 					// distance field shadows only shaders
 					{
-						MobileBasePassShaderName(false, TEXT("FMobileDistanceFieldShadowsAndLQLightMapPolicy"), 0, bMobileHDR, bOnlySkyPermutation, ShaderNameStr);
+						bool bEnableLocalLights = false;
+						MobileBasePassShaderName(false, TEXT("FMobileDistanceFieldShadowsAndLQLightMapPolicy"), bEnableLocalLights, bMobileHDR, bOnlySkyPermutation, ShaderNameStr);
 						const FString Description = FString::Printf(TEXT("Mobile base pass shader with distance field shadows%s"), DescSuffix);
 						ShaderTypeNamesAndDescriptions.FindOrAdd(FLocalVertexFactoryName)
 							.Add(FRepresentativeShaderInfo(ERepresentativeShader::StationarySurface, FName(ShaderNameStr), Description));
@@ -400,14 +397,16 @@ void FMaterialStatsUtils::GetRepresentativeShaderTypesAndDescriptions(TMap<FName
 					{
 						// distance field shadows & CSM shaders
 						{
-							MobileBasePassShaderName(false, TEXT("FMobileDistanceFieldShadowsLightMapAndCSMLightingPolicy"), 0, bMobileHDR, bOnlySkyPermutation, ShaderNameStr);
+							bool bEnableLocalLights = false;
+							MobileBasePassShaderName(false, TEXT("FMobileDistanceFieldShadowsLightMapAndCSMLightingPolicy"), bEnableLocalLights, bMobileHDR, bOnlySkyPermutation, ShaderNameStr);
 							const FString Description = FString::Printf(TEXT("Mobile base pass shader with distance field shadows and CSM%s"), DescSuffix);
 							ShaderTypeNamesAndDescriptions.FindOrAdd(FLocalVertexFactoryName)
 								.Add(FRepresentativeShaderInfo(ERepresentativeShader::StationarySurfaceCSM, FName(ShaderNameStr), Description));
 						}
 
 						{
-							MobileBasePassShaderName(false, TEXT("FMobileDistanceFieldShadowsAndLQLightMapPolicy"), INT32_MAX, bMobileHDR, bOnlySkyPermutation, ShaderNameStr);
+							bool bEnableLocalLights = true;
+							MobileBasePassShaderName(false, TEXT("FMobileDistanceFieldShadowsAndLQLightMapPolicy"), bEnableLocalLights, bMobileHDR, bOnlySkyPermutation, ShaderNameStr);
 							const FString Description = FString::Printf(TEXT("Mobile base pass shader with distance field shadows, CSM and local light(s) %s"), DescSuffix);
 
 							FRepresentativeShaderInfo ShaderInfo = FRepresentativeShaderInfo(ERepresentativeShader::StationarySurfaceNPointLights, FName(ShaderNameStr), Description);
@@ -419,17 +418,19 @@ void FMaterialStatsUtils::GetRepresentativeShaderTypesAndDescriptions(TMap<FName
 				}
 				else //no shadows & lightmapped
 				{
-					MobileBasePassShaderName(false, TEXT("TLightMapPolicyLQ"), 0, bMobileHDR, bOnlySkyPermutation, ShaderNameStr);
-					ShaderTypeNamesAndDescriptions.FindOrAdd(FLocalVertexFactoryName)
-						.Add(FRepresentativeShaderInfo(ERepresentativeShader::StationarySurface, FName(ShaderNameStr),
-							FString::Printf(TEXT("Mobile base pass shader with static lighting%s"), DescSuffix)));
+					{
+						bool bEnableLocalLights = false;
+						MobileBasePassShaderName(false, TEXT("TLightMapPolicyLQ"), bEnableLocalLights, bMobileHDR, bOnlySkyPermutation, ShaderNameStr);
+						ShaderTypeNamesAndDescriptions.FindOrAdd(FLocalVertexFactoryName)
+							.Add(FRepresentativeShaderInfo(ERepresentativeShader::StationarySurface, FName(ShaderNameStr),
+								FString::Printf(TEXT("Mobile base pass shader with static lighting%s"), DescSuffix)));
+					}
 
 					{
-						MobileBasePassShaderName(false, TEXT("TLightMapPolicyLQ"), INT32_MAX, bMobileHDR, bOnlySkyPermutation, ShaderNameStr);
+						bool bEnableLocalLights = true;
+						MobileBasePassShaderName(false, TEXT("TLightMapPolicyLQ"), bEnableLocalLights, bMobileHDR, bOnlySkyPermutation, ShaderNameStr);
 						const FString Description = FString::Printf(TEXT("Mobile base pass shader with static lighting and local light(s) %s"), DescSuffix);
-
 						FRepresentativeShaderInfo ShaderInfo = FRepresentativeShaderInfo(ERepresentativeShader::StationarySurfaceNPointLights, FName(ShaderNameStr), Description);
-
 						ShaderTypeNamesAndDescriptions.FindOrAdd(FLocalVertexFactoryName).Add(ShaderInfo);
 					}
 				}
@@ -438,12 +439,13 @@ void FMaterialStatsUtils::GetRepresentativeShaderTypesAndDescriptions(TMap<FName
 			// only one of these 2 shader types will be displayed
 			
 			// dynamically lit shader NoLightmapPolicy
-			MobileBasePassShaderName(false, TEXT("FNoLightmapPolicy"), 0, bMobileHDR, bOnlySkyPermutation, ShaderNameStr);
+			bool bEnableLocalLights = false;
+			MobileBasePassShaderName(false, TEXT("FNoLightmapPolicy"), bEnableLocalLights, bMobileHDR, bOnlySkyPermutation, ShaderNameStr);
 			ShaderTypeNamesAndDescriptions.FindOrAdd(FLocalVertexFactoryName)
 				.Add(FRepresentativeShaderInfo(ERepresentativeShader::DynamicallyLitObject, FName(ShaderNameStr),
 				FString::Printf(TEXT("Mobile base pass shader with only dynamic lighting%s"), DescSuffix)));
 
-			MobileBasePassShaderName(true, TEXT("FNoLightMapPolicy"), 0, bMobileHDR, bOnlySkyPermutation, ShaderNameStr);
+			MobileBasePassShaderName(true, TEXT("FNoLightMapPolicy"), bEnableLocalLights, bMobileHDR, bOnlySkyPermutation, ShaderNameStr);
 			ShaderTypeNamesAndDescriptions.FindOrAdd(FLocalVertexFactoryName)
 				.Add(FRepresentativeShaderInfo(ERepresentativeShader::StaticMesh, FName(ShaderNameStr),
 				FString::Printf(TEXT("Mobile base pass vertex shader%s"), DescSuffix)));
@@ -453,12 +455,12 @@ void FMaterialStatsUtils::GetRepresentativeShaderTypesAndDescriptions(TMap<FName
 				FString::Printf(TEXT("Mobile base pass vertex shader%s"), DescSuffix)));
 
 			// dynamically lit shader FMobileDirectionalLightAndCSMPolicy
-			MobileBasePassShaderName(false, TEXT("FMobileDirectionalLightAndCSMPolicy"), 0, bMobileHDR, bOnlySkyPermutation, ShaderNameStr);
+			MobileBasePassShaderName(false, TEXT("FMobileDirectionalLightAndCSMPolicy"), bEnableLocalLights, bMobileHDR, bOnlySkyPermutation, ShaderNameStr);
 			ShaderTypeNamesAndDescriptions.FindOrAdd(FLocalVertexFactoryName)
 				.Add(FRepresentativeShaderInfo(ERepresentativeShader::DynamicallyLitObject, FName(ShaderNameStr),
 				FString::Printf(TEXT("Mobile base pass shader with only dynamic lighting%s"), DescSuffix)));
 
-			MobileBasePassShaderName(true, TEXT("FMobileDirectionalLightAndCSMPolicy"), 0, bMobileHDR, bOnlySkyPermutation, ShaderNameStr);
+			MobileBasePassShaderName(true, TEXT("FMobileDirectionalLightAndCSMPolicy"), bEnableLocalLights, bMobileHDR, bOnlySkyPermutation, ShaderNameStr);
 			ShaderTypeNamesAndDescriptions.FindOrAdd(FLocalVertexFactoryName)
 				.Add(FRepresentativeShaderInfo(ERepresentativeShader::StaticMesh, FName(ShaderNameStr),
 				FString::Printf(TEXT("Mobile base pass vertex shader%s"), DescSuffix)));
