Add 'Optional' EditorOnly data for UMaterialInterface and UMaterialFunctionInterface

diff --git a/Engine/Source/Developer/AssetTools/Private/AssetVtConversion.cpp b/Engine/Source/Developer/AssetTools/Private/AssetVtConversion.cpp
--- a/Engine/Source/Developer/AssetTools/Private/AssetVtConversion.cpp
+++ b/Engine/Source/Developer/AssetTools/Private/AssetVtConversion.cpp
@@ -351,24 +351,18 @@ void FVTConversionWorker::FindAllTexturesAndMaterials_Iteration(TArray<UMaterial
 	for (UMaterialFunctionInterface *If : FunctionsUsedByAffectedTextures)
 	{
 		// It's a material?
-		UMaterialFunction *Function = Cast<UMaterialFunction>(If);
-
-		// It's a material instance? we're only interested in it's root for now
-		UMaterialFunctionInstance *FunctionInstance = Cast<UMaterialFunctionInstance>(If);
-		if (FunctionInstance != nullptr)
+		UMaterialFunction* Function = If->GetBaseFunction();
+		if (Function)
 		{
-			Function = Cast<UMaterialFunction>(FunctionInstance->GetBaseFunction());
-		}
-
-		check(Function); // It's something else entirely??
-		FunctionInferfaces.AddUnique(Function);
-		FunctionHeap.AddUnique(Function);
-		InAffectedFunctions.AddUnique(Function);
+			FunctionInferfaces.AddUnique(Function);
+			FunctionHeap.AddUnique(Function);
+			InAffectedFunctions.AddUnique(Function);
 
-		AuditTrail.Add(Function, FAuditTrail(
-			If,
-			FString::Printf(TEXT("this is the parent of"))
-		));
+			AuditTrail.Add(Function, FAuditTrail(
+				If,
+				FString::Printf(TEXT("this is the parent of"))
+			));
+		}
 	}
 
 
@@ -382,7 +376,7 @@ void FVTConversionWorker::FindAllTexturesAndMaterials_Iteration(TArray<UMaterial
 			MaterialParametersPair.Key->GetDependentFunctions(DependentFunctions);
 			for (auto Function : DependentFunctions)
 			{
-				for (UMaterialExpression* FunctionExpression : *Function->GetFunctionExpressions())
+				for (UMaterialExpression* FunctionExpression : Function->GetExpressions())
 				{
 					if (const UMaterialExpressionTextureSampleParameter* TexParameter = Cast<const UMaterialExpressionTextureSampleParameter>(FunctionExpression))
 					{
@@ -612,7 +606,7 @@ void FVTConversionWorker::DoConvert()
 			MaterialTask.EnterProgressFrame();
 
 			bool MatModified = false;
-			for (UMaterialExpression *Expr : Mat->Expressions)
+			for (UMaterialExpression *Expr : Mat->GetExpressions())
 			{
 				UMaterialExpressionTextureBase *TexExpr = Cast<UMaterialExpressionTextureBase>(Expr);
 				if (TexExpr)
@@ -677,8 +671,7 @@ void FVTConversionWorker::DoConvert()
 			MaterialTask.EnterProgressFrame();
 
 			bool FuncModified = false;
-			const TArray<TObjectPtr<UMaterialExpression>> *Expressions = Func->GetFunctionExpressions();
-			for (const TObjectPtr<UMaterialExpression>& Expr : *Expressions)
+			for (const TObjectPtr<UMaterialExpression>& Expr : Func->GetExpressions())
 			{
 				UMaterialExpressionTextureBase *TexExpr = Cast<UMaterialExpressionTextureBase>(Expr);
 				if (TexExpr)
diff --git a/Engine/Source/Developer/MaterialBaking/Private/ExportMaterialProxy.h b/Engine/Source/Developer/MaterialBaking/Private/ExportMaterialProxy.h
--- a/Engine/Source/Developer/MaterialBaking/Private/ExportMaterialProxy.h
+++ b/Engine/Source/Developer/MaterialBaking/Private/ExportMaterialProxy.h
@@ -569,7 +569,7 @@ private:
 
 	UMaterialExpressionCustomOutput* GetCustomOutputExpressionToCompile() const
 	{
-		for (UMaterialExpression* Expression : Material->Expressions)
+		for (UMaterialExpression* Expression : Material->GetExpressions())
 		{
 			UMaterialExpressionCustomOutput* CustomOutputExpression = Cast<UMaterialExpressionCustomOutput>(Expression);
 			if (CustomOutputExpression && CustomOutputExpression->GetDisplayName() == CustomOutputToCompile)
diff --git a/Engine/Source/Developer/MaterialUtilities/Private/MaterialUtilities.cpp b/Engine/Source/Developer/MaterialUtilities/Private/MaterialUtilities.cpp
--- a/Engine/Source/Developer/MaterialUtilities/Private/MaterialUtilities.cpp
+++ b/Engine/Source/Developer/MaterialUtilities/Private/MaterialUtilities.cpp
@@ -153,7 +153,7 @@ UMaterialInterface* FMaterialUtilities::CreateProxyMaterialAndTextures(UPackage*
 			SwitchParameter.ParameterInfo.Name = *(TEXT("Use") + TrimmedPropertyName);
 			SwitchParameter.Value = true;
 			SwitchParameter.bOverride = true;
-			NewStaticParameterSet.StaticSwitchParameters.Add(SwitchParameter);
+			NewStaticParameterSet.EditorOnly.StaticSwitchParameters.Add(SwitchParameter);
 		}
 		else
 		{
@@ -188,10 +188,10 @@ UMaterialInterface* FMaterialUtilities::CreateProxyMaterialAndTextures(UPackage*
 		SwitchParameter.ParameterInfo.Name = TEXT("UseCustomUV");
 		SwitchParameter.Value = true;
 		SwitchParameter.bOverride = true;
-		NewStaticParameterSet.StaticSwitchParameters.Add(SwitchParameter);
+		NewStaticParameterSet.EditorOnly.StaticSwitchParameters.Add(SwitchParameter);
 
 		SwitchParameter.ParameterInfo.Name = *(TEXT("UseUV") + FString::FromInt(MeshData.TextureCoordinateIndex));
-		NewStaticParameterSet.StaticSwitchParameters.Add(SwitchParameter);
+		NewStaticParameterSet.EditorOnly.StaticSwitchParameters.Add(SwitchParameter);
 	}
 
 	Material->UpdateStaticPermutation(NewStaticParameterSet);
@@ -1022,6 +1022,8 @@ UMaterial* FMaterialUtilities::CreateMaterial(const FFlattenMaterial& InFlattenM
 	Material->SetShadingModel(MSM_DefaultLit);
 	OutGeneratedAssets.Add(Material);
 
+	UMaterialEditorOnlyData* MaterialEditorOnly = Material->GetEditorOnlyData();
+
 	int32 MaterialNodeY = -150;
 	int32 MaterialNodeStepY = 180;
 
@@ -1040,8 +1042,8 @@ UMaterial* FMaterialUtilities::CreateMaterial(const FFlattenMaterial& InFlattenM
 		BasecolorExpression->SamplerType = EMaterialSamplerType::SAMPLERTYPE_Color;
 		BasecolorExpression->MaterialExpressionEditorX = -400;
 		BasecolorExpression->MaterialExpressionEditorY = MaterialNodeY;
-		Material->Expressions.Add(BasecolorExpression);
-		Material->BaseColor.Expression = BasecolorExpression;
+		Material->GetExpressionCollection().AddExpression(BasecolorExpression);
+		MaterialEditorOnly->BaseColor.Expression = BasecolorExpression;
 
 		MaterialNodeY += MaterialNodeStepY;
 	}
@@ -1053,8 +1055,8 @@ UMaterial* FMaterialUtilities::CreateMaterial(const FFlattenMaterial& InFlattenM
 		BaseColorExpression->Constant = BaseColor;
 		BaseColorExpression->MaterialExpressionEditorX = -400;
 		BaseColorExpression->MaterialExpressionEditorY = MaterialNodeY;
-		Material->Expressions.Add(BaseColorExpression);
-		Material->BaseColor.Expression = BaseColorExpression;
+		Material->GetExpressionCollection().AddExpression(BaseColorExpression);
+		MaterialEditorOnly->BaseColor.Expression = BaseColorExpression;
 
 		MaterialNodeY += MaterialNodeStepY;
 	}
@@ -1120,39 +1122,39 @@ UMaterial* FMaterialUtilities::CreateMaterial(const FFlattenMaterial& InFlattenM
 		MergedExpression->SamplerType = EMaterialSamplerType::SAMPLERTYPE_Color;
 		MergedExpression->MaterialExpressionEditorX = -400;
 		MergedExpression->MaterialExpressionEditorY = MaterialNodeY;
-		Material->Expressions.Add(MergedExpression);
+		Material->GetExpressionCollection().AddExpression(MergedExpression);
 
 		// Metallic
 		if (bHasMetallic)
 		{
-			Material->Metallic.Expression = MergedExpression;
-			Material->Metallic.Mask = Material->Metallic.Expression->GetOutputs()[0].Mask;
-			Material->Metallic.MaskR = 1;
-			Material->Metallic.MaskG = 0;
-			Material->Metallic.MaskB = 0;
-			Material->Metallic.MaskA = 0;
+			MaterialEditorOnly->Metallic.Expression = MergedExpression;
+			MaterialEditorOnly->Metallic.Mask = MaterialEditorOnly->Metallic.Expression->GetOutputs()[0].Mask;
+			MaterialEditorOnly->Metallic.MaskR = 1;
+			MaterialEditorOnly->Metallic.MaskG = 0;
+			MaterialEditorOnly->Metallic.MaskB = 0;
+			MaterialEditorOnly->Metallic.MaskA = 0;
 		}
 
 		// Specular
 		if (bHasSpecular)
 		{
-			Material->Specular.Expression = MergedExpression;
-			Material->Specular.Mask = Material->Specular.Expression->GetOutputs()[0].Mask;
-			Material->Specular.MaskR = 0;
-			Material->Specular.MaskG = 1;
-			Material->Specular.MaskB = 0;
-			Material->Specular.MaskA = 0;
+			MaterialEditorOnly->Specular.Expression = MergedExpression;
+			MaterialEditorOnly->Specular.Mask = MaterialEditorOnly->Specular.Expression->GetOutputs()[0].Mask;
+			MaterialEditorOnly->Specular.MaskR = 0;
+			MaterialEditorOnly->Specular.MaskG = 1;
+			MaterialEditorOnly->Specular.MaskB = 0;
+			MaterialEditorOnly->Specular.MaskA = 0;
 		}
 
 		// Roughness
 		if (bHasRoughness)
 		{
-			Material->Roughness.Expression = MergedExpression;
-			Material->Roughness.Mask = Material->Roughness.Expression->GetOutputs()[0].Mask;
-			Material->Roughness.MaskR = 0;
-			Material->Roughness.MaskG = 0;
-			Material->Roughness.MaskB = 1;
-			Material->Roughness.MaskA = 0;
+			MaterialEditorOnly->Roughness.Expression = MergedExpression;
+			MaterialEditorOnly->Roughness.Mask = MaterialEditorOnly->Roughness.Expression->GetOutputs()[0].Mask;
+			MaterialEditorOnly->Roughness.MaskR = 0;
+			MaterialEditorOnly->Roughness.MaskG = 0;
+			MaterialEditorOnly->Roughness.MaskB = 1;
+			MaterialEditorOnly->Roughness.MaskA = 0;
 		}
 
 		MaterialNodeY += MaterialNodeStepY;
@@ -1172,8 +1174,8 @@ UMaterial* FMaterialUtilities::CreateMaterial(const FFlattenMaterial& InFlattenM
 			MetallicExpression->SamplerType = EMaterialSamplerType::SAMPLERTYPE_Color;
 			MetallicExpression->MaterialExpressionEditorX = -400;
 			MetallicExpression->MaterialExpressionEditorY = MaterialNodeY;
-			Material->Expressions.Add(MetallicExpression);
-			Material->Metallic.Expression = MetallicExpression;
+			Material->GetExpressionCollection().AddExpression(MetallicExpression);
+			MaterialEditorOnly->Metallic.Expression = MetallicExpression;
 
 			MaterialNodeY += MaterialNodeStepY;
 		}
@@ -1191,8 +1193,8 @@ UMaterial* FMaterialUtilities::CreateMaterial(const FFlattenMaterial& InFlattenM
 			SpecularExpression->SamplerType = EMaterialSamplerType::SAMPLERTYPE_Color;
 			SpecularExpression->MaterialExpressionEditorX = -400;
 			SpecularExpression->MaterialExpressionEditorY = MaterialNodeY;
-			Material->Expressions.Add(SpecularExpression);
-			Material->Specular.Expression = SpecularExpression;
+			Material->GetExpressionCollection().AddExpression(SpecularExpression);
+			MaterialEditorOnly->Specular.Expression = SpecularExpression;
 
 			MaterialNodeY += MaterialNodeStepY;
 		}
@@ -1210,8 +1212,8 @@ UMaterial* FMaterialUtilities::CreateMaterial(const FFlattenMaterial& InFlattenM
 			RoughnessExpression->SamplerType = EMaterialSamplerType::SAMPLERTYPE_Color;
 			RoughnessExpression->MaterialExpressionEditorX = -400;
 			RoughnessExpression->MaterialExpressionEditorY = MaterialNodeY;
-			Material->Expressions.Add(RoughnessExpression);
-			Material->Roughness.Expression = RoughnessExpression;
+			Material->GetExpressionCollection().AddExpression(RoughnessExpression);
+			MaterialEditorOnly->Roughness.Expression = RoughnessExpression;
 
 			MaterialNodeY += MaterialNodeStepY;
 		}
@@ -1223,8 +1225,8 @@ UMaterial* FMaterialUtilities::CreateMaterial(const FFlattenMaterial& InFlattenM
 		MetallicExpression->R = MaterialProxySettings.bMetallicMap ? FLinearColor(InFlattenMaterial.GetPropertySamples(EFlattenMaterialProperties::Metallic)[0]).R : MaterialProxySettings.MetallicConstant;
 		MetallicExpression->MaterialExpressionEditorX = -400;
 		MetallicExpression->MaterialExpressionEditorY = MaterialNodeY;
-		Material->Expressions.Add(MetallicExpression);
-		Material->Metallic.Expression = MetallicExpression;
+		Material->GetExpressionCollection().AddExpression(MetallicExpression);
+		MaterialEditorOnly->Metallic.Expression = MetallicExpression;
 
 		MaterialNodeY += MaterialNodeStepY;
 	}
@@ -1236,8 +1238,8 @@ UMaterial* FMaterialUtilities::CreateMaterial(const FFlattenMaterial& InFlattenM
 		SpecularExpression->R = MaterialProxySettings.bSpecularMap ? FLinearColor(InFlattenMaterial.GetPropertySamples(EFlattenMaterialProperties::Specular)[0]).R : MaterialProxySettings.SpecularConstant;
 		SpecularExpression->MaterialExpressionEditorX = -400;
 		SpecularExpression->MaterialExpressionEditorY = MaterialNodeY;
-		Material->Expressions.Add(SpecularExpression);
-		Material->Specular.Expression = SpecularExpression;
+		Material->GetExpressionCollection().AddExpression(SpecularExpression);
+		MaterialEditorOnly->Specular.Expression = SpecularExpression;
 
 		MaterialNodeY += MaterialNodeStepY;
 	}
@@ -1249,8 +1251,8 @@ UMaterial* FMaterialUtilities::CreateMaterial(const FFlattenMaterial& InFlattenM
 		RoughnessExpression->R = MaterialProxySettings.bRoughnessMap ? FLinearColor(InFlattenMaterial.GetPropertySamples(EFlattenMaterialProperties::Roughness)[0]).R : MaterialProxySettings.RoughnessConstant;
 		RoughnessExpression->MaterialExpressionEditorX = -400;
 		RoughnessExpression->MaterialExpressionEditorY = MaterialNodeY;
-		Material->Expressions.Add(RoughnessExpression);
-		Material->Roughness.Expression = RoughnessExpression;
+		Material->GetExpressionCollection().AddExpression(RoughnessExpression);
+		MaterialEditorOnly->Roughness.Expression = RoughnessExpression;
 
 		MaterialNodeY += MaterialNodeStepY;
 	}
@@ -1268,8 +1270,8 @@ UMaterial* FMaterialUtilities::CreateMaterial(const FFlattenMaterial& InFlattenM
 		NormalExpression->SamplerType = EMaterialSamplerType::SAMPLERTYPE_Normal;
 		NormalExpression->MaterialExpressionEditorX = -400;
 		NormalExpression->MaterialExpressionEditorY = MaterialNodeY;
-		Material->Expressions.Add(NormalExpression);
-		Material->Normal.Expression = NormalExpression;
+		Material->GetExpressionCollection().AddExpression(NormalExpression);
+		MaterialEditorOnly->Normal.Expression = NormalExpression;
 
 		MaterialNodeY+= MaterialNodeStepY;
 	}
@@ -1285,8 +1287,8 @@ UMaterial* FMaterialUtilities::CreateMaterial(const FFlattenMaterial& InFlattenM
 			EmissiveColorExpression->Constant = EmissiveColor.ReinterpretAsLinear() * InFlattenMaterial.EmissiveScale;
 			EmissiveColorExpression->MaterialExpressionEditorX = -400;
 			EmissiveColorExpression->MaterialExpressionEditorY = MaterialNodeY;
-			Material->Expressions.Add(EmissiveColorExpression);
-			Material->EmissiveColor.Expression = EmissiveColorExpression;
+			Material->GetExpressionCollection().AddExpression(EmissiveColorExpression);
+			MaterialEditorOnly->EmissiveColor.Expression = EmissiveColorExpression;
 
 			MaterialNodeY += MaterialNodeStepY;
 		}
@@ -1304,16 +1306,16 @@ UMaterial* FMaterialUtilities::CreateMaterial(const FFlattenMaterial& InFlattenM
 		EmissiveColorExpression->SamplerType = EMaterialSamplerType::SAMPLERTYPE_Color;
 		EmissiveColorExpression->MaterialExpressionEditorX = -400;
 		EmissiveColorExpression->MaterialExpressionEditorY = MaterialNodeY;
-		Material->Expressions.Add(EmissiveColorExpression);
+		Material->GetExpressionCollection().AddExpression(EmissiveColorExpression);
 
 		UMaterialExpressionMultiply* EmissiveColorScale = NewObject<UMaterialExpressionMultiply>(Material);
 		EmissiveColorScale->A.Expression = EmissiveColorExpression;
 		EmissiveColorScale->ConstB = InFlattenMaterial.EmissiveScale;
 		EmissiveColorScale->MaterialExpressionEditorX = -200;
 		EmissiveColorScale->MaterialExpressionEditorY = MaterialNodeY;
-		Material->Expressions.Add(EmissiveColorScale);
+		Material->GetExpressionCollection().AddExpression(EmissiveColorScale);
 
-		Material->EmissiveColor.Expression = EmissiveColorScale;
+		MaterialEditorOnly->EmissiveColor.Expression = EmissiveColorScale;
 		MaterialNodeY += MaterialNodeStepY;
 	}
 
@@ -1325,8 +1327,8 @@ UMaterial* FMaterialUtilities::CreateMaterial(const FFlattenMaterial& InFlattenM
 		OpacityExpression->R = Opacity.R;
 		OpacityExpression->MaterialExpressionEditorX = -400;
 		OpacityExpression->MaterialExpressionEditorY = MaterialNodeY;
-		Material->Expressions.Add(OpacityExpression);
-		Material->Opacity.Expression = OpacityExpression;
+		Material->GetExpressionCollection().AddExpression(OpacityExpression);
+		MaterialEditorOnly->Opacity.Expression = OpacityExpression;
 
 		MaterialNodeY += MaterialNodeStepY;
 	}
@@ -1343,8 +1345,8 @@ UMaterial* FMaterialUtilities::CreateMaterial(const FFlattenMaterial& InFlattenM
 		OpacityExpression->SamplerType = EMaterialSamplerType::SAMPLERTYPE_Color;
 		OpacityExpression->MaterialExpressionEditorX = -400;
 		OpacityExpression->MaterialExpressionEditorY = MaterialNodeY;
-		Material->Expressions.Add(OpacityExpression);
-		Material->Opacity.Expression = OpacityExpression;
+		Material->GetExpressionCollection().AddExpression(OpacityExpression);
+		MaterialEditorOnly->Opacity.Expression = OpacityExpression;
 		MaterialNodeY += MaterialNodeStepY;
 	}
 
@@ -1360,8 +1362,8 @@ UMaterial* FMaterialUtilities::CreateMaterial(const FFlattenMaterial& InFlattenM
 			SubSurfaceColorExpression->Constant = (SubSurfaceColor.ReinterpretAsLinear());
 			SubSurfaceColorExpression->MaterialExpressionEditorX = -400;
 			SubSurfaceColorExpression->MaterialExpressionEditorY = MaterialNodeY;
-			Material->Expressions.Add(SubSurfaceColorExpression);
-			Material->SubsurfaceColor.Expression = SubSurfaceColorExpression;
+			Material->GetExpressionCollection().AddExpression(SubSurfaceColorExpression);
+			MaterialEditorOnly->SubsurfaceColor.Expression = SubSurfaceColorExpression;
 
 			MaterialNodeY += MaterialNodeStepY;
 		}
@@ -1382,9 +1384,9 @@ UMaterial* FMaterialUtilities::CreateMaterial(const FFlattenMaterial& InFlattenM
 		SubSurfaceColorExpression->SamplerType = EMaterialSamplerType::SAMPLERTYPE_Color;
 		SubSurfaceColorExpression->MaterialExpressionEditorX = -400;
 		SubSurfaceColorExpression->MaterialExpressionEditorY = MaterialNodeY;
-		Material->Expressions.Add(SubSurfaceColorExpression);
+		Material->GetExpressionCollection().AddExpression(SubSurfaceColorExpression);
 
-		Material->SubsurfaceColor.Expression = SubSurfaceColorExpression;
+		MaterialEditorOnly->SubsurfaceColor.Expression = SubSurfaceColorExpression;
 		MaterialNodeY += MaterialNodeStepY;
 
 		Material->SetShadingModel(MSM_Subsurface);
@@ -2281,8 +2283,8 @@ bool FMaterialUtilities::ExportMaterial(struct FMaterialMergeData& InMaterialDat
 
 	// Determine whether or not certain properties can be rendered
 	const bool bRenderNormal = (Material->GetMaterial()->HasNormalConnected() || Material->GetMaterial()->bUseMaterialAttributes) && OutFlattenMaterial.ShouldGenerateDataForProperty(EFlattenMaterialProperties::Normal);
-	const bool bRenderTangent = (Material->GetMaterial()->Tangent.IsConnected() || Material->GetMaterial()->bUseMaterialAttributes) && OutFlattenMaterial.ShouldGenerateDataForProperty(EFlattenMaterialProperties::Tangent);
-	const bool bRenderEmissive = (Material->GetMaterial()->EmissiveColor.IsConnected() || Material->GetMaterial()->bUseMaterialAttributes) && OutFlattenMaterial.ShouldGenerateDataForProperty(EFlattenMaterialProperties::Emissive);
+	const bool bRenderTangent = (Material->GetMaterial()->IsPropertyConnected(MP_Tangent) || Material->GetMaterial()->bUseMaterialAttributes) && OutFlattenMaterial.ShouldGenerateDataForProperty(EFlattenMaterialProperties::Tangent);
+	const bool bRenderEmissive = (Material->GetMaterial()->IsPropertyConnected(MP_EmissiveColor) || Material->GetMaterial()->bUseMaterialAttributes) && OutFlattenMaterial.ShouldGenerateDataForProperty(EFlattenMaterialProperties::Emissive);
 	const bool bRenderOpacityMask = Material->IsPropertyActive(MP_OpacityMask) && Material->GetBlendMode() == BLEND_Masked && OutFlattenMaterial.ShouldGenerateDataForProperty(EFlattenMaterialProperties::Opacity);
 	const bool bRenderOpacity = Material->IsPropertyActive(MP_Opacity) && IsTranslucentBlendMode(Material->GetBlendMode()) && OutFlattenMaterial.ShouldGenerateDataForProperty(EFlattenMaterialProperties::Opacity);
 	const bool bRenderSubSurface = Material->IsPropertyActive(MP_SubsurfaceColor) && OutFlattenMaterial.ShouldGenerateDataForProperty(EFlattenMaterialProperties::SubSurface);
diff --git a/Engine/Source/Developer/MaterialUtilities/Private/ProxyMaterialUtilities.cpp b/Engine/Source/Developer/MaterialUtilities/Private/ProxyMaterialUtilities.cpp
--- a/Engine/Source/Developer/MaterialUtilities/Private/ProxyMaterialUtilities.cpp
+++ b/Engine/Source/Developer/MaterialUtilities/Private/ProxyMaterialUtilities.cpp
@@ -301,10 +301,10 @@ UMaterialInstanceConstant* FMaterialUtilities::CreateFlattenMaterialInstance(UPa
 		SwitchParameter.ParameterInfo.Name = TEXT("UseCustomUV");
 		SwitchParameter.Value = true;
 		SwitchParameter.bOverride = true;
-		NewStaticParameterSet.StaticSwitchParameters.Add(SwitchParameter);
+		NewStaticParameterSet.EditorOnly.StaticSwitchParameters.Add(SwitchParameter);
 
 		SwitchParameter.ParameterInfo.Name = *(TEXT("UseUV") + FString::FromInt(FlattenMaterial.UVChannel));
-		NewStaticParameterSet.StaticSwitchParameters.Add(SwitchParameter);
+		NewStaticParameterSet.EditorOnly.StaticSwitchParameters.Add(SwitchParameter);
 	}
 
 	auto CreateTextureFromDefault = [&](const FName TextureName, const FString& AssetLongName, FIntPoint Size, const TArray<FColor>& Samples)
@@ -343,7 +343,7 @@ UMaterialInstanceConstant* FMaterialUtilities::CreateFlattenMaterialInstance(UPa
 			SwitchParameter.ParameterInfo.Name = UseTexture;
 			SwitchParameter.Value = true;
 			SwitchParameter.bOverride = true;
-			NewStaticParameterSet.StaticSwitchParameters.Add(SwitchParameter);
+			NewStaticParameterSet.EditorOnly.StaticSwitchParameters.Add(SwitchParameter);
 
 			OutMaterial->SetTextureParameterValueEditorOnly(TextureName, Texture);
 			
@@ -469,17 +469,17 @@ UMaterialInstanceConstant* FMaterialUtilities::CreateFlattenMaterialInstance(UPa
 		SwitchParameter.ParameterInfo.Name = TEXT("PackMetallic");
 		SwitchParameter.Value = bPackMetallic;
 		SwitchParameter.bOverride = true;
-		NewStaticParameterSet.StaticSwitchParameters.Add(SwitchParameter);
+		NewStaticParameterSet.EditorOnly.StaticSwitchParameters.Add(SwitchParameter);
 
 		SwitchParameter.ParameterInfo.Name = TEXT("PackSpecular");
 		SwitchParameter.Value = bPackSpecular;
 		SwitchParameter.bOverride = true;
-		NewStaticParameterSet.StaticSwitchParameters.Add(SwitchParameter);
+		NewStaticParameterSet.EditorOnly.StaticSwitchParameters.Add(SwitchParameter);
 
 		SwitchParameter.ParameterInfo.Name = TEXT("PackRoughness");
 		SwitchParameter.Value = bPackRoughness;
 		SwitchParameter.bOverride = true;
-		NewStaticParameterSet.StaticSwitchParameters.Add(SwitchParameter);
+		NewStaticParameterSet.EditorOnly.StaticSwitchParameters.Add(SwitchParameter);
 
 		// Set up switch and texture values
 		OutMaterial->SetTextureParameterValueEditorOnly(PackedTextureName, PackedTexture);
diff --git a/Engine/Source/Editor/MaterialEditor/Private/MaterialEditingLibrary.cpp b/Engine/Source/Editor/MaterialEditor/Private/MaterialEditingLibrary.cpp
--- a/Engine/Source/Editor/MaterialEditor/Private/MaterialEditingLibrary.cpp
+++ b/Engine/Source/Editor/MaterialEditor/Private/MaterialEditingLibrary.cpp
@@ -390,7 +390,7 @@ int32 UMaterialEditingLibrary::GetNumMaterialExpressions(const UMaterial* Materi
 	int32 Result = 0;
 	if (Material)
 	{
-		Result = Material->Expressions.Num();
+		Result = Material->GetExpressions().Num();
 	}
 	return Result;
 }
@@ -399,8 +399,7 @@ void UMaterialEditingLibrary::DeleteAllMaterialExpressions(UMaterial* Material)
 {
 	if (Material)
 	{
-		TArray<UMaterialExpression*> AllExpressions = Material->Expressions;
-		for (UMaterialExpression* Expression : AllExpressions)
+		for (UMaterialExpression* Expression : Material->GetExpressions())
 		{
 			DeleteMaterialExpression(Material, Expression);
 		}
@@ -408,7 +407,7 @@ void UMaterialEditingLibrary::DeleteAllMaterialExpressions(UMaterial* Material)
 }
 
 /** Util to iterate over list of expressions, and break any links to specified expression */
-static void BreakLinksToExpression(TArray<UMaterialExpression*>& Expressions, UMaterialExpression* Expression)
+static void BreakLinksToExpression(TConstArrayView<TObjectPtr<UMaterialExpression>> Expressions, UMaterialExpression* Expression)
 {
 	// Need to find any other expressions which are connected to this one, and break link
 	for (UMaterialExpression* TestExp : Expressions)
@@ -433,7 +432,7 @@ void UMaterialEditingLibrary::DeleteMaterialExpression(UMaterial* Material, UMat
 	if (Material && Expression && Expression->GetOuter() == Material)
 	{
 		// Break any links to this expression
-		BreakLinksToExpression(Material->Expressions, Expression);
+		BreakLinksToExpression(Material->GetExpressions(), Expression);
 
 		// Check material parameter inputs, to make sure expression is not connected to it
 		for (int32 InputIndex = 0; InputIndex < MP_MAX; InputIndex++)
@@ -447,7 +446,7 @@ void UMaterialEditingLibrary::DeleteMaterialExpression(UMaterial* Material, UMat
 
 		Material->RemoveExpressionParameter(Expression);
 
-		Material->Expressions.Remove(Expression);
+		Material->GetExpressionCollection().RemoveExpression(Expression);
 
 		Expression->MarkAsGarbage();
 
@@ -476,13 +475,13 @@ UMaterialExpression* UMaterialEditingLibrary::DuplicateMaterialExpression(UMater
 
 		if (Material)
 		{
-			Material->Expressions.Add(NewExpression);
+			Material->GetExpressionCollection().AddExpression(NewExpression);
 			NewExpression->Material = Material;
 		}
 
 		if (MaterialFunction && !Material)
 		{
-			MaterialFunction->FunctionExpressions.Add(NewExpression);
+			MaterialFunction->GetExpressionCollection().AddExpression(NewExpression);
 		}
 
 		// Create a GUID for the node
@@ -520,13 +519,13 @@ UMaterialExpression* UMaterialEditingLibrary::CreateMaterialExpressionEx(UMateri
 
 		if (Material)
 		{
-			Material->Expressions.Add(NewExpression);
+			Material->GetExpressionCollection().AddExpression(NewExpression);
 			NewExpression->Material = Material;
 		}
 
 		if (MaterialFunction && !Material)
 		{
-			MaterialFunction->FunctionExpressions.Add(NewExpression);
+			MaterialFunction->GetExpressionCollection().AddExpression(NewExpression);
 		}
 
 		NewExpression->MaterialExpressionEditorX = NodePosX;
@@ -829,7 +828,7 @@ int32 UMaterialEditingLibrary::GetNumMaterialExpressionsInFunction(const UMateri
 	int32 Result = 0;
 	if (MaterialFunction)
 	{
-		Result = MaterialFunction->FunctionExpressions.Num();
+		Result = MaterialFunction->GetExpressions().Num();
 	}
 	return Result;
 }
@@ -838,8 +837,7 @@ void UMaterialEditingLibrary::DeleteAllMaterialExpressionsInFunction(UMaterialFu
 {
 	if (MaterialFunction)
 	{
-		TArray<UMaterialExpression*> AllExpressions = MaterialFunction->FunctionExpressions;
-		for (UMaterialExpression* Expression : AllExpressions)
+		for (UMaterialExpression* Expression : MaterialFunction->GetExpressions())
 		{
 			DeleteMaterialExpressionInFunction(MaterialFunction, Expression);
 		}
@@ -852,9 +850,9 @@ void UMaterialEditingLibrary::DeleteMaterialExpressionInFunction(UMaterialFuncti
 	if (MaterialFunction && Expression && Expression->GetOuter() == MaterialFunction)
 	{
 		// Break any links to this expression
-		BreakLinksToExpression(MaterialFunction->FunctionExpressions, Expression);
+		BreakLinksToExpression(MaterialFunction->GetExpressions(), Expression);
 
-		MaterialFunction->FunctionExpressions.Remove(Expression);
+		MaterialFunction->GetExpressionCollection().RemoveExpression(Expression);
 
 		Expression->MarkAsGarbage();
 
@@ -903,8 +901,7 @@ void UMaterialEditingLibrary::UpdateMaterialFunction(UMaterialFunctionInterface*
 		}
 
 		// update the world's viewports	
-		UMaterialFunctionInstance* FunctionAsInstance = Cast<UMaterialFunctionInstance>(MaterialFunction);
-		UMaterialFunction* BaseFunction = Cast<UMaterialFunction>(FunctionAsInstance ? FunctionAsInstance->GetBaseFunction() : MaterialFunction);
+		UMaterialFunction* BaseFunction = MaterialFunction->GetBaseFunction();
 
 		UMaterialEditingLibrary::RebuildMaterialInstanceEditors(BaseFunction);
 		FEditorDelegates::RefreshEditor.Broadcast();
diff --git a/Engine/Source/Editor/MaterialEditor/Private/MaterialEditor.cpp b/Engine/Source/Editor/MaterialEditor/Private/MaterialEditor.cpp
--- a/Engine/Source/Editor/MaterialEditor/Private/MaterialEditor.cpp
+++ b/Engine/Source/Editor/MaterialEditor/Private/MaterialEditor.cpp
@@ -193,7 +193,7 @@ FMatExpressionPreview::FMatExpressionPreview(UMaterialExpression* InExpression)
 	FMaterial::AddEditorLoadedMaterialResource(this);
 	FPlatformMisc::CreateGuid(Id);
 
-	check(InExpression->Material && InExpression->Material->Expressions.Contains(InExpression));
+	check(InExpression->Material && InExpression->Material->GetExpressions().Contains(InExpression));
 	SetQualityLevelProperties(GMaxRHIFeatureLevel);
 
 	UMaterial* BaseMaterial = InExpression->Material;
@@ -461,15 +461,8 @@ void FMaterialEditor::InitEditorForMaterial(UMaterial* InMaterial)
 												
 	Material->bAllowDevelopmentShaderCompile = CVarMaterialEdUseDevShaders.GetValueOnGameThread();
 
-	// Remove NULL entries, so the rest of the material editor can assume all entries of Material->Expressions are valid
-	// This can happen if an expression class was removed
-	for (int32 ExpressionIndex = Material->Expressions.Num() - 1; ExpressionIndex >= 0; ExpressionIndex--)
-	{
-		if (!Material->Expressions[ExpressionIndex])
-		{
-			Material->Expressions.RemoveAt(ExpressionIndex);
-		}
-	}
+	// Ensure there are no null entries
+	check(Material->GetExpressions().Find(nullptr) == INDEX_NONE);
 
 	TArray<FString> Groups;
 	GetAllMaterialExpressionGroups(&Groups);
@@ -640,24 +633,14 @@ void FMaterialEditor::InitMaterialEditor( const EToolkitMode::Type Mode, const T
 		MaterialFunction->MaterialGraph = Material->MaterialGraph;
 		MaterialFunction->EditorMaterial = Material;
 
-		Material->Expressions = MaterialFunction->FunctionExpressions;
-		Material->EditorComments = MaterialFunction->FunctionEditorComments;
-		Material->ExpressionExecBegin = MaterialFunction->ExpressionExecBegin;
-		Material->ExpressionExecEnd = MaterialFunction->ExpressionExecEnd;
+		Material->AssignExpressionCollection(MaterialFunction->GetExpressionCollection());
 		Material->bEnableExecWire = MaterialFunction->IsUsingControlFlow();
 		Material->bEnableNewHLSLGenerator = MaterialFunction->IsUsingNewHLSLGenerator();
 
-		// Remove NULL entries, so the rest of the material editor can assume all entries of Material->Expressions are valid
-		// This can happen if an expression class was removed
-		for (int32 ExpressionIndex = Material->Expressions.Num() - 1; ExpressionIndex >= 0; ExpressionIndex--)
-		{
-			if (!Material->Expressions[ExpressionIndex])
-			{
-				Material->Expressions.RemoveAt(ExpressionIndex);
-			}
-		}
+		// Ensure there are no null entries
+		check(Material->GetExpressions().Find(nullptr) == INDEX_NONE);
 
-		if (Material->Expressions.Num() == 0)
+		if (Material->GetExpressions().Num() == 0)
 		{
 			// If this is an empty function, create an output by default and start previewing it
 			if (FocusedGraphEdPtr.IsValid())
@@ -758,9 +741,9 @@ void FMaterialEditor::InitMaterialEditor( const EToolkitMode::Type Mode, const T
 		{
 			bool bSetPreviewExpression = false;
 			UMaterialExpressionFunctionOutput* FirstOutput = NULL;
-			for (int32 ExpressionIndex = Material->Expressions.Num() - 1; ExpressionIndex >= 0; ExpressionIndex--)
+			for (int32 ExpressionIndex = Material->GetExpressions().Num() - 1; ExpressionIndex >= 0; ExpressionIndex--)
 			{
-				UMaterialExpression* Expression = Material->Expressions[ExpressionIndex];
+				UMaterialExpression* Expression = Material->GetExpressions()[ExpressionIndex];
 
 				// Setup the expression to be used with the preview material instead of the function
 				Expression->Function = NULL;
@@ -905,17 +888,19 @@ FMaterialEditor::~FMaterialEditor()
 
 void FMaterialEditor::GetAllMaterialExpressionGroups(TArray<FString>* OutGroups)
 {
+	UMaterialEditorOnlyData* EditorOnlyData = Material->GetEditorOnlyData();
+	if (!EditorOnlyData)
+	{
+		return;
+	}
+
 	TArray<FParameterGroupData> UpdatedGroups;
-	for (int32 MaterialExpressionIndex = 0; MaterialExpressionIndex < Material->Expressions.Num(); ++MaterialExpressionIndex)
+	for (const UMaterialExpression* MaterialExpression : Material->GetExpressions())
 	{
-		UMaterialExpression* MaterialExpression = Material->Expressions[ MaterialExpressionIndex ];
-		UMaterialExpressionParameter* Param = Cast<UMaterialExpressionParameter>(MaterialExpression);
-		UMaterialExpressionTextureSampleParameter* TextureS = Cast<UMaterialExpressionTextureSampleParameter>(MaterialExpression);
-		UMaterialExpressionRuntimeVirtualTextureSampleParameter* RVTS = Cast<UMaterialExpressionRuntimeVirtualTextureSampleParameter>(MaterialExpression);
-		UMaterialExpressionFontSampleParameter* FontS = Cast<UMaterialExpressionFontSampleParameter>(MaterialExpression);
-		if (Param)
+		FMaterialParameterMetadata ParameterMeta;
+		if (MaterialExpression->GetParameterValue(ParameterMeta))
 		{
-			const FString& GroupName = Param->Group.ToString();
+			const FString GroupName = ParameterMeta.Group.ToString();
 			OutGroups->AddUnique(GroupName);
 			if (Material->AttemptInsertNewGroupName(GroupName))
 			{
@@ -923,58 +908,7 @@ void FMaterialEditor::GetAllMaterialExpressionGroups(TArray<FString>* OutGroups)
 			}
 			else
 			{
-				FParameterGroupData* ParameterGroupDataElement = Material->ParameterGroupData.FindByPredicate([&GroupName](const FParameterGroupData& DataElement)
-				{
-					return GroupName == DataElement.GroupName;
-				});
-				UpdatedGroups.Add(FParameterGroupData(GroupName, ParameterGroupDataElement->GroupSortPriority));
-			}
-		}
-		else if (TextureS)
-		{
-			const FString& GroupName = TextureS->Group.ToString();
-			OutGroups->AddUnique(GroupName);
-			if (Material->AttemptInsertNewGroupName(GroupName))
-			{
-				UpdatedGroups.Add(FParameterGroupData(GroupName, 0));
-			}
-			else
-			{
-				FParameterGroupData* ParameterGroupDataElement = Material->ParameterGroupData.FindByPredicate([&GroupName](const FParameterGroupData& DataElement)
-				{
-					return GroupName == DataElement.GroupName;
-				});
-				UpdatedGroups.Add(FParameterGroupData(GroupName, ParameterGroupDataElement->GroupSortPriority));
-			}
-		}
-		else if (RVTS)
-		{
-			const FString& GroupName = RVTS->Group.ToString();
-			OutGroups->AddUnique(GroupName);
-			if (Material->AttemptInsertNewGroupName(GroupName))
-			{
-				UpdatedGroups.Add(FParameterGroupData(GroupName, 0));
-			}
-			else
-			{
-				FParameterGroupData* ParameterGroupDataElement = Material->ParameterGroupData.FindByPredicate([&GroupName](const FParameterGroupData& DataElement)
-					{
-						return GroupName == DataElement.GroupName;
-					});
-				UpdatedGroups.Add(FParameterGroupData(GroupName, ParameterGroupDataElement->GroupSortPriority));
-			}
-		}
-		else if (FontS)
-		{
-			const FString& GroupName = FontS->Group.ToString();
-			OutGroups->AddUnique(GroupName);
-			if (Material->AttemptInsertNewGroupName(GroupName))
-			{
-				UpdatedGroups.Add(FParameterGroupData(GroupName, 0));
-			}
-			else
-			{
-				FParameterGroupData* ParameterGroupDataElement = Material->ParameterGroupData.FindByPredicate([&GroupName](const FParameterGroupData& DataElement)
+				FParameterGroupData* ParameterGroupDataElement = EditorOnlyData->ParameterGroupData.FindByPredicate([&GroupName](const FParameterGroupData& DataElement)
 				{
 					return GroupName == DataElement.GroupName;
 				});
@@ -982,7 +916,7 @@ void FMaterialEditor::GetAllMaterialExpressionGroups(TArray<FString>* OutGroups)
 			}
 		}
 	}
-	Material->ParameterGroupData = UpdatedGroups;
+	EditorOnlyData->ParameterGroupData = UpdatedGroups;
 }
 
 void FMaterialEditor::UpdatePreviewViewportsVisibility()
@@ -1277,7 +1211,7 @@ void FMaterialEditor::ExpandNode(UEdGraphNode* InNodeToExpand, UEdGraph* InSourc
 	{
 		UMaterialExpressionPinBase* PinBase = Cast<UMaterialExpressionPinBase>(Cast<UMaterialGraphNode>(Entry)->MaterialExpression);
 		PinBase->DeleteReroutePins();
-		Material->Expressions.Remove(PinBase);
+		Material->GetExpressionCollection().RemoveExpression(PinBase);
 		PinBase->MarkAsGarbage();
 		Entry->DestroyNode();
 		bPreviewExpressionDeleted |= PinBase == PreviewExpression;
@@ -1287,7 +1221,7 @@ void FMaterialEditor::ExpandNode(UEdGraphNode* InNodeToExpand, UEdGraph* InSourc
 	{
 		UMaterialExpressionPinBase* PinBase = Cast<UMaterialExpressionPinBase>(Cast<UMaterialGraphNode>(Result)->MaterialExpression);
 		PinBase->DeleteReroutePins();
-		Material->Expressions.Remove(PinBase);
+		Material->GetExpressionCollection().RemoveExpression(PinBase);
 		PinBase->MarkAsGarbage();
 		Result->DestroyNode();
 		bPreviewExpressionDeleted |= PinBase == PreviewExpression;
@@ -1303,7 +1237,7 @@ void FMaterialEditor::ExpandNode(UEdGraphNode* InNodeToExpand, UEdGraph* InSourc
 	// Remove the gateway node and source graph
 	UMaterialExpression* CompositeExpression = Cast<UMaterialGraphNode>(InNodeToExpand)->MaterialExpression;
 	CompositeExpression->Modify();
-	Material->Expressions.Remove(CompositeExpression);
+	Material->GetExpressionCollection().RemoveExpression(CompositeExpression);
 	CompositeExpression->MarkAsGarbage();
 	InNodeToExpand->DestroyNode();
 	bPreviewExpressionDeleted |= CompositeExpression == PreviewExpression;
@@ -2294,9 +2228,9 @@ void FMaterialEditor::RecenterEditor()
 	{
 		bool bSetPreviewExpression = false;
 		UMaterialExpressionFunctionOutput* FirstOutput = NULL;
-		for (int32 ExpressionIndex = Material->Expressions.Num() - 1; ExpressionIndex >= 0; ExpressionIndex--)
+		for (int32 ExpressionIndex = Material->GetExpressions().Num() - 1; ExpressionIndex >= 0; ExpressionIndex--)
 		{
-			UMaterialExpression* Expression = Material->Expressions[ExpressionIndex];
+			UMaterialExpression* Expression = Material->GetExpressions()[ExpressionIndex];
 
 			UMaterialExpressionFunctionOutput* FunctionOutput = Cast<UMaterialExpressionFunctionOutput>(Expression);
 			if (FunctionOutput)
@@ -2478,9 +2412,7 @@ void FMaterialEditor::UpdatePreviewMaterial( bool bForce )
 
 		// The preview material's expressions array must stay up to date before recompiling 
 		// So that RebuildMaterialFunctionInfo will see all the nested material functions that may need to be updated
-		ExpressionPreviewMaterial->Expressions = Material->Expressions;
-		ExpressionPreviewMaterial->ExpressionExecBegin = Material->ExpressionExecBegin;
-		ExpressionPreviewMaterial->ExpressionExecEnd = Material->ExpressionExecEnd;
+		ExpressionPreviewMaterial->AssignExpressionCollection(Material->GetExpressionCollection());
 		ExpressionPreviewMaterial->bEnableExecWire = Material->IsUsingControlFlow();
 		ExpressionPreviewMaterial->bEnableNewHLSLGenerator = Material->IsUsingNewHLSLGenerator();
 
@@ -2586,8 +2518,7 @@ bool FMaterialEditor::UpdateOriginalMaterial()
 	if (MaterialFunction)
 	{
 		// Copy the expressions back from the preview material
-		MaterialFunction->FunctionExpressions = Material->Expressions;
-		MaterialFunction->FunctionEditorComments = Material->EditorComments;
+		MaterialFunction->AssignExpressionCollection(Material->GetExpressionCollection());
 
 		// Preserve the thumbnail info
 		UThumbnailInfo* OriginalThumbnailInfo = MaterialFunction->ParentFunction->ThumbnailInfo;
@@ -2620,10 +2551,9 @@ bool FMaterialEditor::UpdateOriginalMaterial()
 		// Restore RF_Standalone on the original material function, as it had been removed from the preview material so that it could be GC'd.
 		MaterialFunction->ParentFunction->SetFlags( RF_Standalone );
 
-		for (int32 ExpressionIndex = 0; ExpressionIndex < MaterialFunction->ParentFunction->FunctionExpressions.Num(); ExpressionIndex++)
+		for (UMaterialExpression* CurrentExpression : MaterialFunction->ParentFunction->GetExpressions())
 		{
-			UMaterialExpression* CurrentExpression = MaterialFunction->ParentFunction->FunctionExpressions[ExpressionIndex];
-			ensureMsgf(CurrentExpression, TEXT("Invalid expression at index [%i] whilst saving material function."), ExpressionIndex);
+			ensureMsgf(CurrentExpression, TEXT("Invalid expression whilst saving material function."));
 
 			// Link the expressions back to their function
 			if (CurrentExpression)
@@ -2632,10 +2562,9 @@ bool FMaterialEditor::UpdateOriginalMaterial()
 				CurrentExpression->Function = MaterialFunction->ParentFunction;
 			}	
 		}
-		for (int32 ExpressionIndex = 0; ExpressionIndex < MaterialFunction->ParentFunction->FunctionEditorComments.Num(); ExpressionIndex++)
+		for (UMaterialExpressionComment* CurrentExpression : MaterialFunction->ParentFunction->GetEditorComments())
 		{
-			UMaterialExpressionComment* CurrentExpression = MaterialFunction->ParentFunction->FunctionEditorComments[ExpressionIndex];
-			ensureMsgf(CurrentExpression, TEXT("Invalid comment at index [%i] whilst saving material function."), ExpressionIndex);
+			ensureMsgf(CurrentExpression, TEXT("Invalid comment whilst saving material function."));
 
 			// Link the expressions back to their function
 			if (CurrentExpression)
@@ -2745,7 +2674,7 @@ void FMaterialEditor::UpdateMaterialinfoList_Old()
 				if (MaterialFunction->GetMaterialFunctionUsage() == EMaterialFunctionUsage::MaterialLayer)
 				{
 					// Material layers must have a single MA input and output only
-					for (UMaterialExpression* Expression : *MaterialFunction->GetFunctionExpressions())
+					for (UMaterialExpression* Expression : MaterialFunction->GetExpressions())
 					{
 						if (UMaterialExpressionFunctionInput* InputExpression = Cast<UMaterialExpressionFunctionInput>(Expression))
 						{
@@ -2777,7 +2706,7 @@ void FMaterialEditor::UpdateMaterialinfoList_Old()
 				else if (MaterialFunction->GetMaterialFunctionUsage() == EMaterialFunctionUsage::MaterialLayerBlend)
 				{
 					// Material layer blends can have two MA inputs and single MA output only
-					for (UMaterialExpression* Expression : *MaterialFunction->GetFunctionExpressions())
+					for (UMaterialExpression* Expression : MaterialFunction->GetExpressions())
 					{
 						if (UMaterialExpressionFunctionInput* InputExpression = Cast<UMaterialExpressionFunctionInput>(Expression))
 						{
@@ -2816,9 +2745,9 @@ void FMaterialEditor::UpdateMaterialinfoList_Old()
 					}
 
 					bool bFoundFunctionOutput = false;
-					for (int32 ExpressionIndex = 0; ExpressionIndex < Material->Expressions.Num(); ExpressionIndex++)
+					for (UMaterialExpression* MaterialExpression : Material->GetExpressions())
 					{
-						if (Material->Expressions[ExpressionIndex]->IsA(UMaterialExpressionFunctionOutput::StaticClass()))
+						if (MaterialExpression->IsA(UMaterialExpressionFunctionOutput::StaticClass()))
 						{
 							bFoundFunctionOutput = true;
 							break;
@@ -3027,9 +2956,9 @@ void FMaterialEditor::UpdateMaterialInfoList()
 		CompileErrors = ExpressionPreviewMaterial->GetMaterialResource(GMaxRHIFeatureLevel)->GetCompileErrors();
 
 		bool bFoundFunctionOutput = false;
-		for (int32 ExpressionIndex = 0; ExpressionIndex < Material->Expressions.Num(); ExpressionIndex++)
+		for (UMaterialExpression* MaterialExpression : Material->GetExpressions())
 		{
-			if (Material->Expressions[ExpressionIndex]->IsA(UMaterialExpressionFunctionOutput::StaticClass()))
+			if (MaterialExpression->IsA(UMaterialExpressionFunctionOutput::StaticClass()))
 			{
 				bFoundFunctionOutput = true;
 				break;
@@ -3101,14 +3030,14 @@ void FMaterialEditor::UpdateGraphNodeStates()
 	FStaticParameterSet StaticSwitchSet;
 	if (bPreviewFeaturesChanged && bPreviewStaticSwitches)
 	{
-		for (UMaterialExpression* Expression : Material->Expressions)
+		for (UMaterialExpression* Expression : Material->GetExpressions())
 		{
 			if (UMaterialExpressionStaticSwitchParameter* StaticSwitch = Cast<UMaterialExpressionStaticSwitchParameter>(Expression))
 			{
 				FStaticSwitchParameter SwitchParam;
 				SwitchParam.Value = StaticSwitch->DefaultValue;
 				SwitchParam.ExpressionGUID = StaticSwitch->ExpressionGUID;
-				StaticSwitchSet.StaticSwitchParameters.Add(SwitchParam);
+				StaticSwitchSet.EditorOnly.StaticSwitchParameters.Add(SwitchParam);
 			}
 		}
 	}
@@ -4132,7 +4061,7 @@ void FMaterialEditor::OnSelectNamedRerouteUsages()
 			{
 				UMaterialExpression* CurrentSelectedExpression = GraphNode->MaterialExpression;
 				UMaterialExpressionNamedRerouteDeclaration* Declaration = Cast<UMaterialExpressionNamedRerouteDeclaration>(CurrentSelectedExpression);
-				for(UMaterialExpression* Expression : Material->Expressions)
+				for(UMaterialExpression* Expression : Material->GetExpressions())
 				{
 					auto* Usage = Cast<UMaterialExpressionNamedRerouteUsage>(Expression);
 					if (Usage && Usage->Declaration == Declaration)
@@ -4281,7 +4210,7 @@ void FMaterialEditor::OnConvertNamedRerouteToReroute()
 				}
 				DeclarationGraphNode->DestroyNode();
 
-				for(UMaterialExpression* Expression : Material->Expressions)
+				for(UMaterialExpression* Expression : Material->GetExpressions())
 				{
 					auto* Usage = Cast<UMaterialExpressionNamedRerouteUsage>(Expression);
 					if (Usage && Usage->Declaration == Declaration)
@@ -4943,7 +4872,7 @@ void FMaterialEditor::SetPreviewExpression(UMaterialExpression* NewPreviewExpres
 		}
 		// If we are already previewing the selected expression toggle previewing off
 		PreviewExpression = NULL;
-		ExpressionPreviewMaterial->Expressions.Empty();
+		ExpressionPreviewMaterial->GetExpressionCollection().Empty();
 		SetPreviewMaterial( Material );
 		// Recompile the preview material to get changes that might have been made during previewing
 		UpdatePreviewMaterial();
@@ -4983,9 +4912,7 @@ void FMaterialEditor::SetPreviewExpression(UMaterialExpression* NewPreviewExpres
 
 		// The expression preview material's expressions array must stay up to date before recompiling 
 		// So that RebuildMaterialFunctionInfo will see all the nested material functions that may need to be updated
-		ExpressionPreviewMaterial->Expressions = Material->Expressions;
-		ExpressionPreviewMaterial->ExpressionExecBegin = Material->ExpressionExecBegin;
-		ExpressionPreviewMaterial->ExpressionExecEnd = Material->ExpressionExecEnd;
+		ExpressionPreviewMaterial->AssignExpressionCollection(Material->GetExpressionCollection());
 
 		// The preview window should now show the expression preview material
 		SetPreviewMaterial( ExpressionPreviewMaterial );
@@ -5185,7 +5112,7 @@ UMaterialExpressionComment* FMaterialEditor::CreateNewMaterialExpressionComment(
 		NewComment = NewObject<UMaterialExpressionComment>(ExpressionOuter, NAME_None, RF_Transactional);
 
 		// Add to the list of comments associated with this material.
-		Material->EditorComments.Add( NewComment );
+		Material->GetExpressionCollection().AddComment( NewComment );
 
 		FSlateRect Bounds;
 		if (FocusedGraphEd && FocusedGraphEd->GetBoundsForSelectedNodes(Bounds, 50.0f))
@@ -5466,7 +5393,7 @@ void FMaterialEditor::DeleteNodesInternal(const TArray<class UEdGraphNode*>& Nod
 				if(MaterialExpression)
 				{
 					MaterialExpression->Modify();
-					Material->Expressions.Remove( MaterialExpression );
+					Material->GetExpressionCollection().RemoveExpression( MaterialExpression );
 					Material->RemoveExpressionParameter(MaterialExpression);
 					// Make sure the deleted expression is caught by gc
 					MaterialExpression->MarkAsGarbage();
@@ -5477,7 +5404,7 @@ void FMaterialEditor::DeleteNodesInternal(const TArray<class UEdGraphNode*>& Nod
 				CommentNode->Modify();
 				CommentNode->BreakAllNodeLinks();
 				CommentNode->MaterialExpressionComment->Modify();
-				Material->EditorComments.Remove( CommentNode->MaterialExpressionComment );
+				Material->GetExpressionCollection().RemoveComment( CommentNode->MaterialExpressionComment );
 			}
 
 			// Now that we are done with the node, remove it
@@ -5661,7 +5588,7 @@ void FMaterialEditor::PasteNodesHere(const FVector2D& Location, const class UEdG
 			}
 
 			NewMaterialExpressions.Add(NewExpression);
-			Material->Expressions.Add(NewExpression);
+			Material->GetExpressionCollection().AddExpression(NewExpression);
 
 			ensure(NewExpression->GraphNode == GraphNode);
 			PostPasteMaterialExpression(NewExpression);
@@ -5674,7 +5601,7 @@ void FMaterialEditor::PasteNodesHere(const FVector2D& Location, const class UEdG
 				CommentNode->MaterialExpressionComment->Material = Material;
 				CommentNode->MaterialExpressionComment->Function = MaterialFunction;
 				CommentNode->MaterialExpressionComment->SubgraphExpression = ExpressionGraph->SubgraphExpression;
-				Material->EditorComments.Add(CommentNode->MaterialExpressionComment);
+				Material->GetExpressionCollection().AddComment(CommentNode->MaterialExpressionComment);
 			}
 		}
 
@@ -5889,10 +5816,10 @@ void FMaterialEditor::UndoGraphAction()
 {
 	FlushRenderingCommands();
 	
-	int32 NumExpressions = Material->Expressions.Num();
+	int32 NumExpressions = Material->GetExpressions().Num();
 	GEditor->UndoTransaction();
 
-	if(NumExpressions != Material->Expressions.Num())
+	if(NumExpressions != Material->GetExpressions().Num())
 	{
 		Material->BuildEditorParameterList();
 	}
@@ -5902,7 +5829,7 @@ void FMaterialEditor::RedoGraphAction()
 {
 	FlushRenderingCommands();
 
-	int32 NumExpressions = Material->Expressions.Num();
+	int32 NumExpressions = Material->GetExpressions().Num();
 	GEditor->RedoTransaction();
 
 	if (TSharedPtr<SGraphEditor> FocusedGraphEd = FocusedGraphEdPtr.Pin())
@@ -5921,7 +5848,7 @@ void FMaterialEditor::RedoGraphAction()
 		FocusedGraphEdPtr.Pin()->NotifyGraphChanged();
 	}
 
-	if(NumExpressions != Material->Expressions.Num())
+	if(NumExpressions != Material->GetExpressions().Num())
 	{
 		Material->BuildEditorParameterList();
 	}
@@ -6331,10 +6258,8 @@ void FMaterialEditor::RefreshExpressionPreviews(bool bForceRefreshAll /*= false*
 		// Refresh all expression previews.
 		FMaterial::DeferredDeleteArray(ExpressionPreviews);
 
-		for (int32 ExpressionIndex = 0; ExpressionIndex < Material->Expressions.Num(); ++ExpressionIndex)
+		for (UMaterialExpression* MaterialExpression : Material->GetExpressions())
 		{
-			UMaterialExpression* MaterialExpression = Material->Expressions[ExpressionIndex];
-
 			UMaterialGraphNode* GraphNode = Cast<UMaterialGraphNode>(MaterialExpression->GraphNode);
 			if (GraphNode)
 			{
@@ -6345,10 +6270,9 @@ void FMaterialEditor::RefreshExpressionPreviews(bool bForceRefreshAll /*= false*
 	else
 	{
 		// Only refresh expressions that are marked for realtime update.
-		for ( int32 ExpressionIndex = 0 ; ExpressionIndex < Material->Expressions.Num() ; ++ExpressionIndex )
+		for (UMaterialExpression* MaterialExpression : Material->GetExpressions())
 		{
-			UMaterialExpression* MaterialExpression = Material->Expressions[ ExpressionIndex ];
-			RefreshExpressionPreview( MaterialExpression, false );
+			RefreshExpressionPreview(MaterialExpression, false);
 		}
 	}
 
@@ -6356,9 +6280,8 @@ void FMaterialEditor::RefreshExpressionPreviews(bool bForceRefreshAll /*= false*
 	ExpressionPreviewsBeingCompiled.Empty(50);
 
 	// Go through all expression previews and create new ones as needed, and maintain a list of previews that are being compiled
-	for( int32 ExpressionIndex = 0; ExpressionIndex < Material->Expressions.Num(); ++ExpressionIndex )
+	for (UMaterialExpression* MaterialExpression : Material->GetExpressions())
 	{
-		UMaterialExpression* MaterialExpression = Material->Expressions[ ExpressionIndex ];
 		if (MaterialExpression && !MaterialExpression->IsA(UMaterialExpressionComment::StaticClass()) )
 		{
 			bool bNewlyCreated;
@@ -6818,7 +6741,7 @@ void FMaterialEditor::CleanUnusedExpressions()
 				}
 
 				MaterialExpression->Modify();
-				Material->Expressions.Remove(MaterialExpression);
+				Material->GetExpressionCollection().RemoveExpression(MaterialExpression);
 				Material->RemoveExpressionParameter(MaterialExpression);
 				// Make sure the deleted expression is caught by gc
 				MaterialExpression->MarkAsGarbage();
@@ -7182,7 +7105,7 @@ void FMaterialEditor::UpdateStatsMaterials()
 		FString EmptyMaterialName = FString(TEXT("MEStatsMaterial_Empty_")) + Material->GetName();
 		EmptyMaterial = (UMaterial*)StaticDuplicateObject(Material, GetTransientPackage(), *EmptyMaterialName, ~RF_Standalone, UPreviewMaterial::StaticClass());
 
-		EmptyMaterial->Expressions.Empty();
+		EmptyMaterial->GetExpressionCollection().Empty();
 
 		//Disconnect all properties from the expressions
 		for (int32 PropIdx = 0; PropIdx < MP_MAX; ++PropIdx)
diff --git a/Engine/Source/Editor/MaterialEditor/Private/MaterialEditorUtilities.cpp b/Engine/Source/Editor/MaterialEditor/Private/MaterialEditorUtilities.cpp
--- a/Engine/Source/Editor/MaterialEditor/Private/MaterialEditorUtilities.cpp
+++ b/Engine/Source/Editor/MaterialEditor/Private/MaterialEditorUtilities.cpp
@@ -28,6 +28,7 @@
 #include "Materials/MaterialExpressionMaterialAttributeLayers.h"
 #include "Materials/MaterialExpressionRerouteBase.h"
 #include "Materials/MaterialExpressionExecBegin.h"
+#include "Materials/MaterialExpressionExecEnd.h"
 
 #include "Toolkits/ToolkitManager.h"
 #include "MaterialEditor.h"
@@ -258,7 +259,7 @@ void FMaterialEditorUtilities::GetVisibleMaterialParameters(const UMaterial* Mat
 
 	if (Material->IsUsingControlFlow())
 	{
-		GetVisibleMaterialParametersFromExpression(FMaterialExpressionKey(Material->ExpressionExecBegin, INDEX_NONE), MaterialInstance, VisibleExpressions, FunctionStack);
+		GetVisibleMaterialParametersFromExpression(FMaterialExpressionKey(Material->GetExpressionExecBegin(), INDEX_NONE), MaterialInstance, VisibleExpressions, FunctionStack);
 	}
 	else
 	{
@@ -370,8 +371,7 @@ void FMaterialEditorUtilities::InitExpressions(UMaterial* Material)
 {
 	FString ParmName;
 
-	Material->EditorComments.Empty();
-	Material->Expressions.Empty();
+	Material->GetExpressionCollection().Empty();
 
 	TArray<UObject*> ChildObjects;
 	GetObjectsWithOuter(Material, ChildObjects, /*bIncludeNestedObjects=*/false);
@@ -384,11 +384,19 @@ void FMaterialEditorUtilities::InitExpressions(UMaterial* Material)
 			// Comment expressions are stored in a separate list.
 			if ( MaterialExpression->IsA( UMaterialExpressionComment::StaticClass() ) )
 			{
-				Material->EditorComments.Add( static_cast<UMaterialExpressionComment*>(MaterialExpression) );
+				Material->GetExpressionCollection().AddComment( static_cast<UMaterialExpressionComment*>(MaterialExpression) );
 			}
 			else
 			{
-				Material->Expressions.Add( MaterialExpression );
+				Material->GetExpressionCollection().AddExpression( MaterialExpression );
+				if (MaterialExpression->IsA(UMaterialExpressionExecBegin::StaticClass()))
+				{
+					Material->GetExpressionCollection().ExpressionExecBegin = static_cast<UMaterialExpressionExecBegin*>(MaterialExpression);
+				}
+				else if (MaterialExpression->IsA(UMaterialExpressionExecEnd::StaticClass()))
+				{
+					Material->GetExpressionCollection().ExpressionExecEnd = static_cast<UMaterialExpressionExecEnd*>(MaterialExpression);
+				}
 			}
 		}
 	}
@@ -397,18 +405,15 @@ void FMaterialEditorUtilities::InitExpressions(UMaterial* Material)
 
 	// Propagate RF_Transactional to all referenced material expressions.
 	Material->SetFlags( RF_Transactional );
-	for( int32 MaterialExpressionIndex = 0 ; MaterialExpressionIndex < Material->Expressions.Num() ; ++MaterialExpressionIndex )
+	for(UMaterialExpression* MaterialExpression : Material->GetExpressions())
 	{
-		UMaterialExpression* MaterialExpression = Material->Expressions[ MaterialExpressionIndex ];
-
 		if(MaterialExpression)
 		{
 			MaterialExpression->SetFlags( RF_Transactional );
 		}
 	}
-	for( int32 MaterialExpressionIndex = 0 ; MaterialExpressionIndex < Material->EditorComments.Num() ; ++MaterialExpressionIndex )
+	for(UMaterialExpressionComment* Comment : Material->GetEditorComments())
 	{
-		UMaterialExpressionComment* Comment = Material->EditorComments[ MaterialExpressionIndex ];
 		Comment->SetFlags( RF_Transactional );
 	}
 }
diff --git a/Engine/Source/Editor/MaterialEditor/Private/MaterialInstanceEditor.cpp b/Engine/Source/Editor/MaterialEditor/Private/MaterialInstanceEditor.cpp
--- a/Engine/Source/Editor/MaterialEditor/Private/MaterialInstanceEditor.cpp
+++ b/Engine/Source/Editor/MaterialEditor/Private/MaterialInstanceEditor.cpp
@@ -304,26 +304,23 @@ void FMaterialInstanceEditor::InitEditorForMaterialFunction(UMaterialFunctionIns
 	FunctionMaterialProxy->SetFlags(RF_Transactional);
 	FunctionMaterialProxy->bIsFunctionPreviewMaterial = true;
 
-	UMaterialFunctionInterface* BaseFunction = MaterialFunctionInstance;
-	while (UMaterialFunctionInstance* Instance = Cast<UMaterialFunctionInstance>(BaseFunction))
+	UMaterialFunctionInterface* BaseFunctionInterface = MaterialFunctionInstance;
+	while (UMaterialFunctionInstance* Instance = Cast<UMaterialFunctionInstance>(BaseFunctionInterface))
 	{
-		BaseFunction = Instance->GetBaseFunction();
+		BaseFunctionInterface = Instance->GetBaseFunction();
+	}
+	if (UMaterialFunction* BaseFunction = Cast<UMaterialFunction>(BaseFunctionInterface))
+	{
+		FunctionMaterialProxy->AssignExpressionCollection(BaseFunction->GetExpressionCollection());
 	}
-	const TArray<TObjectPtr<UMaterialExpression>>* FunctionExpressions = BaseFunction ? BaseFunction->GetFunctionExpressions() : nullptr;
-	FunctionMaterialProxy->Expressions = FunctionExpressions ? *FunctionExpressions : TArray<TObjectPtr<UMaterialExpression>>();
 
 	// Set expressions to be used with preview material
 	bool bSetPreviewExpression = false;
 	UMaterialExpressionFunctionOutput* FirstOutput = NULL;
-	for (int32 ExpressionIndex = FunctionMaterialProxy->Expressions.Num() - 1; ExpressionIndex >= 0; --ExpressionIndex)
+	for (int32 ExpressionIndex = FunctionMaterialProxy->GetExpressions().Num() - 1; ExpressionIndex >= 0; --ExpressionIndex)
 	{
-		UMaterialExpression* Expression = FunctionMaterialProxy->Expressions[ExpressionIndex];
-
-		if (!Expression)
-		{
-			FunctionMaterialProxy->Expressions.RemoveAt(ExpressionIndex);
-			continue;
-		}
+		UMaterialExpression* Expression = FunctionMaterialProxy->GetExpressions()[ExpressionIndex];
+		check(Expression);
 
 		Expression->Function = NULL;
 		Expression->Material = FunctionMaterialProxy;
@@ -529,8 +526,8 @@ void FMaterialInstanceEditor::ReInitMaterialFunctionProxies()
 		TArray<FFontParameterValue> FontParameterValues = FunctionInstanceProxy->FontParameterValues;
 
 		const FStaticParameterSet& OldStaticParameters = FunctionInstanceProxy->GetStaticParameters();
-		TArray<FStaticSwitchParameter> StaticSwitchParameters = OldStaticParameters.StaticSwitchParameters;
-		TArray<FStaticComponentMaskParameter> StaticComponentMaskParameters = OldStaticParameters.StaticComponentMaskParameters;
+		TArray<FStaticSwitchParameter> StaticSwitchParameters = OldStaticParameters.EditorOnly.StaticSwitchParameters;
+		TArray<FStaticComponentMaskParameter> StaticComponentMaskParameters = OldStaticParameters.EditorOnly.StaticComponentMaskParameters;
 
 		// Regenerate proxies
 		InitEditorForMaterialFunction(MaterialFunctionOriginal);
@@ -605,26 +602,26 @@ void FMaterialInstanceEditor::ReInitMaterialFunctionProxies()
 		FStaticParameterSet StaticParametersOverride = FunctionInstanceProxy->GetStaticParameters();
 
 		FunctionInstanceProxy->GetAllStaticSwitchParameterInfo(OutParameterInfo, Guids);
-		StaticParametersOverride.StaticSwitchParameters.Empty();
+		StaticParametersOverride.EditorOnly.StaticSwitchParameters.Empty();
 		for (FStaticSwitchParameter& StaticSwitchParameter : StaticSwitchParameters)
 		{
 			int32 Index = Guids.Find(StaticSwitchParameter.ExpressionGUID);
 			if (Index != INDEX_NONE)
 			{
-				StaticParametersOverride.StaticSwitchParameters.Add(StaticSwitchParameter);
-				StaticParametersOverride.StaticSwitchParameters.Last().ParameterInfo = OutParameterInfo[Index];
+				StaticParametersOverride.EditorOnly.StaticSwitchParameters.Add(StaticSwitchParameter);
+				StaticParametersOverride.EditorOnly.StaticSwitchParameters.Last().ParameterInfo = OutParameterInfo[Index];
 			}
 		}
 
 		FunctionInstanceProxy->GetAllStaticComponentMaskParameterInfo(OutParameterInfo, Guids);
-		StaticParametersOverride.StaticComponentMaskParameters.Empty();
+		StaticParametersOverride.EditorOnly.StaticComponentMaskParameters.Empty();
 		for (FStaticComponentMaskParameter& StaticComponentMaskParameter : StaticComponentMaskParameters)
 		{
 			int32 Index = Guids.Find(StaticComponentMaskParameter.ExpressionGUID);
 			if (Index != INDEX_NONE)
 			{
-				StaticParametersOverride.StaticComponentMaskParameters.Add(StaticComponentMaskParameter);
-				StaticParametersOverride.StaticComponentMaskParameters.Last().ParameterInfo = OutParameterInfo[Index];
+				StaticParametersOverride.EditorOnly.StaticComponentMaskParameters.Add(StaticComponentMaskParameter);
+				StaticParametersOverride.EditorOnly.StaticComponentMaskParameters.Last().ParameterInfo = OutParameterInfo[Index];
 			}
 		}
 
diff --git a/Engine/Source/Editor/MaterialEditor/Private/MaterialLayersFunctionsCustomization.cpp b/Engine/Source/Editor/MaterialEditor/Private/MaterialLayersFunctionsCustomization.cpp
--- a/Engine/Source/Editor/MaterialEditor/Private/MaterialLayersFunctionsCustomization.cpp
+++ b/Engine/Source/Editor/MaterialEditor/Private/MaterialLayersFunctionsCustomization.cpp
@@ -56,20 +56,20 @@ FMaterialLayersFunctionsCustomization::FMaterialLayersFunctionsCustomization(con
 	//Fixup for adding new bool arrays to the class
 	if (MaterialLayersFunctions)
 	{
-		if (MaterialLayersFunctions->Layers.Num() != MaterialLayersFunctions->RestrictToLayerRelatives.Num())
+		if (MaterialLayersFunctions->Layers.Num() != MaterialLayersFunctions->EditorOnly.RestrictToLayerRelatives.Num())
 		{
-			int32 OriginalSize = MaterialLayersFunctions->RestrictToLayerRelatives.Num();
+			int32 OriginalSize = MaterialLayersFunctions->EditorOnly.RestrictToLayerRelatives.Num();
 			for (int32 LayerIt = 0; LayerIt < MaterialLayersFunctions->Layers.Num() - OriginalSize; LayerIt++)
 			{
-				MaterialLayersFunctions->RestrictToLayerRelatives.Add(false);
+				MaterialLayersFunctions->EditorOnly.RestrictToLayerRelatives.Add(false);
 			}
 		}
-		if (MaterialLayersFunctions->Blends.Num() != MaterialLayersFunctions->RestrictToBlendRelatives.Num())
+		if (MaterialLayersFunctions->Blends.Num() != MaterialLayersFunctions->EditorOnly.RestrictToBlendRelatives.Num())
 		{
-			int32 OriginalSize = MaterialLayersFunctions->RestrictToBlendRelatives.Num();
+			int32 OriginalSize = MaterialLayersFunctions->EditorOnly.RestrictToBlendRelatives.Num();
 			for (int32 BlendIt = 0; BlendIt < MaterialLayersFunctions->Blends.Num() - OriginalSize; BlendIt++)
 			{
-				MaterialLayersFunctions->RestrictToBlendRelatives.Add(false);
+				MaterialLayersFunctions->EditorOnly.RestrictToBlendRelatives.Add(false);
 			}
 		}
 	}
@@ -250,7 +250,7 @@ void FMaterialLayersFunctionsCustomization::OnNameChanged(const FText& InText, E
 {
 	const FScopedTransaction Transaction(LOCTEXT("RenamedSection", "Renamed layer and blend section"));
 	SavedStructPropertyHandle->NotifyPreChange();
-	MaterialLayersFunctions->LayerNames[Counter] = InText;
+	MaterialLayersFunctions->EditorOnly.LayerNames[Counter] = InText;
 	SavedStructPropertyHandle->NotifyPostChange(EPropertyChangeType::ValueSet);
 };
 #endif
diff --git a/Engine/Source/Editor/MaterialEditor/Private/MaterialPropertyHelpers.cpp b/Engine/Source/Editor/MaterialEditor/Private/MaterialPropertyHelpers.cpp
--- a/Engine/Source/Editor/MaterialEditor/Private/MaterialPropertyHelpers.cpp
+++ b/Engine/Source/Editor/MaterialEditor/Private/MaterialPropertyHelpers.cpp
@@ -197,7 +197,7 @@ bool FMaterialPropertyHelpers::FilterLayerAssets(const struct FAssetData& InAsse
 					RightPath = LayerFunction->Layers[Index]->GetFName().ToString();
 				}
 			}
-			bShouldFilter = LayerFunction->RestrictToLayerRelatives[Index];
+			bShouldFilter = LayerFunction->EditorOnly.RestrictToLayerRelatives[Index];
 			BaseClassName = UMaterialFunctionMaterialLayer::StaticClass()->GetFName();
 			InstanceClassName = UMaterialFunctionMaterialLayerInstance::StaticClass()->GetFName();
 		}
@@ -213,7 +213,7 @@ bool FMaterialPropertyHelpers::FilterLayerAssets(const struct FAssetData& InAsse
 					RightPath = LayerFunction->Blends[Index]->GetFName().ToString();
 				}
 			}
-			bShouldFilter = LayerFunction->RestrictToBlendRelatives[Index];
+			bShouldFilter = LayerFunction->EditorOnly.RestrictToBlendRelatives[Index];
 			BaseClassName = UMaterialFunctionMaterialLayerBlend::StaticClass()->GetFName();
 			InstanceClassName = UMaterialFunctionMaterialLayerBlendInstance::StaticClass()->GetFName();
 		}
@@ -428,8 +428,8 @@ FReply FMaterialPropertyHelpers::OnClickedSaveNewFunctionInstance(class UMateria
 					ChildInstance->FontParameterValues = EditedInstance->FontParameterValues;
 
 					const FStaticParameterSet& StaticParameters = EditedInstance->GetStaticParameters();
-					ChildInstance->StaticSwitchParameterValues = StaticParameters.StaticSwitchParameters;
-					ChildInstance->StaticComponentMaskParameterValues = StaticParameters.StaticComponentMaskParameters;
+					ChildInstance->StaticSwitchParameterValues = StaticParameters.EditorOnly.StaticSwitchParameters;
+					ChildInstance->StaticComponentMaskParameterValues = StaticParameters.EditorOnly.StaticComponentMaskParameters;
 				}
 			}
 		}
@@ -519,8 +519,8 @@ FReply FMaterialPropertyHelpers::OnClickedSaveNewLayerInstance(class UMaterialFu
 					ChildInstance->FontParameterValues = EditedInstance->FontParameterValues;
 
 					const FStaticParameterSet& StaticParameters = EditedInstance->GetStaticParameters();
-					ChildInstance->StaticSwitchParameterValues = StaticParameters.StaticSwitchParameters;
-					ChildInstance->StaticComponentMaskParameterValues = StaticParameters.StaticComponentMaskParameters;
+					ChildInstance->StaticSwitchParameterValues = StaticParameters.EditorOnly.StaticSwitchParameters;
+					ChildInstance->StaticComponentMaskParameterValues = StaticParameters.EditorOnly.StaticComponentMaskParameters;
 				}
 			}
 		}
@@ -553,72 +553,22 @@ void FMaterialPropertyHelpers::OnOverrideParameter(bool NewValue, class UDEditor
 
 FText FMaterialPropertyHelpers::GetParameterExpressionDescription(UDEditorParameterValue* Parameter, UObject* MaterialEditorInstance)
 {
-	if (Parameter->ExpressionId.IsValid())
-	{
-		UMaterial* BaseMaterial = nullptr;
-
-		UMaterialEditorInstanceConstant* MaterialInstanceEditor = Cast<UMaterialEditorInstanceConstant>(MaterialEditorInstance);
-		if (MaterialInstanceEditor)
-		{
-			BaseMaterial = MaterialInstanceEditor->SourceInstance->GetMaterial();
-		}
-		UMaterialEditorPreviewParameters* MaterialEditor = Cast<UMaterialEditorPreviewParameters>(MaterialEditorInstance);
-		if (MaterialEditor)
-		{
-			BaseMaterial = MaterialEditor->OriginalMaterial;
-		}
-
-		// TODO: This needs to support functions added by SourceInstance layers
-		if (BaseMaterial)
-		{
-			UMaterialExpression* MaterialExpression = BaseMaterial->FindExpressionByGUID<UMaterialExpression>(Parameter->ExpressionId);
-
-			if (MaterialExpression)
-			{
-				return FText::FromString(MaterialExpression->Desc);
-			}
-		}
-	}
-
-	return FText::GetEmpty();
+	return FText::FromString(Parameter->Description);
 }
 
 FText FMaterialPropertyHelpers::GetParameterTooltip(UDEditorParameterValue* Parameter, UObject* MaterialEditorInstance)
 {
-	UMaterial* BaseMaterial = nullptr;
-	UMaterialEditorInstanceConstant* MaterialInstanceEditor = Cast<UMaterialEditorInstanceConstant>(MaterialEditorInstance);
-	if (MaterialInstanceEditor)
+	const FText AssetPath = FText::FromString(Parameter->AssetPath);
+	FText TooltipText;
+	if (!Parameter->Description.IsEmpty())
 	{
-		BaseMaterial = MaterialInstanceEditor->SourceInstance->GetMaterial();
+		TooltipText = FText::Format(LOCTEXT("ParameterInfoDescAndLocation", "{0} \nFound in: {1}"), FText::FromString(Parameter->Description), AssetPath);
 	}
-	UMaterialEditorPreviewParameters* MaterialEditor = Cast<UMaterialEditorPreviewParameters>(MaterialEditorInstance);
-	if (MaterialEditor)
+	else
 	{
-		BaseMaterial = MaterialEditor->OriginalMaterial;
+		TooltipText = FText::Format(LOCTEXT("ParameterInfoLocationOnly", "Found in: {0}"), AssetPath);
 	}
-
-	// TODO: This needs to support functions added by SourceInstance layers
-	if (BaseMaterial)
-	{
-		UMaterialExpression* MaterialExpression = BaseMaterial->FindExpressionByGUID<UMaterialExpression>(Parameter->ExpressionId);
-
-		if (MaterialExpression)
-		{
-			const FText AssetPath = FText::FromString(MaterialExpression->GetAssetPathName());
-			FText TooltipText;
-			if (!MaterialExpression->Desc.IsEmpty())
-			{
-				TooltipText = FText::Format(LOCTEXT("ParameterInfoDescAndLocation", "{0} \nFound in: {1}"), FText::FromString(MaterialExpression->Desc), AssetPath);
-			}
-			else
-			{
-				TooltipText = FText::Format(LOCTEXT("ParameterInfoLocationOnly", "Found in: {0}"), AssetPath);
-			}
-			return TooltipText;
-		}
-	}
-
-	return FText::GetEmpty();
+	return TooltipText;
 }
 
 
diff --git a/Engine/Source/Editor/MaterialEditor/Private/SMaterialLayersFunctionsTree.cpp b/Engine/Source/Editor/MaterialEditor/Private/SMaterialLayersFunctionsTree.cpp
--- a/Engine/Source/Editor/MaterialEditor/Private/SMaterialLayersFunctionsTree.cpp
+++ b/Engine/Source/Editor/MaterialEditor/Private/SMaterialLayersFunctionsTree.cpp
@@ -102,11 +102,11 @@ bool SMaterialLayersFunctionsInstanceTreeItem::GetFilterState(SMaterialLayersFun
 {
 	if (InStackData->ParameterInfo.Association == EMaterialParameterAssociation::LayerParameter)
 	{
-		return InTree->FunctionInstance->RestrictToLayerRelatives[InStackData->ParameterInfo.Index];
+		return InTree->FunctionInstance->EditorOnly.RestrictToLayerRelatives[InStackData->ParameterInfo.Index];
 	}
 	if (InStackData->ParameterInfo.Association == EMaterialParameterAssociation::BlendParameter)
 	{
-		return InTree->FunctionInstance->RestrictToBlendRelatives[InStackData->ParameterInfo.Index];
+		return InTree->FunctionInstance->EditorOnly.RestrictToBlendRelatives[InStackData->ParameterInfo.Index];
 	}
 	return false;
 }
@@ -115,11 +115,11 @@ void SMaterialLayersFunctionsInstanceTreeItem::FilterClicked(const ECheckBoxStat
 {
 	if (InStackData->ParameterInfo.Association == EMaterialParameterAssociation::LayerParameter)
 	{
-		InTree->FunctionInstance->RestrictToLayerRelatives[InStackData->ParameterInfo.Index] = !InTree->FunctionInstance->RestrictToLayerRelatives[InStackData->ParameterInfo.Index];
+		InTree->FunctionInstance->EditorOnly.RestrictToLayerRelatives[InStackData->ParameterInfo.Index] = !InTree->FunctionInstance->EditorOnly.RestrictToLayerRelatives[InStackData->ParameterInfo.Index];
 	}
 	if (InStackData->ParameterInfo.Association == EMaterialParameterAssociation::BlendParameter)
 	{
-		InTree->FunctionInstance->RestrictToBlendRelatives[InStackData->ParameterInfo.Index] = !InTree->FunctionInstance->RestrictToBlendRelatives[InStackData->ParameterInfo.Index];
+		InTree->FunctionInstance->EditorOnly.RestrictToBlendRelatives[InStackData->ParameterInfo.Index] = !InTree->FunctionInstance->EditorOnly.RestrictToBlendRelatives[InStackData->ParameterInfo.Index];
 	}
 }
 
@@ -137,7 +137,7 @@ void SMaterialLayersFunctionsInstanceTreeItem::OnNameChanged(const FText& InText
 {
 	const FScopedTransaction Transaction(LOCTEXT("RenamedSection", "Renamed layer and blend section"));
 	InTree->FunctionInstanceHandle->NotifyPreChange();
-	InTree->FunctionInstance->LayerNames[Counter] = InText;
+	InTree->FunctionInstance->EditorOnly.LayerNames[Counter] = InText;
 	InTree->FunctionInstance->UnlinkLayerFromParent(Counter);
 	InTree->MaterialEditorInstance->CopyToSourceInstance(true);
 	InTree->FunctionInstanceHandle->NotifyPostChange(EPropertyChangeType::ValueSet);
@@ -210,7 +210,7 @@ FReply SMaterialLayersFunctionsInstanceTreeItem::OnLayerDrop(const FDragDropEven
 
 bool SMaterialLayersFunctionsInstanceTree::IsOverriddenExpression(class UDEditorParameterValue* Parameter, int32 InIndex)
 {
-	return FMaterialPropertyHelpers::IsOverriddenExpression(Parameter) && FunctionInstance->LayerStates[InIndex];
+	return FMaterialPropertyHelpers::IsOverriddenExpression(Parameter) && FunctionInstance->EditorOnly.LayerStates[InIndex];
 }
 
 bool SMaterialLayersFunctionsInstanceTree::IsOverriddenExpression(TObjectPtr<UDEditorParameterValue> Parameter, int32 InIndex)
@@ -400,8 +400,8 @@ void SMaterialLayersFunctionsInstanceTreeItem::Construct(const FArguments& InArg
 			.TextStyle(FAppStyle::Get(), "DetailsView.CategoryTextStyle")
 			.TransformPolicy(ETextTransformPolicy::ToUpper);
 		const int32 LayerStateIndex = StackParameterData->ParameterInfo.Association == EMaterialParameterAssociation::BlendParameter ? StackParameterData->ParameterInfo.Index + 1 : StackParameterData->ParameterInfo.Index;
-		LeftSideWidget->SetEnabled(InArgs._InTree->FunctionInstance->LayerStates[LayerStateIndex]);
-		RightSideWidget->SetEnabled(InArgs._InTree->FunctionInstance->LayerStates[LayerStateIndex]);
+		LeftSideWidget->SetEnabled(InArgs._InTree->FunctionInstance->EditorOnly.LayerStates[LayerStateIndex]);
+		RightSideWidget->SetEnabled(InArgs._InTree->FunctionInstance->EditorOnly.LayerStates[LayerStateIndex]);
 	}
 // END GROUP
 
@@ -542,7 +542,7 @@ void SMaterialLayersFunctionsInstanceTreeItem::Construct(const FArguments& InArg
 				.ToolTipText(LOCTEXT("SaveToChildInstance", "Save To Child Instance"))
 			];
 			
-		LeftSideWidget->SetEnabled(InArgs._InTree->FunctionInstance->LayerStates[LayerStateIndex]);
+		LeftSideWidget->SetEnabled(InArgs._InTree->FunctionInstance->EditorOnly.LayerStates[LayerStateIndex]);
 	}
 // END ASSET
 
@@ -927,7 +927,7 @@ void SMaterialLayersFunctionsInstanceTreeItem::Construct(const FArguments& InArg
 		StackParameterData->ParameterNode->CreatePropertyHandle()->SetOnPropertyValueChanged(FSimpleDelegate::CreateSP(Tree, &SMaterialLayersFunctionsInstanceTree::UpdateThumbnailMaterial, StackParameterData->ParameterInfo.Association, StackParameterData->ParameterInfo.Index, false));
 		StackParameterData->ParameterNode->CreatePropertyHandle()->SetOnChildPropertyValueChanged(FSimpleDelegate::CreateSP(Tree, &SMaterialLayersFunctionsInstanceTree::UpdateThumbnailMaterial, StackParameterData->ParameterInfo.Association, StackParameterData->ParameterInfo.Index, false));
 
-		LeftSideWidget->SetEnabled(InArgs._InTree->FunctionInstance->LayerStates[LayerStateIndex]);
+		LeftSideWidget->SetEnabled(InArgs._InTree->FunctionInstance->EditorOnly.LayerStates[LayerStateIndex]);
 	}
 // END PROPERTY
 
@@ -939,10 +939,10 @@ void SMaterialLayersFunctionsInstanceTreeItem::Construct(const FArguments& InArg
 		RightSideWidget = NodeWidgets.ValueWidget.ToSharedRef();
 
 		const int32 LayerStateIndex = StackParameterData->ParameterInfo.Association == EMaterialParameterAssociation::BlendParameter ? StackParameterData->ParameterInfo.Index + 1 : StackParameterData->ParameterInfo.Index;
-		LeftSideWidget->SetEnabled(InArgs._InTree->FunctionInstance->LayerStates[LayerStateIndex]);
+		LeftSideWidget->SetEnabled(InArgs._InTree->FunctionInstance->EditorOnly.LayerStates[LayerStateIndex]);
 		TAttribute<bool> EnabledAttribute = TAttribute<bool>::Create([this, LayerStateIndex]() -> bool
 			{
-				return FMaterialPropertyHelpers::IsOverriddenExpression(StackParameterData->Parameter) && Tree->FunctionInstance->LayerStates[LayerStateIndex];
+				return FMaterialPropertyHelpers::IsOverriddenExpression(StackParameterData->Parameter) && Tree->FunctionInstance->EditorOnly.LayerStates[LayerStateIndex];
 			});
 		RightSideWidget->SetEnabled(EnabledAttribute);
 	}
@@ -1149,20 +1149,20 @@ void SMaterialLayersFunctionsInstanceTree::Construct(const FArguments& InArgs)
 	//Fixup for adding new bool arrays to the class
 	if (FunctionInstance)
 	{
-		if (FunctionInstance->Layers.Num() != FunctionInstance->RestrictToLayerRelatives.Num())
+		if (FunctionInstance->Layers.Num() != FunctionInstance->EditorOnly.RestrictToLayerRelatives.Num())
 		{
-			int32 OriginalSize = FunctionInstance->RestrictToLayerRelatives.Num();
+			int32 OriginalSize = FunctionInstance->EditorOnly.RestrictToLayerRelatives.Num();
 			for (int32 LayerIt = 0; LayerIt < FunctionInstance->Layers.Num() - OriginalSize; LayerIt++)
 			{
-				FunctionInstance->RestrictToLayerRelatives.Add(false);
+				FunctionInstance->EditorOnly.RestrictToLayerRelatives.Add(false);
 			}
 		}
-		if (FunctionInstance->Blends.Num() != FunctionInstance->RestrictToBlendRelatives.Num())
+		if (FunctionInstance->Blends.Num() != FunctionInstance->EditorOnly.RestrictToBlendRelatives.Num())
 		{
-			int32 OriginalSize = FunctionInstance->RestrictToBlendRelatives.Num();
+			int32 OriginalSize = FunctionInstance->EditorOnly.RestrictToBlendRelatives.Num();
 			for (int32 BlendIt = 0; BlendIt < FunctionInstance->Blends.Num() - OriginalSize; BlendIt++)
 			{
-				FunctionInstance->RestrictToBlendRelatives.Add(false);
+				FunctionInstance->EditorOnly.RestrictToBlendRelatives.Add(false);
 			}
 		}
 	}
@@ -1332,7 +1332,7 @@ FReply SMaterialLayersFunctionsInstanceTree::ToggleLayerVisibility(int32 Index)
 			FunctionInstance->SetBlendedLayerVisibility(Index, true);
 			bLayerIsolated = false;
 		}
-		for (int32 LayerIt = 1; LayerIt < FunctionInstance->LayerStates.Num(); LayerIt++)
+		for (int32 LayerIt = 1; LayerIt < FunctionInstance->EditorOnly.LayerStates.Num(); LayerIt++)
 		{
 			if (LayerIt != Index)
 			{
@@ -2634,20 +2634,20 @@ void SMaterialLayersFunctionsMaterialTree::Construct(const FArguments& InArgs)
 	//Fixup for adding new bool arrays to the class
 	if (FunctionInstance)
 	{
-		if (FunctionInstance->Layers.Num() != FunctionInstance->RestrictToLayerRelatives.Num())
+		if (FunctionInstance->Layers.Num() != FunctionInstance->EditorOnly.RestrictToLayerRelatives.Num())
 		{
-			int32 OriginalSize = FunctionInstance->RestrictToLayerRelatives.Num();
+			int32 OriginalSize = FunctionInstance->EditorOnly.RestrictToLayerRelatives.Num();
 			for (int32 LayerIt = 0; LayerIt < FunctionInstance->Layers.Num() - OriginalSize; LayerIt++)
 			{
-				FunctionInstance->RestrictToLayerRelatives.Add(false);
+				FunctionInstance->EditorOnly.RestrictToLayerRelatives.Add(false);
 			}
 		}
-		if (FunctionInstance->Blends.Num() != FunctionInstance->RestrictToBlendRelatives.Num())
+		if (FunctionInstance->Blends.Num() != FunctionInstance->EditorOnly.RestrictToBlendRelatives.Num())
 		{
-			int32 OriginalSize = FunctionInstance->RestrictToBlendRelatives.Num();
+			int32 OriginalSize = FunctionInstance->EditorOnly.RestrictToBlendRelatives.Num();
 			for (int32 BlendIt = 0; BlendIt < FunctionInstance->Blends.Num() - OriginalSize; BlendIt++)
 			{
-				FunctionInstance->RestrictToBlendRelatives.Add(false);
+				FunctionInstance->EditorOnly.RestrictToBlendRelatives.Add(false);
 			}
 		}
 	}
diff --git a/Engine/Source/Editor/MaterialEditor/Private/Tests/MaterialEditorTests.cpp b/Engine/Source/Editor/MaterialEditor/Private/Tests/MaterialEditorTests.cpp
--- a/Engine/Source/Editor/MaterialEditor/Private/Tests/MaterialEditorTests.cpp
+++ b/Engine/Source/Editor/MaterialEditor/Private/Tests/MaterialEditorTests.cpp
@@ -64,7 +64,7 @@ namespace MaterialEditorPromotionTestUtils
 			UE_LOG(LogEditorMaterialEditorPromotionTests, Display, TEXT("Assigned the normal map texture to the new node"));
 			UMaterial* EditorMaterial = Cast<UMaterial>(CurrentMaterialEditor->GetMaterialInterface());
 			UMaterialGraph* MaterialGraph = EditorMaterial->MaterialGraph;
-			EditorMaterial->Normal.Connect(0, NewTextureSampleExpression);
+			EditorMaterial->GetEditorOnlyData()->Normal.Connect(0, NewTextureSampleExpression);
 			UE_LOG(LogEditorMaterialEditorPromotionTests, Display, TEXT("Connected the new node to the normal pin"));
 			MaterialGraph->LinkGraphNodesFromMaterial();
 
diff --git a/Engine/Source/Editor/MeshPaint/Private/IMeshPaintGeometryAdapter.cpp b/Engine/Source/Editor/MeshPaint/Private/IMeshPaintGeometryAdapter.cpp
--- a/Engine/Source/Editor/MeshPaint/Private/IMeshPaintGeometryAdapter.cpp
+++ b/Engine/Source/Editor/MeshPaint/Private/IMeshPaintGeometryAdapter.cpp
@@ -47,12 +47,10 @@ void IMeshPaintGeometryAdapter::DefaultQueryPaintableTextures(int32 MaterialInde
 		FPaintableTexture PaintableTexture;
 		// Find all the unique textures used in the top material level of the selected actor materials
 
-		const TArray<UMaterialExpression*>& Expressions = Material->GetMaterial()->Expressions;
-
 		// Only grab the textures from the top level of samples
-		for (auto ItExpressions = Expressions.CreateConstIterator(); ItExpressions; ItExpressions++)
+		for (UMaterialExpression* Expression : Material->GetMaterial()->GetExpressions())
 		{
-			UMaterialExpressionTextureBase* TextureBase = Cast<UMaterialExpressionTextureBase>(*ItExpressions);
+			UMaterialExpressionTextureBase* TextureBase = Cast<UMaterialExpressionTextureBase>(Expression);
 			if (TextureBase != NULL &&
 				TextureBase->Texture != NULL &&
 				!TextureBase->Texture->IsNormalMap() && 
@@ -65,7 +63,7 @@ void IMeshPaintGeometryAdapter::DefaultQueryPaintableTextures(int32 MaterialInde
 
 				// Texture Samples can have UV's specified, check the first node for whether it has a custom UV channel set. 
 				// We only check the first as the Mesh paint mode does not support painting with UV's modified in the shader.
-				UMaterialExpressionTextureSample* TextureSample = Cast<UMaterialExpressionTextureSample>(*ItExpressions);
+				UMaterialExpressionTextureSample* TextureSample = Cast<UMaterialExpressionTextureSample>(Expression);
 				if (TextureSample != NULL)
 				{
 					UMaterialExpressionTextureCoordinate* TextureCoords = Cast<UMaterialExpressionTextureCoordinate>(TextureSample->Coordinates.Expression);
diff --git a/Engine/Source/Editor/UnrealEd/Classes/MaterialEditor/DEditorParameterValue.h b/Engine/Source/Editor/UnrealEd/Classes/MaterialEditor/DEditorParameterValue.h
--- a/Engine/Source/Editor/UnrealEd/Classes/MaterialEditor/DEditorParameterValue.h
+++ b/Engine/Source/Editor/UnrealEd/Classes/MaterialEditor/DEditorParameterValue.h
@@ -30,6 +30,12 @@ class UNREALED_API UDEditorParameterValue : public UObject
 	UPROPERTY()
 	FGuid ExpressionId;
 
+	UPROPERTY()
+	FString Description;
+
+	UPROPERTY()
+	FString AssetPath;
+
 #if WITH_EDITORONLY_DATA
 	/** Controls where this parameter is displayed in a material instance parameter list.  The lower the number the higher up in the parameter list. */
 	UPROPERTY()
@@ -41,6 +47,8 @@ class UNREALED_API UDEditorParameterValue : public UObject
 
 	virtual bool GetValue(FMaterialParameterMetadata& OutResult) const
 	{
+		OutResult.Description = Description;
+		OutResult.AssetPath = AssetPath;
 		OutResult.ExpressionGuid = ExpressionId;
 		OutResult.SortPriority = SortPriority;
 		OutResult.bOverride = bOverride;
diff --git a/Engine/Source/Editor/UnrealEd/Private/Commandlets/DumpMaterialShaderTypes.cpp b/Engine/Source/Editor/UnrealEd/Private/Commandlets/DumpMaterialShaderTypes.cpp
--- a/Engine/Source/Editor/UnrealEd/Private/Commandlets/DumpMaterialShaderTypes.cpp
+++ b/Engine/Source/Editor/UnrealEd/Private/Commandlets/DumpMaterialShaderTypes.cpp
@@ -374,9 +374,9 @@ int ProcessMaterialInstances(const ITargetPlatform* TargetPlatform, const EShade
 			if (MaterialInstance->bHasStaticPermutationResource)
 			{
 				const FStaticParameterSet& ParameterSet = MaterialInstance->GetStaticParameters();
-				for (int32 StaticSwitchIndex = 0; StaticSwitchIndex < ParameterSet.StaticSwitchParameters.Num(); ++StaticSwitchIndex)
+				for (int32 StaticSwitchIndex = 0; StaticSwitchIndex < ParameterSet.EditorOnly.StaticSwitchParameters.Num(); ++StaticSwitchIndex)
 				{
-					const FStaticSwitchParameter& StaticSwitchParameter = ParameterSet.StaticSwitchParameters[StaticSwitchIndex];
+					const FStaticSwitchParameter& StaticSwitchParameter = ParameterSet.EditorOnly.StaticSwitchParameters[StaticSwitchIndex];
 					StaticParameterString += FString::Printf(
 						TEXT(", StaticSwitch'%s'=%s"),
 						*StaticSwitchParameter.ParameterInfo.ToString(),
diff --git a/Engine/Source/Editor/UnrealEd/Private/EditorParameters.cpp b/Engine/Source/Editor/UnrealEd/Private/EditorParameters.cpp
--- a/Engine/Source/Editor/UnrealEd/Private/EditorParameters.cpp
+++ b/Engine/Source/Editor/UnrealEd/Private/EditorParameters.cpp
@@ -176,6 +176,8 @@ UDEditorParameterValue* UDEditorParameterValue::Create(UObject* Owner,
 	check(Parameter);
 	Parameter->ParameterInfo = ParameterInfo;
 	Parameter->ExpressionId = Meta.ExpressionGuid;
+	Parameter->Description = Meta.Description;
+	Parameter->AssetPath = Meta.AssetPath;
 	Parameter->SortPriority = Meta.SortPriority;
 	Parameter->bOverride = Meta.bOverride;
 	return Parameter;
diff --git a/Engine/Source/Editor/UnrealEd/Private/Factories/EditorFactories.cpp b/Engine/Source/Editor/UnrealEd/Private/Factories/EditorFactories.cpp
--- a/Engine/Source/Editor/UnrealEd/Private/Factories/EditorFactories.cpp
+++ b/Engine/Source/Editor/UnrealEd/Private/Factories/EditorFactories.cpp
@@ -484,12 +484,14 @@ UObject* UMaterialFactoryNew::FactoryCreateNew(UClass* Class,UObject* InParent,F
 			TextureSampler->AutoSetSampleType();
 		}
 
-		NewMaterial->Expressions.Add(TextureSampler);
+		NewMaterial->GetExpressionCollection().AddExpression(TextureSampler);
+
+		UMaterialEditorOnlyData* NewMaterialEditorOnly = NewMaterial->GetEditorOnlyData();
 
 		FExpressionOutput& Output = TextureSampler->GetOutputs()[0];
 		FExpressionInput& Input = (TextureSampler->SamplerType == SAMPLERTYPE_Normal)
-			? (FExpressionInput&)NewMaterial->Normal
-			: (FExpressionInput&)NewMaterial->BaseColor;
+			? (FExpressionInput&)NewMaterialEditorOnly->Normal
+			: (FExpressionInput&)NewMaterialEditorOnly->BaseColor;
 
 		Input.Expression = TextureSampler;
 		Input.Mask = Output.Mask;
@@ -2844,7 +2846,7 @@ public:
 		{
 			check( sizeof(ColorDataType) == 8 );
 			uint16 RGBA[4] = { 0xFFFF,0xFFFF,0xFFFF, 0 };
-			memcpy(&WhiteWithZeroAlpha,RGBA,sizeof(ColorDataType));
+			memcpy(&WhiteWithZeroAlpha,RGBA,8);
 		}
 
 		// Left -> Right
@@ -4446,9 +4448,11 @@ UObject* UTextureFactory::FactoryCreateBinary
 		// Notify the asset registry
 		FAssetRegistryModule::AssetCreated(Material);
 
+		UMaterialEditorOnlyData* MaterialEditorOnly = Material->GetEditorOnlyData();
+
 		// Create a texture reference for the texture we just imported and hook it up to the diffuse channel
 		UMaterialExpression* Expression = NewObject<UMaterialExpression>(Material, UMaterialExpressionTextureSample::StaticClass());
-		Material->Expressions.Add( Expression );
+		Material->GetExpressionCollection().AddExpression( Expression );
 		TArray<FExpressionOutput> Outputs;
 
 		// If the user hasn't turned on any of the link checkboxes, default "bRGBToBaseColor" to being on.
@@ -4460,86 +4464,86 @@ UObject* UTextureFactory::FactoryCreateBinary
 		// Set up the links the user asked for
 		if( bRGBToBaseColor )
 		{
-			Material->BaseColor.Expression = Expression;
-			((UMaterialExpressionTextureSample*)Material->BaseColor.Expression)->Texture = Texture;
+			MaterialEditorOnly->BaseColor.Expression = Expression;
+			((UMaterialExpressionTextureSample*)MaterialEditorOnly->BaseColor.Expression)->Texture = Texture;
 
-			Outputs = Material->BaseColor.Expression->GetOutputs();
+			Outputs = MaterialEditorOnly->BaseColor.Expression->GetOutputs();
 			FExpressionOutput* Output = Outputs.GetData();
-			Material->BaseColor.Mask = Output->Mask;
-			Material->BaseColor.MaskR = Output->MaskR;
-			Material->BaseColor.MaskG = Output->MaskG;
-			Material->BaseColor.MaskB = Output->MaskB;
-			Material->BaseColor.MaskA = Output->MaskA;
+			MaterialEditorOnly->BaseColor.Mask = Output->Mask;
+			MaterialEditorOnly->BaseColor.MaskR = Output->MaskR;
+			MaterialEditorOnly->BaseColor.MaskG = Output->MaskG;
+			MaterialEditorOnly->BaseColor.MaskB = Output->MaskB;
+			MaterialEditorOnly->BaseColor.MaskA = Output->MaskA;
 		}
 
 		if( bRGBToEmissive )
 		{
-			Material->EmissiveColor.Expression = Expression;
-			((UMaterialExpressionTextureSample*)Material->EmissiveColor.Expression)->Texture = Texture;
+			MaterialEditorOnly->EmissiveColor.Expression = Expression;
+			((UMaterialExpressionTextureSample*)MaterialEditorOnly->EmissiveColor.Expression)->Texture = Texture;
 
-			Outputs = Material->EmissiveColor.Expression->GetOutputs();
+			Outputs = MaterialEditorOnly->EmissiveColor.Expression->GetOutputs();
 			FExpressionOutput* Output = Outputs.GetData();
-			Material->EmissiveColor.Mask = Output->Mask;
-			Material->EmissiveColor.MaskR = Output->MaskR;
-			Material->EmissiveColor.MaskG = Output->MaskG;
-			Material->EmissiveColor.MaskB = Output->MaskB;
-			Material->EmissiveColor.MaskA = Output->MaskA;
+			MaterialEditorOnly->EmissiveColor.Mask = Output->Mask;
+			MaterialEditorOnly->EmissiveColor.MaskR = Output->MaskR;
+			MaterialEditorOnly->EmissiveColor.MaskG = Output->MaskG;
+			MaterialEditorOnly->EmissiveColor.MaskB = Output->MaskB;
+			MaterialEditorOnly->EmissiveColor.MaskA = Output->MaskA;
 		}
 
 		if( bAlphaToRoughness )
 		{
-			Material->Roughness.Expression = Expression;
-			((UMaterialExpressionTextureSample*)Material->Roughness.Expression)->Texture = Texture;
+			MaterialEditorOnly->Roughness.Expression = Expression;
+			((UMaterialExpressionTextureSample*)MaterialEditorOnly->Roughness.Expression)->Texture = Texture;
 
-			Outputs = Material->Roughness.Expression->GetOutputs();
+			Outputs = MaterialEditorOnly->Roughness.Expression->GetOutputs();
 			FExpressionOutput* Output = Outputs.GetData();
-			Material->Roughness.Mask = Output->Mask;
-			Material->Roughness.MaskR = 0;
-			Material->Roughness.MaskG = 0;
-			Material->Roughness.MaskB = 0;
-			Material->Roughness.MaskA = 1;
+			MaterialEditorOnly->Roughness.Mask = Output->Mask;
+			MaterialEditorOnly->Roughness.MaskR = 0;
+			MaterialEditorOnly->Roughness.MaskG = 0;
+			MaterialEditorOnly->Roughness.MaskB = 0;
+			MaterialEditorOnly->Roughness.MaskA = 1;
 		}
 
 		if( bAlphaToEmissive )
 		{
-			Material->EmissiveColor.Expression = Expression;
-			((UMaterialExpressionTextureSample*)Material->EmissiveColor.Expression)->Texture = Texture;
+			MaterialEditorOnly->EmissiveColor.Expression = Expression;
+			((UMaterialExpressionTextureSample*)MaterialEditorOnly->EmissiveColor.Expression)->Texture = Texture;
 
-			Outputs = Material->EmissiveColor.Expression->GetOutputs();
+			Outputs = MaterialEditorOnly->EmissiveColor.Expression->GetOutputs();
 			FExpressionOutput* Output = Outputs.GetData();
-			Material->EmissiveColor.Mask = Output->Mask;
-			Material->EmissiveColor.MaskR = 0;
-			Material->EmissiveColor.MaskG = 0;
-			Material->EmissiveColor.MaskB = 0;
-			Material->EmissiveColor.MaskA = 1;
+			MaterialEditorOnly->EmissiveColor.Mask = Output->Mask;
+			MaterialEditorOnly->EmissiveColor.MaskR = 0;
+			MaterialEditorOnly->EmissiveColor.MaskG = 0;
+			MaterialEditorOnly->EmissiveColor.MaskB = 0;
+			MaterialEditorOnly->EmissiveColor.MaskA = 1;
 		}
 
 		if( bAlphaToOpacity )
 		{
-			Material->Opacity.Expression = Expression;
-			((UMaterialExpressionTextureSample*)Material->Opacity.Expression)->Texture = Texture;
+			MaterialEditorOnly->Opacity.Expression = Expression;
+			((UMaterialExpressionTextureSample*)MaterialEditorOnly->Opacity.Expression)->Texture = Texture;
 
-			Outputs = Material->Opacity.Expression->GetOutputs();
+			Outputs = MaterialEditorOnly->Opacity.Expression->GetOutputs();
 			FExpressionOutput* Output = Outputs.GetData();
-			Material->Opacity.Mask = Output->Mask;
-			Material->Opacity.MaskR = 0;
-			Material->Opacity.MaskG = 0;
-			Material->Opacity.MaskB = 0;
-			Material->Opacity.MaskA = 1;
+			MaterialEditorOnly->Opacity.Mask = Output->Mask;
+			MaterialEditorOnly->Opacity.MaskR = 0;
+			MaterialEditorOnly->Opacity.MaskG = 0;
+			MaterialEditorOnly->Opacity.MaskB = 0;
+			MaterialEditorOnly->Opacity.MaskA = 1;
 		}
 
 		if( bAlphaToOpacityMask )
 		{
-			Material->OpacityMask.Expression = Expression;
-			((UMaterialExpressionTextureSample*)Material->OpacityMask.Expression)->Texture = Texture;
+			MaterialEditorOnly->OpacityMask.Expression = Expression;
+			((UMaterialExpressionTextureSample*)MaterialEditorOnly->OpacityMask.Expression)->Texture = Texture;
 
-			Outputs = Material->OpacityMask.Expression->GetOutputs();
+			Outputs = MaterialEditorOnly->OpacityMask.Expression->GetOutputs();
 			FExpressionOutput* Output = Outputs.GetData();
-			Material->OpacityMask.Mask = Output->Mask;
-			Material->OpacityMask.MaskR = 0;
-			Material->OpacityMask.MaskG = 0;
-			Material->OpacityMask.MaskB = 0;
-			Material->OpacityMask.MaskA = 1;
+			MaterialEditorOnly->OpacityMask.Mask = Output->Mask;
+			MaterialEditorOnly->OpacityMask.MaskR = 0;
+			MaterialEditorOnly->OpacityMask.MaskG = 0;
+			MaterialEditorOnly->OpacityMask.MaskB = 0;
+			MaterialEditorOnly->OpacityMask.MaskA = 1;
 		}
 
 		Material->TwoSided	= bTwoSided;
diff --git a/Engine/Source/Editor/UnrealEd/Private/Fbx/FbxMainExport.cpp b/Engine/Source/Editor/UnrealEd/Private/Fbx/FbxMainExport.cpp
--- a/Engine/Source/Editor/UnrealEd/Private/Fbx/FbxMainExport.cpp
+++ b/Engine/Source/Editor/UnrealEd/Private/Fbx/FbxMainExport.cpp
@@ -1393,6 +1393,8 @@ FbxSurfaceMaterial* FFbxExporter::ExportMaterial(UMaterialInterface* MaterialInt
 		//Nothing to export
 		return nullptr;
 	}
+
+	UMaterialEditorOnlyData* MaterialEditorOnly = Material->GetEditorOnlyData();
 	
 	// Set the shading model
 	if (Material->GetShadingModels().HasOnlyShadingModel(MSM_DefaultLit))
@@ -1410,15 +1412,15 @@ FbxSurfaceMaterial* FFbxExporter::ExportMaterial(UMaterialInterface* MaterialInt
 	//Get the base material connected expression parameter name for all the supported fbx material exported properties
 	//We only support material input where the connected expression is a parameter of type (constant, scalar, vector, texture, TODO virtual texture)
 
-	FName BaseColorParamName = (!Material->BaseColor.UseConstant && Material->BaseColor.Expression) ? Material->BaseColor.Expression->GetParameterName() : NAME_None;
+	FName BaseColorParamName = (!MaterialEditorOnly->BaseColor.UseConstant && MaterialEditorOnly->BaseColor.Expression) ? MaterialEditorOnly->BaseColor.Expression->GetParameterName() : NAME_None;
 	bool BaseColorParamSet = false;
-	FName EmissiveParamName = (!Material->EmissiveColor.UseConstant && Material->EmissiveColor.Expression) ? Material->EmissiveColor.Expression->GetParameterName() : NAME_None;
+	FName EmissiveParamName = (!MaterialEditorOnly->EmissiveColor.UseConstant && MaterialEditorOnly->EmissiveColor.Expression) ? MaterialEditorOnly->EmissiveColor.Expression->GetParameterName() : NAME_None;
 	bool EmissiveParamSet = false;
-	FName NormalParamName = Material->Normal.Expression ? Material->Normal.Expression->GetParameterName() : NAME_None;
+	FName NormalParamName = MaterialEditorOnly->Normal.Expression ? MaterialEditorOnly->Normal.Expression->GetParameterName() : NAME_None;
 	bool NormalParamSet = false;
-	FName OpacityParamName = (!Material->Opacity.UseConstant && Material->Opacity.Expression) ? Material->Opacity.Expression->GetParameterName() : NAME_None;
+	FName OpacityParamName = (!MaterialEditorOnly->Opacity.UseConstant && MaterialEditorOnly->Opacity.Expression) ? MaterialEditorOnly->Opacity.Expression->GetParameterName() : NAME_None;
 	bool OpacityParamSet = false;
-	FName OpacityMaskParamName = (!Material->OpacityMask.UseConstant && Material->OpacityMask.Expression) ? Material->OpacityMask.Expression->GetParameterName() : NAME_None;
+	FName OpacityMaskParamName = (!MaterialEditorOnly->OpacityMask.UseConstant && MaterialEditorOnly->OpacityMask.Expression) ? MaterialEditorOnly->OpacityMask.Expression->GetParameterName() : NAME_None;
 	bool OpacityMaskParamSet = false;
 
 	UMaterialInstance* MaterialInstance = Cast<UMaterialInstance>(MaterialInterface);
@@ -1576,24 +1578,24 @@ FbxSurfaceMaterial* FFbxExporter::ExportMaterial(UMaterialInterface* MaterialInt
 	//The OpacityMaskParam set the transparency factor, so set it only if it was not set
 	if(!OpacityMaskParamSet)
 	{
-		((FbxSurfaceLambert*)FbxMaterial)->TransparencyFactor.Set(Material->Opacity.Constant);
+		((FbxSurfaceLambert*)FbxMaterial)->TransparencyFactor.Set(MaterialEditorOnly->Opacity.Constant);
 	}
 
 	// Fill in the profile_COMMON effect with the material information.
 	//Fill the texture or constant
 	if(!BaseColorParamSet)
 	{
-		if (!FillFbxTextureProperty(FbxSurfaceMaterial::sDiffuse, Material->BaseColor, FbxMaterial))
+		if (!FillFbxTextureProperty(FbxSurfaceMaterial::sDiffuse, MaterialEditorOnly->BaseColor, FbxMaterial))
 		{
-			((FbxSurfaceLambert*)FbxMaterial)->Diffuse.Set(SetMaterialComponent(Material->BaseColor, true));
+			((FbxSurfaceLambert*)FbxMaterial)->Diffuse.Set(SetMaterialComponent(MaterialEditorOnly->BaseColor, true));
 		}
 	}
 
 	if (!EmissiveParamSet)
 	{
-		if (!FillFbxTextureProperty(FbxSurfaceMaterial::sEmissive, Material->EmissiveColor, FbxMaterial))
+		if (!FillFbxTextureProperty(FbxSurfaceMaterial::sEmissive, MaterialEditorOnly->EmissiveColor, FbxMaterial))
 		{
-			((FbxSurfaceLambert*)FbxMaterial)->Emissive.Set(SetMaterialComponent(Material->EmissiveColor, true));
+			((FbxSurfaceLambert*)FbxMaterial)->Emissive.Set(SetMaterialComponent(MaterialEditorOnly->EmissiveColor, true));
 		}
 	}
 
@@ -1604,18 +1606,18 @@ FbxSurfaceMaterial* FFbxExporter::ExportMaterial(UMaterialInterface* MaterialInt
 	{
 		if (!OpacityParamSet)
 		{
-			if (!FillFbxTextureProperty(FbxSurfaceMaterial::sTransparentColor, Material->Opacity, FbxMaterial))
+			if (!FillFbxTextureProperty(FbxSurfaceMaterial::sTransparentColor, MaterialEditorOnly->Opacity, FbxMaterial))
 			{
-				FbxDouble3 OpacityValue((FbxDouble)(Material->Opacity.Constant), (FbxDouble)(Material->Opacity.Constant), (FbxDouble)(Material->Opacity.Constant));
+				FbxDouble3 OpacityValue((FbxDouble)(MaterialEditorOnly->Opacity.Constant), (FbxDouble)(MaterialEditorOnly->Opacity.Constant), (FbxDouble)(MaterialEditorOnly->Opacity.Constant));
 				((FbxSurfaceLambert*)FbxMaterial)->TransparentColor.Set(OpacityValue);
 			}
 		}
 
 		if (!OpacityMaskParamSet)
 		{
-			if (!FillFbxTextureProperty(FbxSurfaceMaterial::sTransparencyFactor, Material->OpacityMask, FbxMaterial))
+			if (!FillFbxTextureProperty(FbxSurfaceMaterial::sTransparencyFactor, MaterialEditorOnly->OpacityMask, FbxMaterial))
 			{
-				((FbxSurfaceLambert*)FbxMaterial)->TransparencyFactor.Set(Material->OpacityMask.Constant);
+				((FbxSurfaceLambert*)FbxMaterial)->TransparencyFactor.Set(MaterialEditorOnly->OpacityMask.Constant);
 			}
 		}
 	}
@@ -1625,7 +1627,7 @@ FbxSurfaceMaterial* FFbxExporter::ExportMaterial(UMaterialInterface* MaterialInt
 		if (!NormalParamSet)
 		{
 			//Set the Normal map only if there is a texture sampler
-			FillFbxTextureProperty(FbxSurfaceMaterial::sNormalMap, Material->Normal, FbxMaterial);
+			FillFbxTextureProperty(FbxSurfaceMaterial::sNormalMap, MaterialEditorOnly->Normal, FbxMaterial);
 		}
 	}
 
