Strata - added material editor context to simply create/link a Slab node from the context of a pin.

diff --git a/Engine/Source/Editor/MaterialEditor/Private/MaterialEditor.cpp b/Engine/Source/Editor/MaterialEditor/Private/MaterialEditor.cpp
--- a/Engine/Source/Editor/MaterialEditor/Private/MaterialEditor.cpp
+++ b/Engine/Source/Editor/MaterialEditor/Private/MaterialEditor.cpp
@@ -142,6 +142,7 @@
 #include "Materials/MaterialExpressionMaterialLayerOutput.h"
 #include "Materials/MaterialExpressionNamedReroute.h"
 #include "Materials/MaterialExpressionReroute.h"
+#include "Materials/MaterialExpressionStrata.h"
 
 #include "MaterialStats.h"
 #include "MaterialEditorTabs.h"
@@ -2126,6 +2127,21 @@ void FMaterialEditor::AddGraphEditorPinActionsToContextMenu(FToolMenuSection& In
 			PromoteToParameterAction
 		);
 	}
+
+	{
+		FToolUIAction CreateSlabNodeAction;
+		CreateSlabNodeAction.ExecuteAction = FToolMenuExecuteAction::CreateSP(this, &FMaterialEditor::OnCreateSlabNode);
+		CreateSlabNodeAction.IsActionVisibleDelegate = FToolMenuIsActionButtonVisible::CreateSP(this, &FMaterialEditor::OnCanCreateSlabNode);
+
+		TSharedPtr<FUICommandInfo> CreateSlabNodeCommand = FMaterialEditorCommands::Get().CreateSlabNode;
+		InSection.AddMenuEntry(
+			CreateSlabNodeCommand->GetCommandName(),
+			CreateSlabNodeCommand->GetLabel(),
+			CreateSlabNodeCommand->GetDescription(),
+			CreateSlabNodeCommand->GetIcon(),
+			CreateSlabNodeAction
+		);
+	}
 }
 
 void FMaterialEditor::DrawMaterialInfoStrings(
@@ -4541,7 +4557,6 @@ UClass* FMaterialEditor::GetOnPromoteToParameterClass(const UEdGraphPin* TargetP
 			case MP_PixelDepthOffset:
 			case MP_ShadingModel:
 			case MP_OpacityMask:
-			case MP_FrontMaterial:
 				return UMaterialExpressionScalarParameter::StaticClass();
 
 			case MP_WorldPositionOffset:
@@ -4552,6 +4567,10 @@ UClass* FMaterialEditor::GetOnPromoteToParameterClass(const UEdGraphPin* TargetP
 			case MP_Normal:
 			case MP_Tangent:
 				return UMaterialExpressionVectorParameter::StaticClass();
+
+			case MP_FrontMaterial:
+				return nullptr;
+
 		}
 	}
 	else if (OtherPinNode)
@@ -4582,6 +4601,8 @@ UClass* FMaterialEditor::GetOnPromoteToParameterClass(const UEdGraphPin* TargetP
 					case MCT_TextureCube: 
 					case MCT_VolumeTexture: 
 					case MCT_Texture: return UMaterialExpressionTextureObjectParameter::StaticClass();
+
+					case MCT_Strata: return nullptr;
 				}
 
 				break;
@@ -4646,6 +4667,82 @@ bool FMaterialEditor::OnCanPromoteToParameter(const FToolMenuContext& InMenuCont
 	return false;
 }
 
+void FMaterialEditor::OnCreateSlabNode(const FToolMenuContext& InMenuContext) const
+{
+	UGraphNodeContextMenuContext* NodeContext = InMenuContext.FindContext<UGraphNodeContextMenuContext>();
+	const UEdGraphPin* TargetPin = NodeContext->Pin;
+	UMaterialGraphNode_Base* PinNode = Cast<UMaterialGraphNode_Base>(TargetPin->GetOwningNode());
+
+	FMaterialGraphSchemaAction_NewNode Action;
+	Action.MaterialExpressionClass = UMaterialExpressionStrataSlabBSDF::StaticClass();
+
+	check(PinNode);
+	UEdGraph* GraphObj = PinNode->GetGraph();
+	check(GraphObj);
+
+	const FScopedTransaction Transaction(LOCTEXT("CreateSlabNode", "Create Slab Node"));
+	GraphObj->Modify();
+
+	// Set position of new node to be close to node we clicked on
+	FVector2D NewNodePos;
+	NewNodePos.X = PinNode->NodePosX - 200;
+	NewNodePos.Y = PinNode->NodePosY;
+
+	UMaterialGraphNode* MaterialNode = Cast<UMaterialGraphNode>(Action.PerformAction(GraphObj, const_cast<UEdGraphPin*>(TargetPin), NewNodePos));
+
+	if (MaterialNode->MaterialExpression->HasAParameterName())
+	{
+		MaterialNode->MaterialExpression->SetParameterName(TargetPin->PinName);
+		MaterialNode->MaterialExpression->ValidateParameterName(false);
+	}
+
+	if (MaterialEditorInstance != nullptr)
+	{
+		MaterialCustomPrimitiveDataWidget->UpdateEditorInstance(MaterialEditorInstance);
+	}
+}
+
+bool FMaterialEditor::OnCanCreateSlabNode(const FToolMenuContext& InMenuContext) const
+{
+	UGraphNodeContextMenuContext* NodeContext = InMenuContext.FindContext<UGraphNodeContextMenuContext>();
+	const UEdGraphPin* TargetPin = NodeContext->Pin;
+	UMaterialGraphNode_Root* RootPinNode = Cast<UMaterialGraphNode_Root>(TargetPin->GetOwningNode());
+	UMaterialGraphNode* OtherPinNode = Cast<UMaterialGraphNode>(TargetPin->GetOwningNode());
+	if (RootPinNode != nullptr)
+	{
+		EMaterialProperty propertyId = (EMaterialProperty)FCString::Atoi(*TargetPin->PinType.PinSubCategory.ToString());
+		switch (propertyId)
+		{
+		case MP_FrontMaterial:
+			return true;
+		}
+	}
+	else if (OtherPinNode)
+	{
+		const TArray<FExpressionInput*> ExpressionInputs = OtherPinNode->MaterialExpression->GetInputs();
+		FName TargetPinName = OtherPinNode->GetShortenPinName(TargetPin->PinName);
+
+		for (int32 Index = 0; Index < ExpressionInputs.Num(); ++Index)
+		{
+			FExpressionInput* Input = ExpressionInputs[Index];
+			FName InputName = OtherPinNode->MaterialExpression->GetInputName(Index);
+			InputName = OtherPinNode->GetShortenPinName(InputName);
+
+			if (InputName == TargetPinName)
+			{
+				switch (OtherPinNode->MaterialExpression->GetInputType(Index))
+				{
+				case MCT_Strata: 
+					return true;
+				}
+				break;
+			}
+		}
+	}
+
+	return false;
+}
+
 FString FMaterialEditor::GetDocLinkForSelectedNode()
 {
 	FString DocumentationLink;
diff --git a/Engine/Source/Editor/MaterialEditor/Private/MaterialEditor.h b/Engine/Source/Editor/MaterialEditor/Private/MaterialEditor.h
--- a/Engine/Source/Editor/MaterialEditor/Private/MaterialEditor.h
+++ b/Engine/Source/Editor/MaterialEditor/Private/MaterialEditor.h
@@ -771,6 +771,11 @@ private:
 	/** Will  return the UClass to create from the Pin Type */
 	UClass* GetOnPromoteToParameterClass(const UEdGraphPin* TargetPin) const;
 
+	/** Will create a slab node as input to the pin */
+	void OnCreateSlabNode(const FToolMenuContext& InMenuContext) const;
+	/** Used to know if we can create a slab node as input to the pin */
+	bool OnCanCreateSlabNode(const FToolMenuContext& InMenuContext) const;
+
 	/** Open documentation for the selected node class */
 	void OnGoToDocumentation();
 	/** Can we open documentation for the selected node */
diff --git a/Engine/Source/Editor/MaterialEditor/Private/MaterialEditorActions.cpp b/Engine/Source/Editor/MaterialEditor/Private/MaterialEditorActions.cpp
--- a/Engine/Source/Editor/MaterialEditor/Private/MaterialEditorActions.cpp
+++ b/Engine/Source/Editor/MaterialEditor/Private/MaterialEditorActions.cpp
@@ -60,7 +60,8 @@ void FMaterialEditorCommands::RegisterCommands()
 	UI_COMMAND( ForceRefreshPreviews, "Force Refresh Previews", "Forces a refresh of all previews", EUserInterfaceActionType::Button, FInputChord(EKeys::SpaceBar) );
 	UI_COMMAND( CreateComponentMaskNode, "Create ComponentMask Node", "Creates a ComponentMask node at the current cursor position.", EUserInterfaceActionType::Button, FInputChord(EModifierKey::Shift, EKeys::C));
 	UI_COMMAND( FindInMaterial, "Search", "Finds expressions and comments in the current Material", EUserInterfaceActionType::Button, FInputChord(EModifierKey::Control, EKeys::F));
-	UI_COMMAND( PromoteToParameter, "Promote to Parameter", "Promote selected Pin to parameter of pin type", EUserInterfaceActionType::Button, FInputChord());	
+	UI_COMMAND( PromoteToParameter, "Promote to Parameter", "Promote selected Pin to parameter of pin type", EUserInterfaceActionType::Button, FInputChord());
+	UI_COMMAND( CreateSlabNode, "Create Slab Node", "Create a Strata Slab node as input to the pin", EUserInterfaceActionType::Button, FInputChord());
 
 	UI_COMMAND(QualityLevel_All, "All", "Sets node preview to show all quality levels.)", EUserInterfaceActionType::RadioButton, FInputChord());
 	UI_COMMAND(QualityLevel_Epic, "Epic", "Sets node preview to Epic quality.", EUserInterfaceActionType::RadioButton, FInputChord());
diff --git a/Engine/Source/Editor/MaterialEditor/Public/MaterialEditorActions.h b/Engine/Source/Editor/MaterialEditor/Public/MaterialEditorActions.h
--- a/Engine/Source/Editor/MaterialEditor/Public/MaterialEditorActions.h
+++ b/Engine/Source/Editor/MaterialEditor/Public/MaterialEditorActions.h
@@ -183,6 +183,9 @@ public:
 	/** Create component mask node */
 	TSharedPtr< FUICommandInfo > PromoteToParameter;
 
+	/** Create component mask node */
+	TSharedPtr< FUICommandInfo > CreateSlabNode;
+
 	TSharedPtr< FUICommandInfo > QualityLevel_All;
 	TSharedPtr< FUICommandInfo > QualityLevel_Epic;
 	TSharedPtr< FUICommandInfo > QualityLevel_High;
