Strata - added contextual pin menu to add most used Strata nodes. For input or ouput pins.

diff --git a/Engine/Source/Editor/MaterialEditor/Private/MaterialEditor.cpp b/Engine/Source/Editor/MaterialEditor/Private/MaterialEditor.cpp
--- a/Engine/Source/Editor/MaterialEditor/Private/MaterialEditor.cpp
+++ b/Engine/Source/Editor/MaterialEditor/Private/MaterialEditor.cpp
@@ -2129,18 +2129,24 @@ void FMaterialEditor::AddGraphEditorPinActionsToContextMenu(FToolMenuSection& In
 	}
 
 	{
-		FToolUIAction CreateSlabNodeAction;
-		CreateSlabNodeAction.ExecuteAction = FToolMenuExecuteAction::CreateSP(this, &FMaterialEditor::OnCreateSlabNode);
-		CreateSlabNodeAction.IsActionVisibleDelegate = FToolMenuIsActionButtonVisible::CreateSP(this, &FMaterialEditor::OnCanCreateSlabNode);
-
-		TSharedPtr<FUICommandInfo> CreateSlabNodeCommand = FMaterialEditorCommands::Get().CreateSlabNode;
-		InSection.AddMenuEntry(
-			CreateSlabNodeCommand->GetCommandName(),
-			CreateSlabNodeCommand->GetLabel(),
-			CreateSlabNodeCommand->GetDescription(),
-			CreateSlabNodeCommand->GetIcon(),
-			CreateSlabNodeAction
-		);
+		auto AddStrataContextualMenu = [&](TSharedPtr<FUICommandInfo> CreateNodeCommand, EStrataNodeForPin StrataNodeForPin)
+		{
+			FToolUIAction CreateNodeAction;
+			CreateNodeAction.ExecuteAction = FToolMenuExecuteAction::CreateSP(this, &FMaterialEditor::OnCreateStrataNodeForPin, StrataNodeForPin);
+			CreateNodeAction.IsActionVisibleDelegate = FToolMenuIsActionButtonVisible::CreateSP(this, &FMaterialEditor::OnCanCreateStrataNodeForPin, StrataNodeForPin);
+
+			InSection.AddMenuEntry(
+				CreateNodeCommand->GetCommandName(),
+				CreateNodeCommand->GetLabel(),
+				CreateNodeCommand->GetDescription(),
+				CreateNodeCommand->GetIcon(),
+				CreateNodeAction
+			);
+		};
+		AddStrataContextualMenu(FMaterialEditorCommands::Get().CreateSlabNode, EStrataNodeForPin::Slab);
+		AddStrataContextualMenu(FMaterialEditorCommands::Get().CreateHorizontalMixNode, EStrataNodeForPin::HorizontalMix);
+		AddStrataContextualMenu(FMaterialEditorCommands::Get().CreateVerticalLayerNode, EStrataNodeForPin::VerticalLayer);
+		AddStrataContextualMenu(FMaterialEditorCommands::Get().CreateWeightNode, EStrataNodeForPin::Weight);
 	}
 }
 
@@ -4667,33 +4673,60 @@ bool FMaterialEditor::OnCanPromoteToParameter(const FToolMenuContext& InMenuCont
 	return false;
 }
 
-void FMaterialEditor::OnCreateSlabNode(const FToolMenuContext& InMenuContext) const
+void FMaterialEditor::OnCreateStrataNodeForPin(const FToolMenuContext& InMenuContext, EStrataNodeForPin NodeForPin) const
 {
 	UGraphNodeContextMenuContext* NodeContext = InMenuContext.FindContext<UGraphNodeContextMenuContext>();
 	const UEdGraphPin* TargetPin = NodeContext->Pin;
 	UMaterialGraphNode_Base* PinNode = Cast<UMaterialGraphNode_Base>(TargetPin->GetOwningNode());
+	const bool bTargetPinIsInput = TargetPin && (TargetPin->Direction == EEdGraphPinDirection::EGPD_Input);
 
 	FMaterialGraphSchemaAction_NewNode Action;
 	Action.MaterialExpressionClass = UMaterialExpressionStrataSlabBSDF::StaticClass();
+	if (NodeForPin == EStrataNodeForPin::HorizontalMix)
+	{
+		Action.MaterialExpressionClass = UMaterialExpressionStrataHorizontalMixing::StaticClass();
+	}
+	else if (NodeForPin == EStrataNodeForPin::VerticalLayer)
+	{
+		Action.MaterialExpressionClass = UMaterialExpressionStrataVerticalLayering::StaticClass();
+	}
+	else if (NodeForPin == EStrataNodeForPin::Weight)
+	{
+		Action.MaterialExpressionClass = UMaterialExpressionStrataWeight::StaticClass();
+	}
 
 	check(PinNode);
 	UEdGraph* GraphObj = PinNode->GetGraph();
 	check(GraphObj);
 
-	const FScopedTransaction Transaction(LOCTEXT("CreateSlabNode", "Create Slab Node"));
+	const FScopedTransaction Transaction(LOCTEXT("CreateStrataNode", "Create Strata Node"));
 	GraphObj->Modify();
 
 	// Set position of new node to be close to node we clicked on
 	FVector2D NewNodePos;
-	NewNodePos.X = PinNode->NodePosX - 200;
+	NewNodePos.X = PinNode->NodePosX + (bTargetPinIsInput ? -200 : 300);
 	NewNodePos.Y = PinNode->NodePosY;
 
-	UMaterialGraphNode* MaterialNode = Cast<UMaterialGraphNode>(Action.PerformAction(GraphObj, const_cast<UEdGraphPin*>(TargetPin), NewNodePos));
 
-	if (MaterialNode->MaterialExpression->HasAParameterName())
+	if (bTargetPinIsInput)
+	{
+		UMaterialGraphNode* NewNode = Cast<UMaterialGraphNode>(Action.PerformAction(GraphObj, const_cast<UEdGraphPin*>(TargetPin), NewNodePos));
+	}
+	else
 	{
-		MaterialNode->MaterialExpression->SetParameterName(TargetPin->PinName);
-		MaterialNode->MaterialExpression->ValidateParameterName(false);
+		// Link manually
+		UMaterialGraphNode* NewNode = Cast<UMaterialGraphNode>(Action.PerformAction(GraphObj, nullptr, NewNodePos));
+		const TArray<FExpressionInput*> NewNodeExpressionInputs = NewNode->MaterialExpression->GetInputs();
+
+		// From that direction, the node is never going to be a root node (a root node has no output we can connect from).
+		UMaterialGraphNode* TargetPinNode = Cast<UMaterialGraphNode>(TargetPin->GetOwningNode());
+
+		check(NewNodeExpressionInputs.Num() > 0 && TargetPin->SourceIndex < TargetPinNode->MaterialExpression->GetOutputs().Num());
+
+		FName TargetPinName = TargetPinNode->MaterialExpression->GetOutputs()[TargetPin->SourceIndex].OutputName;
+		UMaterialEditingLibrary::ConnectMaterialExpressions(TargetPinNode->MaterialExpression, TargetPinName.ToString(), NewNode->MaterialExpression, FString());
+
+		Material->MaterialGraph->RebuildGraph();
 	}
 
 	if (MaterialEditorInstance != nullptr)
@@ -4702,40 +4735,56 @@ void FMaterialEditor::OnCreateSlabNode(const FToolMenuContext& InMenuContext) co
 	}
 }
 
-bool FMaterialEditor::OnCanCreateSlabNode(const FToolMenuContext& InMenuContext) const
+bool FMaterialEditor::OnCanCreateStrataNodeForPin(const FToolMenuContext& InMenuContext, EStrataNodeForPin NodeForPin) const
 {
 	UGraphNodeContextMenuContext* NodeContext = InMenuContext.FindContext<UGraphNodeContextMenuContext>();
 	const UEdGraphPin* TargetPin = NodeContext->Pin;
 	UMaterialGraphNode_Root* RootPinNode = Cast<UMaterialGraphNode_Root>(TargetPin->GetOwningNode());
 	UMaterialGraphNode* OtherPinNode = Cast<UMaterialGraphNode>(TargetPin->GetOwningNode());
-	if (RootPinNode != nullptr)
+
+	if (TargetPin && (TargetPin->Direction == EEdGraphPinDirection::EGPD_Input) && (TargetPin->LinkedTo.Num() == 0))
 	{
-		EMaterialProperty propertyId = (EMaterialProperty)FCString::Atoi(*TargetPin->PinType.PinSubCategory.ToString());
-		switch (propertyId)
+		if (RootPinNode != nullptr)
 		{
-		case MP_FrontMaterial:
-			return true;
+			EMaterialProperty propertyId = (EMaterialProperty)FCString::Atoi(*TargetPin->PinType.PinSubCategory.ToString());
+			switch (propertyId)
+			{
+			case MP_FrontMaterial:
+				return true;
+			}
 		}
-	}
-	else if (OtherPinNode)
-	{
-		const TArray<FExpressionInput*> ExpressionInputs = OtherPinNode->MaterialExpression->GetInputs();
-		FName TargetPinName = OtherPinNode->GetShortenPinName(TargetPin->PinName);
-
-		for (int32 Index = 0; Index < ExpressionInputs.Num(); ++Index)
+		else if (OtherPinNode)
 		{
-			FExpressionInput* Input = ExpressionInputs[Index];
-			FName InputName = OtherPinNode->MaterialExpression->GetInputName(Index);
-			InputName = OtherPinNode->GetShortenPinName(InputName);
+			const TArray<FExpressionInput*> ExpressionInputs = OtherPinNode->MaterialExpression->GetInputs();
+			FName TargetPinName = OtherPinNode->GetShortenPinName(TargetPin->PinName);
 
-			if (InputName == TargetPinName)
+			for (int32 Index = 0; Index < ExpressionInputs.Num(); ++Index)
 			{
-				switch (OtherPinNode->MaterialExpression->GetInputType(Index))
+				FExpressionInput* Input = ExpressionInputs[Index];
+				FName InputName = OtherPinNode->MaterialExpression->GetInputName(Index);
+				InputName = OtherPinNode->GetShortenPinName(InputName);
+
+				if (InputName == TargetPinName)
 				{
-				case MCT_Strata: 
-					return true;
+					switch (OtherPinNode->MaterialExpression->GetInputType(Index))
+					{
+					case MCT_Strata:
+						return true;
+					}
+					break;
 				}
-				break;
+			}
+		}
+	}
+	else if (TargetPin && (TargetPin->Direction == EEdGraphPinDirection::EGPD_Output) && (TargetPin->LinkedTo.Num() == 0) && NodeForPin != EStrataNodeForPin::Slab)
+	{
+		if (OtherPinNode)
+		{
+			const TArray<FExpressionOutput>& ExpressionOutputs = OtherPinNode->MaterialExpression->GetOutputs();
+			check(TargetPin->SourceIndex < ExpressionOutputs.Num());
+			if (OtherPinNode->MaterialExpression->GetOutputType(TargetPin->SourceIndex) == MCT_Strata)
+			{
+				return true;
 			}
 		}
 	}
diff --git a/Engine/Source/Editor/MaterialEditor/Private/MaterialEditor.h b/Engine/Source/Editor/MaterialEditor/Private/MaterialEditor.h
--- a/Engine/Source/Editor/MaterialEditor/Private/MaterialEditor.h
+++ b/Engine/Source/Editor/MaterialEditor/Private/MaterialEditor.h
@@ -771,10 +771,17 @@ private:
 	/** Will  return the UClass to create from the Pin Type */
 	UClass* GetOnPromoteToParameterClass(const UEdGraphPin* TargetPin) const;
 
-	/** Will create a slab node as input to the pin */
-	void OnCreateSlabNode(const FToolMenuContext& InMenuContext) const;
-	/** Used to know if we can create a slab node as input to the pin */
-	bool OnCanCreateSlabNode(const FToolMenuContext& InMenuContext) const;
+	enum class EStrataNodeForPin : uint8
+	{
+		Slab,
+		HorizontalMix,
+		VerticalLayer,
+		Weight
+	};
+	/** Will create a Strata node as input to the pin */
+	void OnCreateStrataNodeForPin(const FToolMenuContext& InMenuContext, EStrataNodeForPin NodeForPin) const;
+	/** Used to know if we can create a Strata node as input to the pin */
+	bool OnCanCreateStrataNodeForPin(const FToolMenuContext& InMenuContext, EStrataNodeForPin NodeForPin) const;
 
 	/** Open documentation for the selected node class */
 	void OnGoToDocumentation();
diff --git a/Engine/Source/Editor/MaterialEditor/Private/MaterialEditorActions.cpp b/Engine/Source/Editor/MaterialEditor/Private/MaterialEditorActions.cpp
--- a/Engine/Source/Editor/MaterialEditor/Private/MaterialEditorActions.cpp
+++ b/Engine/Source/Editor/MaterialEditor/Private/MaterialEditorActions.cpp
@@ -61,7 +61,10 @@ void FMaterialEditorCommands::RegisterCommands()
 	UI_COMMAND( CreateComponentMaskNode, "Create ComponentMask Node", "Creates a ComponentMask node at the current cursor position.", EUserInterfaceActionType::Button, FInputChord(EModifierKey::Shift, EKeys::C));
 	UI_COMMAND( FindInMaterial, "Search", "Finds expressions and comments in the current Material", EUserInterfaceActionType::Button, FInputChord(EModifierKey::Control, EKeys::F));
 	UI_COMMAND( PromoteToParameter, "Promote to Parameter", "Promote selected Pin to parameter of pin type", EUserInterfaceActionType::Button, FInputChord());
-	UI_COMMAND( CreateSlabNode, "Create Slab Node", "Create a Strata Slab node as input to the pin", EUserInterfaceActionType::Button, FInputChord());
+	UI_COMMAND( CreateSlabNode, "Create Slab Node", "Create a Slab Node linked to the pin", EUserInterfaceActionType::Button, FInputChord());
+	UI_COMMAND( CreateHorizontalMixNode, "Create Horizontal Mix Node", "Create a Horizontal Mix Node linked to the pin", EUserInterfaceActionType::Button, FInputChord());
+	UI_COMMAND( CreateVerticalLayerNode, "Create Vertical Layer Node", "Create a Vertical Layer Node linked to the pin", EUserInterfaceActionType::Button, FInputChord());
+	UI_COMMAND( CreateWeightNode, "Create Weight Node", "Create a Weight Node linked to the pin", EUserInterfaceActionType::Button, FInputChord());
 
 	UI_COMMAND(QualityLevel_All, "All", "Sets node preview to show all quality levels.)", EUserInterfaceActionType::RadioButton, FInputChord());
 	UI_COMMAND(QualityLevel_Epic, "Epic", "Sets node preview to Epic quality.", EUserInterfaceActionType::RadioButton, FInputChord());
diff --git a/Engine/Source/Editor/MaterialEditor/Public/MaterialEditorActions.h b/Engine/Source/Editor/MaterialEditor/Public/MaterialEditorActions.h
--- a/Engine/Source/Editor/MaterialEditor/Public/MaterialEditorActions.h
+++ b/Engine/Source/Editor/MaterialEditor/Public/MaterialEditorActions.h
@@ -185,6 +185,9 @@ public:
 
 	/** Create component mask node */
 	TSharedPtr< FUICommandInfo > CreateSlabNode;
+	TSharedPtr< FUICommandInfo > CreateHorizontalMixNode;
+	TSharedPtr< FUICommandInfo > CreateVerticalLayerNode;
+	TSharedPtr< FUICommandInfo > CreateWeightNode;
 
 	TSharedPtr< FUICommandInfo > QualityLevel_All;
 	TSharedPtr< FUICommandInfo > QualityLevel_Epic;
