Fix static analysis warnings
#rnx

diff --git a/Engine/Plugins/Experimental/MLDeformer/MLDeformerFramework/Source/MLDeformerFrameworkEditor/Private/MLDeformerEditorModel.cpp b/Engine/Plugins/Experimental/MLDeformer/MLDeformerFramework/Source/MLDeformerFrameworkEditor/Private/MLDeformerEditorModel.cpp
--- a/Engine/Plugins/Experimental/MLDeformer/MLDeformerFramework/Source/MLDeformerFrameworkEditor/Private/MLDeformerEditorModel.cpp
+++ b/Engine/Plugins/Experimental/MLDeformer/MLDeformerFramework/Source/MLDeformerFrameworkEditor/Private/MLDeformerEditorModel.cpp
@@ -1238,7 +1238,7 @@ namespace UE::MLDeformer
 		{
 			if (EditorActor->GetMLDeformerComponent())
 			{
-				USkeletalMeshComponent* SkelMeshComponent = EditorActor ? EditorActor->GetSkeletalMeshComponent() : nullptr;
+				USkeletalMeshComponent* SkelMeshComponent = EditorActor->GetSkeletalMeshComponent();
 				UMLDeformerAsset* DeformerAsset = Model->GetDeformerAsset();
 				EditorActor->GetMLDeformerComponent()->SetupComponent(DeformerAsset, SkelMeshComponent);
 			}
diff --git a/Engine/Plugins/FX/Niagara/Source/NiagaraEditor/Private/NiagaraNodeParameterMapGet.cpp b/Engine/Plugins/FX/Niagara/Source/NiagaraEditor/Private/NiagaraNodeParameterMapGet.cpp
--- a/Engine/Plugins/FX/Niagara/Source/NiagaraEditor/Private/NiagaraNodeParameterMapGet.cpp
+++ b/Engine/Plugins/FX/Niagara/Source/NiagaraEditor/Private/NiagaraNodeParameterMapGet.cpp
@@ -424,7 +424,7 @@ FText UNiagaraNodeParameterMapGet::GetNodeTitle(ENodeTitleType::Type TitleType)
 
 void UNiagaraNodeParameterMapGet::AddOrphanedPinPairGuids(UEdGraphPin* OutputPin, UEdGraphPin* DefaultPin)
 {
-	ensure(OutputPin && DefaultPin && OutputPin->bOrphanedPin && DefaultPin->bOrphanedPin && OutputPin->PersistentGuid.IsValid() && DefaultPin->PersistentGuid.IsValid());
+	ensure(OutputPin->bOrphanedPin && DefaultPin->bOrphanedPin && OutputPin->PersistentGuid.IsValid() && DefaultPin->PersistentGuid.IsValid());
 
 	PinOutputToPinDefaultPersistentId.Add(OutputPin->PersistentGuid, DefaultPin->PersistentGuid);
 }
diff --git a/Engine/Source/Editor/MaterialEditor/Private/MaterialEditor.cpp b/Engine/Source/Editor/MaterialEditor/Private/MaterialEditor.cpp
--- a/Engine/Source/Editor/MaterialEditor/Private/MaterialEditor.cpp
+++ b/Engine/Source/Editor/MaterialEditor/Private/MaterialEditor.cpp
@@ -4678,7 +4678,7 @@ void FMaterialEditor::OnCreateStrataNodeForPin(const FToolMenuContext& InMenuCon
 	UGraphNodeContextMenuContext* NodeContext = InMenuContext.FindContext<UGraphNodeContextMenuContext>();
 	const UEdGraphPin* TargetPin = NodeContext->Pin;
 	UMaterialGraphNode_Base* PinNode = Cast<UMaterialGraphNode_Base>(TargetPin->GetOwningNode());
-	const bool bTargetPinIsInput = TargetPin && (TargetPin->Direction == EEdGraphPinDirection::EGPD_Input);
+	const bool bTargetPinIsInput = TargetPin->Direction == EEdGraphPinDirection::EGPD_Input;
 
 	FMaterialGraphSchemaAction_NewNode Action;
 	Action.MaterialExpressionClass = UMaterialExpressionStrataSlabBSDF::StaticClass();
@@ -4742,7 +4742,7 @@ bool FMaterialEditor::OnCanCreateStrataNodeForPin(const FToolMenuContext& InMenu
 	UMaterialGraphNode_Root* RootPinNode = Cast<UMaterialGraphNode_Root>(TargetPin->GetOwningNode());
 	UMaterialGraphNode* OtherPinNode = Cast<UMaterialGraphNode>(TargetPin->GetOwningNode());
 
-	if (TargetPin && (TargetPin->Direction == EEdGraphPinDirection::EGPD_Input) && (TargetPin->LinkedTo.Num() == 0))
+	if ((TargetPin->Direction == EEdGraphPinDirection::EGPD_Input) && (TargetPin->LinkedTo.Num() == 0))
 	{
 		if (RootPinNode != nullptr)
 		{
