Add MaterialEditingLibrary functions to get the output name of a MaterialExpression given as an Input Node for another MaterialExpression.

diff --git a/Engine/Source/Editor/MaterialEditor/Private/MaterialEditingLibrary.cpp b/Engine/Source/Editor/MaterialEditor/Private/MaterialEditingLibrary.cpp
--- a/Engine/Source/Editor/MaterialEditor/Private/MaterialEditingLibrary.cpp
+++ b/Engine/Source/Editor/MaterialEditor/Private/MaterialEditingLibrary.cpp
@@ -800,6 +800,50 @@ UMaterialExpression* UMaterialEditingLibrary::GetMaterialPropertyInputNode(UMate
 	return nullptr;
 }
 
+static FString GetExpressionOutputName(const FExpressionOutput& Output)
+{
+	if (!Output.OutputName.IsNone())
+	{
+		return Output.OutputName.ToString();
+	}
+	else if (Output.Mask)
+	{
+		if (Output.MaskR && !Output.MaskG && !Output.MaskB && !Output.MaskA)
+		{
+			return TEXT("R");
+		}
+		else if (!Output.MaskR && Output.MaskG && !Output.MaskB && !Output.MaskA)
+		{
+			return TEXT("G");
+		}
+		else if (!Output.MaskR && !Output.MaskG && Output.MaskB && !Output.MaskA)
+		{
+			return TEXT("B");
+		}
+		else if (!Output.MaskR && !Output.MaskG && !Output.MaskB && Output.MaskA)
+		{
+			return TEXT("A");
+		}
+	}
+	return FString();
+}
+
+FString UMaterialEditingLibrary::GetMaterialPropertyInputNodeOutputName(UMaterial* Material, EMaterialProperty Property)
+{
+	if (Material)
+	{
+		FExpressionInput* ExpressionInput = Material->GetExpressionInputForProperty(Property);
+		if (ExpressionInput->OutputIndex != INDEX_NONE
+			&& ExpressionInput->Expression
+			&& ExpressionInput->OutputIndex < ExpressionInput->Expression->Outputs.Num())
+		{
+			FExpressionOutput& Output = ExpressionInput->Expression->Outputs[ExpressionInput->OutputIndex];
+			return GetExpressionOutputName(Output);
+		}
+	}
+	return FString();
+}
+
 TArray<UMaterialExpression*> UMaterialEditingLibrary::GetInputsForMaterialExpression(UMaterial* Material, UMaterialExpression* MaterialExpression)
 {
 	TArray<UMaterialExpression*> MaterialExpressions;
@@ -814,6 +858,24 @@ TArray<UMaterialExpression*> UMaterialEditingLibrary::GetInputsForMaterialExpres
 	return MaterialExpressions;
 }
 
+bool UMaterialEditingLibrary::GetInputNodeOutputNameForMaterialExpression(UMaterialExpression* MaterialExpression, UMaterialExpression* InputNode, FString& OutputName)
+{
+	OutputName = TEXT("");
+	for (const FExpressionInput* Input : MaterialExpression->GetInputs())
+	{
+		if (Input->Expression == InputNode)
+		{
+			if(Input->OutputIndex != INDEX_NONE && Input->OutputIndex < InputNode->Outputs.Num())
+			{
+				FExpressionOutput& Output = InputNode->Outputs[Input->OutputIndex];
+				OutputName = GetExpressionOutputName(Output);
+				return true;
+			}
+		}
+	}
+	return false;
+}
+
 TArray<UTexture*> UMaterialEditingLibrary::GetUsedTextures(UMaterial* Material)
 {
 	TArray<UTexture*> OutTextures;
diff --git a/Engine/Source/Editor/MaterialEditor/Public/MaterialEditingLibrary.h b/Engine/Source/Editor/MaterialEditor/Public/MaterialEditingLibrary.h
--- a/Engine/Source/Editor/MaterialEditor/Public/MaterialEditingLibrary.h
+++ b/Engine/Source/Editor/MaterialEditor/Public/MaterialEditingLibrary.h
@@ -195,10 +195,18 @@ public:
 	UFUNCTION(BlueprintPure, Category = "MaterialEditing")
 	static UMaterialExpression* GetMaterialPropertyInputNode(UMaterial* Material, EMaterialProperty Property);
 
+	/** Get the node output name providing the output for a given material property from an active material editor */
+	UFUNCTION(BlueprintPure, Category = "MaterialEditing")
+	static FString GetMaterialPropertyInputNodeOutputName(UMaterial* Material, EMaterialProperty Property);
+
 	/** Get the set of nodes acting as inputs to a node from an active material editor */
 	UFUNCTION(BlueprintPure, Category = "MaterialEditing")
 	static TArray<UMaterialExpression*> GetInputsForMaterialExpression(UMaterial* Material, UMaterialExpression* MaterialExpression);
 
+	/** Get the output name of input node connected to MaterialExpression from an active material editor */
+	UFUNCTION(BlueprintPure, Category = "MaterialEditing")
+	static bool GetInputNodeOutputNameForMaterialExpression(UMaterialExpression* MaterialExpression, UMaterialExpression* InputNode, FString& OutputName);
+
 	/** Get the list of textures used by a material */
 	UFUNCTION(BlueprintPure, Category = "MaterialEditing")
 	static TArray<UTexture*> GetUsedTextures(UMaterial* Material);
diff --git a/Engine/Source/Runtime/Engine/Classes/Materials/MaterialExpression.h b/Engine/Source/Runtime/Engine/Classes/Materials/MaterialExpression.h
--- a/Engine/Source/Runtime/Engine/Classes/Materials/MaterialExpression.h
+++ b/Engine/Source/Runtime/Engine/Classes/Materials/MaterialExpression.h
@@ -185,10 +185,10 @@ class ENGINE_API UMaterialExpression : public UObject
 	static void InitializeNumExecutionInputs(TArrayView<UMaterialExpression*> Expressions);
 #endif
 
-	UPROPERTY()
+	UPROPERTY(BlueprintReadWrite, Category = MaterialEditing)
 	int32 MaterialExpressionEditorX;
 
-	UPROPERTY()
+	UPROPERTY(BlueprintReadWrite, Category = MaterialEditing)
 	int32 MaterialExpressionEditorY;
 
 	/** Expression's Graph representation */
