Fixed banned word warnings.

#rnx

diff --git a/Engine/Source/Editor/MaterialEditor/Private/MaterialInstanceEditor.cpp b/Engine/Source/Editor/MaterialEditor/Private/MaterialInstanceEditor.cpp
--- a/Engine/Source/Editor/MaterialEditor/Private/MaterialInstanceEditor.cpp
+++ b/Engine/Source/Editor/MaterialEditor/Private/MaterialInstanceEditor.cpp
@@ -788,7 +788,7 @@ void FMaterialInstanceEditor::CreateInternalWidgets()
 	DetailsViewArgs.bShowModifiedPropertiesOption = false;
 	DetailsViewArgs.bShowCustomFilterOption = true;
 	MaterialInstanceDetails = PropertyEditorModule.CreateDetailView( DetailsViewArgs );
-	// the sizes of the parameter lists are only based on the master material and not changed out from under the details panel 
+	// the sizes of the parameter lists are only based on the parent material and not changed out from under the details panel 
 	// When a parameter is added open MI editors are refreshed
 	// the tree should also refresh if one of the layer or blend assets is swapped
 
diff --git a/Engine/Source/Editor/MaterialEditor/Private/SMaterialLayersFunctionsTree.cpp b/Engine/Source/Editor/MaterialEditor/Private/SMaterialLayersFunctionsTree.cpp
--- a/Engine/Source/Editor/MaterialEditor/Private/SMaterialLayersFunctionsTree.cpp
+++ b/Engine/Source/Editor/MaterialEditor/Private/SMaterialLayersFunctionsTree.cpp
@@ -1386,7 +1386,7 @@ void SMaterialLayersFunctionsInstanceTree::CreateGroupsWidget()
 	{
 		FPropertyRowGeneratorArgs Args;
 		Generator = Module.CreatePropertyRowGenerator(Args);
-		// the sizes of the parameter lists are only based on the master material and not changed out from under the details panel 
+		// the sizes of the parameter lists are only based on the parent material and not changed out from under the details panel 
 		// When a parameter is added open MI editors are refreshed
 		// the tree should also refresh if one of the layer or blend assets is swapped
 
