UE-131224 - Fix LWC preview issue in FMaterialEditor::OnNodeDoubleClicked

We currently have Blueprint real number support for all of the basic core types, including FVector4f. We can remove the LWC_TODO, and change the PreviewValue to a FVector4f.

diff --git a/Engine/Source/Editor/MaterialEditor/Private/MaterialEditor.cpp b/Engine/Source/Editor/MaterialEditor/Private/MaterialEditor.cpp
--- a/Engine/Source/Editor/MaterialEditor/Private/MaterialEditor.cpp
+++ b/Engine/Source/Editor/MaterialEditor/Private/MaterialEditor.cpp
@@ -7115,11 +7115,10 @@ void FMaterialEditor::OnNodeDoubleClicked(class UEdGraphNode* Node)
 		}
 		else if (InputExpression)
 		{
-			// LWC_TODO: FIX THIS. Can't keep a pointer as float to a double type.
-			//ChannelEditStruct.Red = &InputExpression->PreviewValue.X;
-			//ChannelEditStruct.Green = &InputExpression->PreviewValue.Y;
-			//ChannelEditStruct.Blue = &InputExpression->PreviewValue.Z;
-			//ChannelEditStruct.Alpha = &InputExpression->PreviewValue.W;
+			ChannelEditStruct.Red = &InputExpression->PreviewValue.X;
+			ChannelEditStruct.Green = &InputExpression->PreviewValue.Y;
+			ChannelEditStruct.Blue = &InputExpression->PreviewValue.Z;
+			ChannelEditStruct.Alpha = &InputExpression->PreviewValue.W;
 		}
 		else if (VectorExpression)
 		{
diff --git a/Engine/Source/Runtime/Engine/Classes/Materials/MaterialExpressionFunctionInput.h b/Engine/Source/Runtime/Engine/Classes/Materials/MaterialExpressionFunctionInput.h
--- a/Engine/Source/Runtime/Engine/Classes/Materials/MaterialExpressionFunctionInput.h
+++ b/Engine/Source/Runtime/Engine/Classes/Materials/MaterialExpressionFunctionInput.h
@@ -61,7 +61,7 @@ class UMaterialExpressionFunctionInput : public UMaterialExpression
 
 	/** Value used to preview this input when editing the material function. */
 	UPROPERTY(EditAnywhere, BlueprintReadWrite, Category=MaterialExpressionFunctionInput, meta=(OverridingInputProperty = "Preview"))
-	FVector4 PreviewValue;
+	FVector4f PreviewValue;
 
 	/** Whether to use the preview value or texture as the default value for this input. */
 	UPROPERTY(EditAnywhere, BlueprintReadWrite, Category=MaterialExpressionFunctionInput)
diff --git a/Engine/Source/Runtime/Engine/Private/Materials/MaterialExpressions.cpp b/Engine/Source/Runtime/Engine/Private/Materials/MaterialExpressions.cpp
--- a/Engine/Source/Runtime/Engine/Private/Materials/MaterialExpressions.cpp
+++ b/Engine/Source/Runtime/Engine/Private/Materials/MaterialExpressions.cpp
@@ -15571,7 +15571,7 @@ bool UMaterialExpressionMaterialFunctionCall::IsInputConnectionRequired(int32 In
 	return true;
 }
 
-static FString GetInputDefaultValueString(EFunctionInputType InputType, const FVector4& PreviewValue)
+static FString GetInputDefaultValueString(EFunctionInputType InputType, const FVector4f& PreviewValue)
 {
 	static_assert(FunctionInput_Scalar < FunctionInput_Vector4, "Enum values out of order.");
 	check(InputType <= FunctionInput_Vector4);
