Fixed assertion fired when the input or output preview of a collapsed node is expanded.

Issue was caused by the code assuming that the expression always has at least one output.

diff --git a/Engine/Source/Editor/MaterialEditor/Private/MaterialEditor.cpp b/Engine/Source/Editor/MaterialEditor/Private/MaterialEditor.cpp
--- a/Engine/Source/Editor/MaterialEditor/Private/MaterialEditor.cpp
+++ b/Engine/Source/Editor/MaterialEditor/Private/MaterialEditor.cpp
@@ -275,9 +275,8 @@ int32 FMatExpressionPreview::CompilePropertyAndSetMaterialProperty(EMaterialProp
 	{
 		// Hardcoding output 0 as we don't have the UI to specify any other output
 		const int32 OutputIndex = 0;
-		const uint32 Output0Type = Expression->GetOutputType(OutputIndex);
 		int32 PreviewCodeChunk = INDEX_NONE;
-		if (Output0Type != MCT_Strata)
+		if (!Expression->Outputs.IsEmpty() && Expression->GetOutputType(OutputIndex) != MCT_Strata)
 		{
 			PreviewCodeChunk = Expression->CompilePreview(Compiler, OutputIndex);
 		}
