Since the editor only data is no longer stripped we don't need to make every parameter visible.

- This fixes issues where we would display parameters that weren't used in the material.

diff --git a/Engine/Source/Editor/MaterialEditor/Private/MaterialEditorInstanceDetailCustomization.cpp b/Engine/Source/Editor/MaterialEditor/Private/MaterialEditorInstanceDetailCustomization.cpp
--- a/Engine/Source/Editor/MaterialEditor/Private/MaterialEditorInstanceDetailCustomization.cpp
+++ b/Engine/Source/Editor/MaterialEditor/Private/MaterialEditorInstanceDetailCustomization.cpp
@@ -287,19 +287,10 @@ void FMaterialInstanceParameterDetails::CreateGroupsWidget(TSharedRef<IPropertyH
 		{
 			bShowSaveButtons = true;
 			bool bCreateGroup = false;
-			bool bIsCookedMaterial = false;
-			if (MaterialEditorInstance->SourceInstance)
-			{
-				if (UMaterial* Material = MaterialEditorInstance->SourceInstance->GetMaterial())
-				{
-					bIsCookedMaterial = Material->GetPackage()->HasAnyPackageFlags(PKG_Cooked);
-				}
-			}
 			for (int32 ParamIdx = 0; ParamIdx < ParameterGroup.Parameters.Num() && !bCreateGroup; ++ParamIdx)
 			{
 				UDEditorParameterValue* Parameter = ParameterGroup.Parameters[ParamIdx];
-				// All parameters on cooked materials are visible since the underlying data to discover visibility has been stripped.
-				const bool bIsVisible = bIsCookedMaterial || MaterialEditorInstance->VisibleExpressions.Contains(Parameter->ParameterInfo);
+				const bool bIsVisible = MaterialEditorInstance->VisibleExpressions.Contains(Parameter->ParameterInfo);
 				bCreateGroup = bIsVisible && (!MaterialEditorInstance->bShowOnlyOverrides || FMaterialPropertyHelpers::IsOverriddenExpression(Parameter));
 			}
 			
