Added MaterialEditor override for "FEditorUndoClient::MatchesContext()" function to limit the undo operations only to Material related UObjects.

This fixes the issue of the user undoing something unrelated to the Material and the MaterialEditor erroneously flagging the Material as dirty, causing the "Save Material" dialog to show up when the user closest he MaterialEditor window.

diff --git a/Engine/Source/Editor/MaterialEditor/Private/MaterialEditor.cpp b/Engine/Source/Editor/MaterialEditor/Private/MaterialEditor.cpp
--- a/Engine/Source/Editor/MaterialEditor/Private/MaterialEditor.cpp
+++ b/Engine/Source/Editor/MaterialEditor/Private/MaterialEditor.cpp
@@ -7,6 +7,7 @@
 #include "Engine/Engine.h"
 #include "EngineModule.h"
 #include "Misc/MessageDialog.h"
+#include "Misc/TransactionObjectEvent.h"
 #include "Modules/ModuleManager.h"
 #include "SlateOptMacros.h"
 #include "Widgets/Layout/SSeparator.h"
@@ -34,6 +35,7 @@
 #include "Materials/MaterialInstanceDynamic.h"
 #include "Materials/MaterialInstanceConstant.h"
 #include "Materials/MaterialParameterCollection.h"
+#include "Materials/MaterialParameterCollectionInstance.h"
 #include "StaticParameterSet.h"
 #include "Engine/TextureCube.h"
 #include "Engine/Texture2DArray.h"
@@ -2105,6 +2107,36 @@ void FMaterialEditor::AddGraphEditorPinActionsToContextMenu(FToolMenuSection& In
 	}
 }
 
+bool FMaterialEditor::MatchesContext(const FTransactionContext& InContext, const TArray<TPair<UObject*, FTransactionObjectEvent>>& TransactionObjectContexts) const
+{
+	for (const TPair<UObject*, FTransactionObjectEvent>& TransactionObjectContext : TransactionObjectContexts)
+	{
+		UObject* Object = TransactionObjectContext.Get<0>();
+
+		// Evaluate whether any object we are interested in matches an object part of the transaction.
+		bool bIsMaterialRelatedObject =
+			   Object->IsA<UMaterialInterface>()
+			|| Object->IsA<UMaterialEditorOnlyData>()
+			|| Object->IsA<UMaterialExpression>()
+			|| Object->IsA<UMaterialGraph>()
+			|| Object->IsA<UMaterialFunctionInterface>()
+			|| Object->IsA<UMaterialEditingLibrary>()
+			|| Object->IsA<UMaterialEditorSettings>()
+			|| Object->IsA<UMaterialEditorInstanceConstant>()
+			|| Object->IsA<UMaterialEditorMenuContext>()
+			|| Object->IsA<UMaterialEditorOptions>()
+			|| Object->IsA<UMaterialParameterCollection>()
+			|| Object->IsA<UMaterialParameterCollectionInstance>();
+
+		if (bIsMaterialRelatedObject)
+		{
+			return true;
+		}
+	}
+
+	return false;
+}
+
 void FMaterialEditor::DrawMaterialInfoStrings(
 	FCanvas* Canvas, 
 	const UMaterial* Material, 
diff --git a/Engine/Source/Editor/MaterialEditor/Private/MaterialEditor.h b/Engine/Source/Editor/MaterialEditor/Private/MaterialEditor.h
--- a/Engine/Source/Editor/MaterialEditor/Private/MaterialEditor.h
+++ b/Engine/Source/Editor/MaterialEditor/Private/MaterialEditor.h
@@ -425,6 +425,9 @@ public:
 
 	virtual void AddGraphEditorPinActionsToContextMenu(FToolMenuSection& InSection) const override;
 
+	/** Overrides function in FEditorUndoClient. Called to see if the context of the current undo/redo operation is a match for the client. */
+	virtual bool MatchesContext(const FTransactionContext& InContext, const TArray<TPair<UObject*, FTransactionObjectEvent>>& TransactionObjectContexts) const override;
+
 public:
 	/** Set to true when modifications have been made to the material */
 	bool bMaterialDirty;
