Fix crash when copy pasting nodes in the Material Editor.

diff --git a/Engine/Source/Editor/MaterialEditor/Private/MaterialEditor.cpp b/Engine/Source/Editor/MaterialEditor/Private/MaterialEditor.cpp
--- a/Engine/Source/Editor/MaterialEditor/Private/MaterialEditor.cpp
+++ b/Engine/Source/Editor/MaterialEditor/Private/MaterialEditor.cpp
@@ -329,15 +329,19 @@ int32 FMatExpressionPreview::CompilePropertyAndSetMaterialProperty(EMaterialProp
 
 UMaterialInterface* FMatExpressionPreview::GetMaterialInterface() const
 {
-	UMaterial* ExprMat = Expression->Material;
-	if (ExprMat)
+	if (Expression.IsValid())
 	{
-		FMaterialRenderProxy* MatProxy = ExprMat->GetRenderProxy();
-		if (MatProxy)
+		UMaterial* ExprMat = Expression->Material;
+		if (ExprMat)
 		{
-			return MatProxy->GetMaterialInterface();
+			FMaterialRenderProxy* MatProxy = ExprMat->GetRenderProxy();
+			if (MatProxy)
+			{
+				return MatProxy->GetMaterialInterface();
+			}
 		}
 	}
+
 	return nullptr;
 }
 
