Fix for material editor crash that may occur when users a platform is loaded but not valid due to the users manually disabling individual platform/plugins.

diff --git a/Engine/Source/Editor/MaterialEditor/Private/MaterialStats.h b/Engine/Source/Editor/MaterialEditor/Private/MaterialStats.h
--- a/Engine/Source/Editor/MaterialEditor/Private/MaterialStats.h
+++ b/Engine/Source/Editor/MaterialEditor/Private/MaterialStats.h
@@ -531,11 +531,13 @@ void FMaterialStats::SetMaterial(TMaterial& MaterialPtr)
 	if (MaterialInterface != MaterialPtr)
 	{
 		MaterialInterface = MaterialPtr;
-
 		for (const auto& Entry : ShaderPlatformStatsDB)
 		{
-			auto Platform = Entry.Value;
-			Platform->SetMaterial(MaterialPtr);
+			auto& Platform = Entry.Value;
+			if (Platform.IsValid())
+			{
+				Platform->SetMaterial(MaterialPtr);
+			}
 		}
 	}
 }
