Fix a crash when reverting files, caused by assert firing IsChildOf changes as assets are removed from the registry

diff --git a/Engine/Source/Editor/MaterialEditor/Private/SMaterialPalette.cpp b/Engine/Source/Editor/MaterialEditor/Private/SMaterialPalette.cpp
--- a/Engine/Source/Editor/MaterialEditor/Private/SMaterialPalette.cpp
+++ b/Engine/Source/Editor/MaterialEditor/Private/SMaterialPalette.cpp
@@ -232,10 +232,7 @@ void SMaterialPalette::RenameAssetFromRegistry(const FAssetData& InAddedAssetDat
 
 void SMaterialPalette::RefreshAssetInRegistry(const FAssetData& InAddedAssetData)
 {
-	// Grab the asset class, it will be checked for being a material function.
-	UClass* Asset = FindObject<UClass>(InAddedAssetData.AssetClassPath);
-
-	if (Asset->IsChildOf(UMaterialFunction::StaticClass()))
+	if (InAddedAssetData.IsInstanceOf(UMaterialFunction::StaticClass()))
 	{
 		RefreshActionsList(true);
 	}
