M - Implement double click on MaterialEditor Errors to jump to the MaterialNode (if any) that caused the error.

diff --git a/Engine/Source/Editor/MaterialEditor/Private/MaterialEditor.cpp b/Engine/Source/Editor/MaterialEditor/Private/MaterialEditor.cpp
--- a/Engine/Source/Editor/MaterialEditor/Private/MaterialEditor.cpp
+++ b/Engine/Source/Editor/MaterialEditor/Private/MaterialEditor.cpp
@@ -7,6 +7,7 @@
 #include "Engine/Engine.h"
 #include "EngineModule.h"
 #include "Misc/MessageDialog.h"
+#include "Misc/UObjectToken.h"
 #include "Misc/TransactionObjectEvent.h"
 #include "Modules/ModuleManager.h"
 #include "SlateOptMacros.h"
@@ -540,6 +541,7 @@ void FMaterialEditor::InitMaterialEditor( const EToolkitMode::Type Mode, const T
 
 	MaterialStatsManager = FMaterialStatsUtils::CreateMaterialStats(this);
 	MaterialStatsManager->SetMaterialDisplayName(OriginalMaterial->GetName());
+	MaterialStatsManager->GetOldStatsListing()->OnMessageTokenClicked().AddSP(this, &FMaterialEditor::OnMessageLogLinkActivated);
 
 	if (!Material->MaterialGraph)
 	{
@@ -2723,6 +2725,42 @@ bool FMaterialEditor::UpdateOriginalMaterial()
 	return true;
 }
 
+void FMaterialEditor::OnMessageLogLinkActivated(const class TSharedRef<IMessageToken>& Token)
+{
+	const TSharedRef<FUObjectToken> UObjectToken = StaticCastSharedRef<FUObjectToken>(Token);
+	if (UObjectToken->GetObject().IsValid())
+	{
+		UMaterialExpression* Expression = Cast<UMaterialExpression>(UObjectToken->GetObject().Get());
+		if(UObject* MaterialOrFunction = Expression->GetAssetOwner())
+		{
+			UAssetEditorSubsystem* AssetEditor = GEditor->GetEditorSubsystem<UAssetEditorSubsystem>();
+			if(AssetEditor->OpenEditorForAsset(MaterialOrFunction))
+			{
+				FMaterialEditor* TargetEditor = static_cast<FMaterialEditor*>(AssetEditor->FindEditorForAsset(MaterialOrFunction, true));
+				checkf(TargetEditor, TEXT("Could not find Editor for Asset: %s"), *(MaterialOrFunction->GetFName().ToString()));
+
+				FMaterialExpressionCollection& Collection = Expression->Function ? TargetEditor->MaterialFunction->GetEditorOnlyData()->ExpressionCollection
+																		   : TargetEditor->Material->GetEditorOnlyData()->ExpressionCollection;
+
+				for (const TObjectPtr<UMaterialExpression>& EditorExpression : Collection.Expressions)
+				{
+					if (EditorExpression->MaterialExpressionGuid == Expression->MaterialExpressionGuid
+						&& EditorExpression->MaterialExpressionEditorX == Expression->MaterialExpressionEditorX
+						&& EditorExpression->MaterialExpressionEditorY == Expression->MaterialExpressionEditorY
+						&& EditorExpression->GetName() == Expression->GetName())
+					{
+						if (EditorExpression->GraphNode)
+						{
+							TargetEditor->JumpToNode(EditorExpression->GraphNode);
+							break;
+						}
+					}
+				}
+			}
+		}
+	}
+}
+
 void FMaterialEditor::UpdateMaterialinfoList_Old()
 {
 	bool bForceDisplay = false;
@@ -2739,7 +2777,9 @@ void FMaterialEditor::UpdateMaterialinfoList_Old()
 
 		for (int32 i = 0; i < NumFeatureLevels; ++i)
 		{
-			TArray<FString> CompileErrors;
+			TArray<FString>				 CompileErrors;
+			TArray<UMaterialExpression*> FailingExpression;
+
 			ERHIFeatureLevel::Type FeatureLevel = FeatureLevelsToDisplay[i];
 			const FMaterialResource* MaterialResource = MaterialForStats->GetMaterialResource(FeatureLevel);
 
@@ -2761,6 +2801,7 @@ void FMaterialEditor::UpdateMaterialinfoList_Old()
 							if (NumInputs > 1 || !InputExpression->IsResultMaterialAttributes(0))
 							{
 								CompileErrors.Add(TEXT("Layer graphs only support a single material attributes input."));
+								FailingExpression.Add(nullptr);
 							}
 						}
 						else if (UMaterialExpressionFunctionOutput* OutputExpression = Cast<UMaterialExpressionFunctionOutput>(Expression))
@@ -2769,17 +2810,20 @@ void FMaterialEditor::UpdateMaterialinfoList_Old()
 							if (NumOutputs > 1 || !OutputExpression->IsResultMaterialAttributes(0))
 							{
 								CompileErrors.Add(TEXT("Layer graphs only support a single material attributes output."));
+								FailingExpression.Add(nullptr);
 							}
 						}
 						else if (UMaterialExpressionMaterialAttributeLayers* RecursiveLayer = Cast<UMaterialExpressionMaterialAttributeLayers>(Expression))
 						{
 							CompileErrors.Add(TEXT("Layer graphs do not support layers within layers."));
+							FailingExpression.Add(nullptr);
 						}
 					}
 
 					if (NumInputs > 1 || NumOutputs < 1)
 					{
 						CompileErrors.Add(TEXT("Layer graphs require a single material attributes output and optionally, a single material attributes input."));
+						FailingExpression.Add(nullptr);
 					}
 				}
 				else if (MaterialFunction->GetMaterialFunctionUsage() == EMaterialFunctionUsage::MaterialLayerBlend)
@@ -2793,6 +2837,7 @@ void FMaterialEditor::UpdateMaterialinfoList_Old()
 							if (NumInputs > 2 || !InputExpression->IsResultMaterialAttributes(0))
 							{
 								CompileErrors.Add(TEXT("Layer blend graphs only support two material attributes inputs."));
+								FailingExpression.Add(nullptr);
 							}
 						}
 						else if (UMaterialExpressionFunctionOutput* OutputExpression = Cast<UMaterialExpressionFunctionOutput>(Expression))
@@ -2801,17 +2846,20 @@ void FMaterialEditor::UpdateMaterialinfoList_Old()
 							if (NumOutputs > 1 || !OutputExpression->IsResultMaterialAttributes(0))
 							{
 								CompileErrors.Add(TEXT("Layer blend graphs only support a single material attributes output."));
+								FailingExpression.Add(nullptr);
 							}
 						}
 						else if (UMaterialExpressionMaterialAttributeLayers* RecursiveLayer = Cast<UMaterialExpressionMaterialAttributeLayers>(Expression))
 						{
 							CompileErrors.Add(TEXT("Layer blend graphs do not support layers within layers."));
+							FailingExpression.Add(nullptr);
 						}
 					}
 
 					if (NumOutputs < 1)
 					{
 						CompileErrors.Add(TEXT("Layer blend graphs can have up to two material attributes inputs and a single output."));
+						FailingExpression.Add(nullptr);
 					}
 				}
 				else
@@ -2821,6 +2869,7 @@ void FMaterialEditor::UpdateMaterialinfoList_Old()
 					if (CurrentResource)
 					{
 						CompileErrors = CurrentResource->GetCompileErrors();
+						FailingExpression = CurrentResource->GetErrorExpressions();
 					}
 
 					bool bFoundFunctionOutput = false;
@@ -2836,12 +2885,14 @@ void FMaterialEditor::UpdateMaterialinfoList_Old()
 					if (!bFoundFunctionOutput)
 					{
 						CompileErrors.Add(TEXT("Missing a function output"));
+						FailingExpression.Add(nullptr);
 					}
 				}
 			}
 			else
 			{
 				CompileErrors = MaterialResource->GetCompileErrors();
+				FailingExpression = MaterialResource->GetErrorExpressions();
 			}
 
 			// Only show general info if there are no errors and stats are enabled - Stats show for Materials, layers and blends
@@ -2975,6 +3026,10 @@ void FMaterialEditor::UpdateMaterialinfoList_Old()
 				FString ErrorString = FString::Printf(TEXT("[%s] %s"), *FeatureLevelName, *CompileErrors[ErrorIndex]);
 				TempMaterialInfoList.Add(MakeShareable(new FMaterialInfo(ErrorString, FLinearColor::Red)));
 				TSharedRef<FTokenizedMessage> Line = FTokenizedMessage::Create( EMessageSeverity::Error );
+				if(ensure(FailingExpression.Num() == CompileErrors.Num()) && FailingExpression[ErrorIndex])
+				{
+					Line->SetMessageLink(FUObjectToken::Create(FailingExpression[ErrorIndex]));
+				}
 				Line->AddToken( FTextToken::Create( FText::FromString( ErrorString ) ) );
 				Messages.Add(Line);
 				bForceDisplay = true;
diff --git a/Engine/Source/Editor/MaterialEditor/Private/MaterialEditor.h b/Engine/Source/Editor/MaterialEditor/Private/MaterialEditor.h
--- a/Engine/Source/Editor/MaterialEditor/Private/MaterialEditor.h
+++ b/Engine/Source/Editor/MaterialEditor/Private/MaterialEditor.h
@@ -584,6 +584,8 @@ protected:
 	void OnDistributeNodesV();
 
 private:
+	void OnMessageLogLinkActivated(const class TSharedRef<IMessageToken>& Token);
+
 	/** Builds the toolbar widget for the material editor */
 	void ExtendToolbar();
 	void RegisterToolBar();
diff --git a/Engine/Source/Runtime/Engine/Private/Materials/HLSLMaterialTranslator.cpp b/Engine/Source/Runtime/Engine/Private/Materials/HLSLMaterialTranslator.cpp
--- a/Engine/Source/Runtime/Engine/Private/Materials/HLSLMaterialTranslator.cpp
+++ b/Engine/Source/Runtime/Engine/Private/Materials/HLSLMaterialTranslator.cpp
@@ -3649,23 +3649,24 @@ int32 FHLSLMaterialTranslator::Error(const TCHAR* Text)
 	if (!bUsingErrorProxy)
 	{
 		// Standard error handling, immediately append one-off errors and signal failure
-		CompileErrors.AddUnique(ErrorString);
-	
+		bSuccess = false;
 		if (ExpressionToError)
 		{
-			ErrorExpressions.Add(ExpressionToError);
 			ExpressionToError->LastErrorText = Text;
+			for (int32 i = 0; i < ErrorExpressions.Num(); i++)
+			{
+				if(ErrorExpressions[i] == ExpressionToError && CompileErrors[i] == ErrorString)
+				{
+					return INDEX_NONE;
+				}
+			}
 		}
-
-		bSuccess = false;
-	}
-	else
-	{
-		// When a proxy is intercepting errors, ignore the failure and match arrays to allow later error type selection
-		CompileErrors.Add(ErrorString);
-		ErrorExpressions.Add(ExpressionToError);		
 	}
-		
+
+	// When a proxy is intercepting errors, ignore the failure and match arrays to allow later error type selection
+	CompileErrors.Add(ErrorString);
+	ErrorExpressions.Add(ExpressionToError);
+
 	return INDEX_NONE;
 }
 
