Only try and check for a matching expression to error on when we at least have one result in the Expressions Array as there might be cases where we pass GraphCompilation but than later get errors from the backend compiler which do not have a coresponding entry in the expression array.

diff --git a/Engine/Source/Editor/MaterialEditor/Private/MaterialEditor.cpp b/Engine/Source/Editor/MaterialEditor/Private/MaterialEditor.cpp
--- a/Engine/Source/Editor/MaterialEditor/Private/MaterialEditor.cpp
+++ b/Engine/Source/Editor/MaterialEditor/Private/MaterialEditor.cpp
@@ -3026,7 +3026,7 @@ void FMaterialEditor::UpdateMaterialinfoList_Old()
 				FString ErrorString = FString::Printf(TEXT("[%s] %s"), *FeatureLevelName, *CompileErrors[ErrorIndex]);
 				TempMaterialInfoList.Add(MakeShareable(new FMaterialInfo(ErrorString, FLinearColor::Red)));
 				TSharedRef<FTokenizedMessage> Line = FTokenizedMessage::Create( EMessageSeverity::Error );
-				if(ensure(FailingExpression.Num() == CompileErrors.Num()) && FailingExpression[ErrorIndex])
+				if(FailingExpression.Num() && ensure(FailingExpression.Num() == CompileErrors.Num()) && FailingExpression[ErrorIndex])
 				{
 					Line->SetMessageLink(FUObjectToken::Create(FailingExpression[ErrorIndex]));
 				}
