Expose GetMaterialExpressionNodePosition() in Blueprint to be able to place new expressions in relation to existing ones.

diff --git a/Engine/Source/Editor/MaterialEditor/Private/MaterialEditingLibrary.cpp b/Engine/Source/Editor/MaterialEditor/Private/MaterialEditingLibrary.cpp
--- a/Engine/Source/Editor/MaterialEditor/Private/MaterialEditingLibrary.cpp
+++ b/Engine/Source/Editor/MaterialEditor/Private/MaterialEditingLibrary.cpp
@@ -876,6 +876,15 @@ bool UMaterialEditingLibrary::GetInputNodeOutputNameForMaterialExpression(UMater
 	return false;
 }
 
+void UMaterialEditingLibrary::GetMaterialExpressionNodePosition(UMaterialExpression* MaterialExpression, int32& NodePosX, int32& NodePosY)
+{
+	if (MaterialExpression)
+	{
+		NodePosX = MaterialExpression->MaterialExpressionEditorX;
+		NodePosY = MaterialExpression->MaterialExpressionEditorY;
+	}
+}
+
 TArray<UTexture*> UMaterialEditingLibrary::GetUsedTextures(UMaterial* Material)
 {
 	TArray<UTexture*> OutTextures;
diff --git a/Engine/Source/Editor/MaterialEditor/Public/MaterialEditingLibrary.h b/Engine/Source/Editor/MaterialEditor/Public/MaterialEditingLibrary.h
--- a/Engine/Source/Editor/MaterialEditor/Public/MaterialEditingLibrary.h
+++ b/Engine/Source/Editor/MaterialEditor/Public/MaterialEditingLibrary.h
@@ -207,6 +207,10 @@ public:
 	UFUNCTION(BlueprintPure, Category = "MaterialEditing")
 	static bool GetInputNodeOutputNameForMaterialExpression(UMaterialExpression* MaterialExpression, UMaterialExpression* InputNode, FString& OutputName);
 
+	/** Get the position of the MaterialExpression node. */
+	UFUNCTION(BlueprintPure, Category = "MaterialEditing")
+	static void GetMaterialExpressionNodePosition(UMaterialExpression* MaterialExpression, int32& NodePosX, int32& NodePosY);
+
 	/** Get the list of textures used by a material */
 	UFUNCTION(BlueprintPure, Category = "MaterialEditing")
 	static TArray<UTexture*> GetUsedTextures(UMaterial* Material);
