Removed PropertyRow call to CustomResetToDefault for Labeled Texture Parameter customization

diff --git a/Engine/Source/Editor/MaterialEditor/Private/MaterialEditorInstanceDetailCustomization.cpp b/Engine/Source/Editor/MaterialEditor/Private/MaterialEditorInstanceDetailCustomization.cpp
--- a/Engine/Source/Editor/MaterialEditor/Private/MaterialEditorInstanceDetailCustomization.cpp
+++ b/Engine/Source/Editor/MaterialEditor/Private/MaterialEditorInstanceDetailCustomization.cpp
@@ -750,8 +750,7 @@ void FMaterialInstanceParameterDetails::CreateLabeledTextureParameterValueWidget
 						.DisplayName(FText::FromName(Parameter->ParameterInfo.Name))
 						.EditCondition(IsParamEnabled, FOnBooleanValueChanged::CreateStatic(&FMaterialPropertyHelpers::OnOverrideParameter, Parameter, MaterialEditorInstance))
 						.ToolTip(FMaterialPropertyHelpers::GetParameterTooltip(Parameter, MaterialEditorInstance))
-						.Visibility(TAttribute<EVisibility>::Create(TAttribute<EVisibility>::FGetter::CreateStatic(&FMaterialPropertyHelpers::ShouldShowExpression, Parameter, MaterialEditorInstance, ShowHiddenDelegate)))
-						.OverrideResetToDefault(ResetOverride);
+						.Visibility(TAttribute<EVisibility>::Create(TAttribute<EVisibility>::FGetter::CreateStatic(&FMaterialPropertyHelpers::ShouldShowExpression, Parameter, MaterialEditorInstance, ShowHiddenDelegate)));
 
 
 					TSharedPtr<SWidget> NameWidget;
