Add material stats permutation for Runtime Virtual Texture output.
Useful for basic understanding of perf for landscape shaders that write RVT.
Mark material stats objects as Preview Materials. Without that we cull out any work gated by a layer weight when calculating stats.

diff --git a/Engine/Source/Editor/MaterialEditor/Private/MaterialStatsCommon.cpp b/Engine/Source/Editor/MaterialEditor/Private/MaterialStatsCommon.cpp
--- a/Engine/Source/Editor/MaterialEditor/Private/MaterialStatsCommon.cpp
+++ b/Engine/Source/Editor/MaterialEditor/Private/MaterialStatsCommon.cpp
@@ -201,6 +201,10 @@ FString FMaterialStatsUtils::RepresentativeShaderTypeToString(const ERepresentat
 			return TEXT("UI Instanced Vertex Shader");
 		break;
 
+		case ERepresentativeShader::RuntimeVirtualTextureOutput:
+			return TEXT("Runtime Virtual Texture Output");
+			break;
+
 		default:
 			return TEXT("Unknown shader name");
 		break;
@@ -342,6 +346,16 @@ void FMaterialStatsUtils::GetRepresentativeShaderTypesAndDescriptions(TMap<FName
 		// This is ad-hoc, and ideally we have a better way for finding this shader type in the future.
 		ShaderTypeNamesAndDescriptions.FindOrAdd(FLocalVertexFactoryName)
 			.Add(FRepresentativeShaderInfo(ERepresentativeShader::StationarySurface, FName(TEXT("TBasePassPSFCachedVolumeIndirectLightingPolicy")), TEXT("MaxSampler")));
+
+		// For materials that write to a runtime virtual texture add a pixel shader stat.
+		if (TargetMaterial->HasRuntimeVirtualTextureOutput())
+		{
+			static const FName LandscapeFactoryName = TEXT("FLandscapeVertexFactory");
+			static const FName RuntimeVirtualTexturePSName = TEXT("TVirtualTexturePSBaseColorNormalSpecular");
+
+			ShaderTypeNamesAndDescriptions.FindOrAdd(LandscapeFactoryName)
+				.Add(FRepresentativeShaderInfo(ERepresentativeShader::RuntimeVirtualTextureOutput, FName(RuntimeVirtualTexturePSName), TEXT("Runtime Virtual Texture Output")));
+		}
 	}
 	else
 	{
diff --git a/Engine/Source/Editor/MaterialEditor/Private/MaterialStatsGrid.cpp b/Engine/Source/Editor/MaterialEditor/Private/MaterialStatsGrid.cpp
--- a/Engine/Source/Editor/MaterialEditor/Private/MaterialStatsGrid.cpp
+++ b/Engine/Source/Editor/MaterialEditor/Private/MaterialStatsGrid.cpp
@@ -166,6 +166,7 @@ FStatsGridRow_Shaders::EShaderClass FStatsGridRow_Shaders::GetShaderClass(const
 		case ERepresentativeShader::StationarySurfaceNPointLights:
 		case ERepresentativeShader::DynamicallyLitObject:
 		case ERepresentativeShader::UIDefaultFragmentShader:
+		case ERepresentativeShader::RuntimeVirtualTextureOutput:
 			return EShaderClass::FragmentShader;
 		break;
 
diff --git a/Engine/Source/Editor/MaterialEditor/Public/MaterialStatsCommon.h b/Engine/Source/Editor/MaterialEditor/Public/MaterialStatsCommon.h
--- a/Engine/Source/Editor/MaterialEditor/Public/MaterialStatsCommon.h
+++ b/Engine/Source/Editor/MaterialEditor/Public/MaterialStatsCommon.h
@@ -23,6 +23,9 @@ public:
 	 */
 	virtual bool IsPersistent() const override { return false; }
 
+	/** Use preview settings when compiling material. Used to avoid culling nodes behind landscape weights. */
+	virtual bool IsPreview() const override { return true; }
+
 	/** Material resource stats never requires a synchronous compilation, otherwise opening up empty (newly created) material instance will block compiling default mat's shaders. */
 	virtual bool RequiresSynchronousCompilation() const override { return false; }
 
@@ -49,6 +52,7 @@ enum class ERepresentativeShader
 	StationarySurfaceCSM,
 	StationarySurfaceNPointLights,
 	DynamicallyLitObject,
+	RuntimeVirtualTextureOutput,
 	UIDefaultFragmentShader,
 	LastFragmentShader = UIDefaultFragmentShader,
 
