Skip work for null MaterialResources

diff --git a/Engine/Source/Editor/MaterialEditor/Private/MaterialEditor.cpp b/Engine/Source/Editor/MaterialEditor/Private/MaterialEditor.cpp
--- a/Engine/Source/Editor/MaterialEditor/Private/MaterialEditor.cpp
+++ b/Engine/Source/Editor/MaterialEditor/Private/MaterialEditor.cpp
@@ -2789,6 +2789,11 @@ void FMaterialEditor::UpdateMaterialinfoList_Old()
 			ERHIFeatureLevel::Type FeatureLevel = FeatureLevelsToDisplay[i];
 			const FMaterialResource* MaterialResource = MaterialForStats->GetMaterialResource(FeatureLevel);
 
+			if (MaterialResource == nullptr)
+			{
+				continue;
+			}
+
 			if (MaterialFunction && ExpressionPreviewMaterial)
 			{
 				bool bHasValidOutput = true;
