Fixed crash occurring when the user selects multiple material parameter collection nodes belonging to different collections.

diff --git a/Engine/Source/Editor/MaterialEditor/Private/MaterialEditorDetailCustomization.cpp b/Engine/Source/Editor/MaterialEditor/Private/MaterialEditorDetailCustomization.cpp
--- a/Engine/Source/Editor/MaterialEditor/Private/MaterialEditorDetailCustomization.cpp
+++ b/Engine/Source/Editor/MaterialEditor/Private/MaterialEditorDetailCustomization.cpp
@@ -597,7 +597,15 @@ bool FMaterialExpressionCollectionParameterDetails::IsParameterNameComboEnabled(
 	if (CollectionPropertyHandle->IsValidHandle())
 	{
 		UObject* CollectionObject = nullptr;
-		verify(CollectionPropertyHandle->GetValue(CollectionObject) == FPropertyAccess::Success);
+		FPropertyAccess::Result Result = CollectionPropertyHandle->GetValue(CollectionObject);
+
+		// No name combo enabled if multiple parameter collection nodes are selected.
+		if (Result == FPropertyAccess::MultipleValues)
+		{
+			return false;
+		}
+
+		verify(Result == FPropertyAccess::Success);
 		Collection = Cast<UMaterialParameterCollection>(CollectionObject);
 	}
 
@@ -682,7 +690,15 @@ void FMaterialExpressionCollectionParameterDetails::PopulateParameters()
 	if (CollectionPropertyHandle->IsValidHandle())
 	{
 		UObject* CollectionObject = nullptr;
-		verify(CollectionPropertyHandle->GetValue(CollectionObject) == FPropertyAccess::Success);
+		FPropertyAccess::Result Result = CollectionPropertyHandle->GetValue(CollectionObject);
+
+		// Return an empty set of parameters if multiple paramter collection nodes are selected.
+		if (Result == FPropertyAccess::MultipleValues)
+		{
+			return;
+		}
+
+		verify(Result == FPropertyAccess::Success);
 		Collection = Cast<UMaterialParameterCollection>(CollectionObject);
 	}
 
