Add IsThinSurface on root node for flagging surface with thin shell (i.e., hollow interrior).

The defines a clear semantic for how thickness of a bottom layer should be handled with Strata. This CL only pipes the data for the various parts. A following CL will add a thickness property on the root node.

diff --git a/Engine/Plugins/Editor/BlueprintMaterialTextureNodes/Source/BlueprintMaterialTextureNodes/Private/BlueprintMaterialTextureNodesBPLibrary.cpp b/Engine/Plugins/Editor/BlueprintMaterialTextureNodes/Source/BlueprintMaterialTextureNodes/Private/BlueprintMaterialTextureNodesBPLibrary.cpp
--- a/Engine/Plugins/Editor/BlueprintMaterialTextureNodes/Source/BlueprintMaterialTextureNodes/Private/BlueprintMaterialTextureNodesBPLibrary.cpp
+++ b/Engine/Plugins/Editor/BlueprintMaterialTextureNodes/Source/BlueprintMaterialTextureNodes/Private/BlueprintMaterialTextureNodesBPLibrary.cpp
@@ -402,6 +402,24 @@ bool UBlueprintMaterialTextureNodesBPLibrary::SetMICTwoSided_EditorOnly(UMateria
 	return 0;
 }
 
+bool UBlueprintMaterialTextureNodesBPLibrary::SetMICIsThinSurface_EditorOnly(UMaterialInstanceConstant* Material, bool bIsThinSurface)
+{
+#if WITH_EDITOR
+	if (Material != nullptr)
+	{
+		Material->BasePropertyOverrides.bOverride_bIsThinSurface = true;
+		Material->BasePropertyOverrides.bIsThinSurface = bIsThinSurface;
+		UpdateMIC(Material);
+
+		return 1;
+	}
+	FMessageLog("Blueprint").Warning(LOCTEXT("SetMICIsThinSurface_InvalidMIC", "SetMICIsThinSurface_EditorOnly: MIC must be non-null."));
+#else
+	FMessageLog("Blueprint").Error(LOCTEXT("SetMICIsThinSurface_CannotBeModifiedAtRuntime", "SetMICIsThinSurface: Can't modify MIC at run time."));
+#endif
+	return 0;
+}
+
 bool UBlueprintMaterialTextureNodesBPLibrary::SetMICDitheredLODTransition_EditorOnly(UMaterialInstanceConstant* Material, bool DitheredLODTransition)
 {
 #if WITH_EDITOR
diff --git a/Engine/Plugins/Editor/BlueprintMaterialTextureNodes/Source/BlueprintMaterialTextureNodes/Public/BlueprintMaterialTextureNodesBPLibrary.h b/Engine/Plugins/Editor/BlueprintMaterialTextureNodes/Source/BlueprintMaterialTextureNodes/Public/BlueprintMaterialTextureNodesBPLibrary.h
--- a/Engine/Plugins/Editor/BlueprintMaterialTextureNodes/Source/BlueprintMaterialTextureNodes/Public/BlueprintMaterialTextureNodesBPLibrary.h
+++ b/Engine/Plugins/Editor/BlueprintMaterialTextureNodes/Source/BlueprintMaterialTextureNodes/Public/BlueprintMaterialTextureNodesBPLibrary.h
@@ -102,6 +102,13 @@ class BLUEPRINTMATERIALTEXTURENODES_API UBlueprintMaterialTextureNodesBPLibrary
 	UFUNCTION(BlueprintCallable, meta = (DisplayName = "Set MIC Two Sided Editor Only", Keywords = "Set MIC Two Sided", UnsafeDuringActorConstruction = "true"), Category = Rendering)
 	static bool SetMICTwoSided_EditorOnly(UMaterialInstanceConstant* Material, bool TwoSided = false);
 
+	/**
+	* Overrides the IsThinSurface setting of a Material Instance Constant
+	* Only works in the editor
+	*/
+	UFUNCTION(BlueprintCallable, meta = (DisplayName = "Set MIC IsThinSurface Editor Only", Keywords = "Set MIC IsThinSurface", UnsafeDuringActorConstruction = "true"), Category = Rendering)
+	static bool SetMICIsThinSurface_EditorOnly(UMaterialInstanceConstant* Material, bool bIsThinSurface = false);	
+
 	/**
 	* Overrides the Blend Mode of a Material Instance Constant
 	* Only works in the editor
diff --git a/Engine/Plugins/Editor/GLTFImporter/Source/GLTFCore/Public/GLTFMaterialExpressions.h b/Engine/Plugins/Editor/GLTFImporter/Source/GLTFCore/Public/GLTFMaterialExpressions.h
--- a/Engine/Plugins/Editor/GLTFImporter/Source/GLTFCore/Public/GLTFMaterialExpressions.h
+++ b/Engine/Plugins/Editor/GLTFImporter/Source/GLTFCore/Public/GLTFMaterialExpressions.h
@@ -323,6 +323,8 @@ namespace GLTF
 		virtual void SetBlendMode(int InBlendMode) = 0;
 		virtual bool GetTwoSided() const           = 0;
 		virtual void SetTwoSided(bool bTwoSided)   = 0;
+		virtual bool GetIsThinSurface() const = 0;
+		virtual void SetIsThinSurface(bool bIsThinSurface) = 0;
 		virtual void SetShadingModel(EGLTFMaterialShadingModel InShadingModel) = 0;
 		virtual void SetTranslucencyLightingMode(int InLightingMode) = 0;
 
diff --git a/Engine/Plugins/Editor/GLTFImporter/Source/GLTFImporter/Private/GLTFMaterialElement.cpp b/Engine/Plugins/Editor/GLTFImporter/Source/GLTFImporter/Private/GLTFMaterialElement.cpp
--- a/Engine/Plugins/Editor/GLTFImporter/Source/GLTFImporter/Private/GLTFMaterialElement.cpp
+++ b/Engine/Plugins/Editor/GLTFImporter/Source/GLTFImporter/Private/GLTFMaterialElement.cpp
@@ -234,6 +234,16 @@ void FGLTFMaterialElement::SetTwoSided(bool bTwoSided)
 	Material->TwoSided = bTwoSided;
 }
 
+bool FGLTFMaterialElement::GetIsThinSurface() const
+{
+	return Material->IsThinSurface();
+}
+
+void FGLTFMaterialElement::SetIsThinSurface(bool bIsThinSurface)
+{
+	Material->bIsThinSurface = bIsThinSurface;
+}
+
 void FGLTFMaterialElement::SetShadingModel(GLTF::EGLTFMaterialShadingModel InShadingModel)
 {
 	EMaterialShadingModel MaterialShadingModel;
diff --git a/Engine/Plugins/Editor/GLTFImporter/Source/GLTFImporter/Private/GLTFMaterialElement.h b/Engine/Plugins/Editor/GLTFImporter/Source/GLTFImporter/Private/GLTFMaterialElement.h
--- a/Engine/Plugins/Editor/GLTFImporter/Source/GLTFImporter/Private/GLTFMaterialElement.h
+++ b/Engine/Plugins/Editor/GLTFImporter/Source/GLTFImporter/Private/GLTFMaterialElement.h
@@ -29,6 +29,8 @@ public:
 	virtual void SetBlendMode(int InBlendMode) override;
 	virtual bool GetTwoSided() const override;
 	virtual void SetTwoSided(bool bTwoSided) override;
+	virtual bool GetIsThinSurface() const override;
+	virtual void SetIsThinSurface(bool bIsThinSurface) override;
 	virtual void SetShadingModel(GLTF::EGLTFMaterialShadingModel InShadingModel);
 	virtual void SetTranslucencyLightingMode(int InLightingMode) override;
 	virtual void Finalize() override;
diff --git a/Engine/Plugins/Editor/MaterialAnalyzer/Source/Private/SMaterialAnalyzer.cpp b/Engine/Plugins/Editor/MaterialAnalyzer/Source/Private/SMaterialAnalyzer.cpp
--- a/Engine/Plugins/Editor/MaterialAnalyzer/Source/Private/SMaterialAnalyzer.cpp
+++ b/Engine/Plugins/Editor/MaterialAnalyzer/Source/Private/SMaterialAnalyzer.cpp
@@ -35,6 +35,7 @@ SMaterialAnalyzer::SMaterialAnalyzer()
 	BasePropertyOverrideNames.Add(TEXT("bOverride_DitheredLODTransition"), TEXT("DitheredLODTransitionOverride"));
 	BasePropertyOverrideNames.Add(TEXT("bOverride_CastDynamicShadowAsMasked"), TEXT("CastDynamicShadowAsMaskedOverride"));
 	BasePropertyOverrideNames.Add(TEXT("bOverride_TwoSided"), TEXT("TwoSidedOverride"));
+	BasePropertyOverrideNames.Add(TEXT("bOverride_bIsThinSurface"), TEXT("bIsThinSurfaceOverride"));
 	BasePropertyOverrideNames.Add(TEXT("bOverride_OutputTranslucentVelocity"), TEXT("bOutputTranslucentVelocity"));
 }
 
@@ -765,6 +766,14 @@ void FAnalyzeMaterialTreeAsyncTask::DoWork()
 				bIsOverridden = CurrentMaterialInstance->BasePropertyOverrides.bOverride_TwoSided;
 			}
 		}
+		else if (BasePropertyOverrideName.Key.IsEqual(TEXT("bOverride_bIsThinSurface")))
+		{
+			TempValue = CurrentMaterialInterface->IsThinSurface();
+			if (CurrentMaterialInstance)
+			{
+				bIsOverridden = CurrentMaterialInstance->BasePropertyOverrides.bOverride_bIsThinSurface;
+			}
+		}
 		else if (BasePropertyOverrideName.Key.IsEqual(TEXT("bOverride_OutputTranslucentVelocity")))
 		{
 			TempValue = CurrentMaterialInterface->IsTranslucencyWritingVelocity();
diff --git a/Engine/Plugins/Enterprise/DatasmithGLTFImporter/Source/DatasmithGLTFTranslator/Private/DatasmithGLTFMaterialElement.cpp b/Engine/Plugins/Enterprise/DatasmithGLTFImporter/Source/DatasmithGLTFTranslator/Private/DatasmithGLTFMaterialElement.cpp
--- a/Engine/Plugins/Enterprise/DatasmithGLTFImporter/Source/DatasmithGLTFTranslator/Private/DatasmithGLTFMaterialElement.cpp
+++ b/Engine/Plugins/Enterprise/DatasmithGLTFImporter/Source/DatasmithGLTFTranslator/Private/DatasmithGLTFMaterialElement.cpp
@@ -42,6 +42,16 @@ void FDatasmithGLTFMaterialElement::SetTwoSided(bool bTwoSided)
 	MaterialElement->SetTwoSided(bTwoSided);
 }
 
+bool FDatasmithGLTFMaterialElement::GetIsThinSurface() const
+{
+	return MaterialElement->GetIsThinSurface();
+}
+
+void FDatasmithGLTFMaterialElement::SetIsThinSurface(bool bTwoSided)
+{
+	MaterialElement->SetIsThinSurface(bTwoSided);
+}
+
 void FDatasmithGLTFMaterialElement::SetShadingModel(GLTF::EGLTFMaterialShadingModel InShadingModel)
 {
 	EDatasmithShadingModel DatasmithShadingModel;
diff --git a/Engine/Plugins/Enterprise/DatasmithGLTFImporter/Source/DatasmithGLTFTranslator/Private/DatasmithGLTFMaterialElement.h b/Engine/Plugins/Enterprise/DatasmithGLTFImporter/Source/DatasmithGLTFTranslator/Private/DatasmithGLTFMaterialElement.h
--- a/Engine/Plugins/Enterprise/DatasmithGLTFImporter/Source/DatasmithGLTFTranslator/Private/DatasmithGLTFMaterialElement.h
+++ b/Engine/Plugins/Enterprise/DatasmithGLTFImporter/Source/DatasmithGLTFTranslator/Private/DatasmithGLTFMaterialElement.h
@@ -32,6 +32,8 @@ public:
 	virtual void SetBlendMode(int InBlendMode) override;
 	virtual bool GetTwoSided() const override;
 	virtual void SetTwoSided(bool bTwoSided) override;
+	virtual bool GetIsThinSurface() const override;
+	virtual void SetIsThinSurface(bool bIsThinSurface) override;
 	virtual void SetShadingModel(GLTF::EGLTFMaterialShadingModel InShadingModel);
 	virtual void SetTranslucencyLightingMode(int InLightingMode) override;
 	virtual void Finalize() override;
diff --git a/Engine/Plugins/Enterprise/DatasmithImporter/Source/DatasmithImporter/Private/DatasmithMaterialExpressions.cpp b/Engine/Plugins/Enterprise/DatasmithImporter/Source/DatasmithImporter/Private/DatasmithMaterialExpressions.cpp
--- a/Engine/Plugins/Enterprise/DatasmithImporter/Source/DatasmithImporter/Private/DatasmithMaterialExpressions.cpp
+++ b/Engine/Plugins/Enterprise/DatasmithImporter/Source/DatasmithImporter/Private/DatasmithMaterialExpressions.cpp
@@ -3236,6 +3236,7 @@ UMaterialInterface* FDatasmithMaterialExpressions::CreateUEPbrMaterial(UPackage*
 
 	UnrealMaterial->bUseMaterialAttributes = MaterialElement->GetUseMaterialAttributes();
 	UnrealMaterial->TwoSided = MaterialElement->GetTwoSided();
+	UnrealMaterial->bIsThinSurface = MaterialElement->GetIsThinSurface();
 	UnrealMaterial->BlendMode = GetUEPbrImportBlendMode(MaterialElement, AssetsContext);
 	UnrealMaterial->OpacityMaskClipValue = MaterialElement->GetOpacityMaskClipValue();
 	UnrealMaterial->TranslucencyLightingMode = static_cast<ETranslucencyLightingMode>(MaterialElement->GetTranslucencyLightingMode());
diff --git a/Engine/Plugins/Enterprise/DatasmithImporter/Source/DatasmithImporter/Private/DatasmithMaterialImporter.cpp b/Engine/Plugins/Enterprise/DatasmithImporter/Source/DatasmithImporter/Private/DatasmithMaterialImporter.cpp
--- a/Engine/Plugins/Enterprise/DatasmithImporter/Source/DatasmithImporter/Private/DatasmithMaterialImporter.cpp
+++ b/Engine/Plugins/Enterprise/DatasmithImporter/Source/DatasmithImporter/Private/DatasmithMaterialImporter.cpp
@@ -138,6 +138,7 @@ namespace DatasmithMaterialImporterUtils
 
 		Hash = HashCombine( Hash, GetTypeHash( MaterialElement->GetUseMaterialAttributes() ) );
 		Hash = HashCombine( Hash, GetTypeHash( MaterialElement->GetBlendMode() ) );
+		Hash = HashCombine( Hash, GetTypeHash( MaterialElement->GetIsThinSurface() ) );
 		Hash = HashCombine( Hash, GetTypeHash( MaterialElement->GetShadingModel() ) );
 
 		Hash = HashCombine( Hash, ComputeExpressionInputHash( &MaterialElement->GetBaseColor(), VisitedExpressions ) );
diff --git a/Engine/Plugins/Enterprise/GLTFExporter/Source/GLTFExporter/Private/Converters/GLTFMaterialProxyFactory.cpp b/Engine/Plugins/Enterprise/GLTFExporter/Source/GLTFExporter/Private/Converters/GLTFMaterialProxyFactory.cpp
--- a/Engine/Plugins/Enterprise/GLTFExporter/Source/GLTFExporter/Private/Converters/GLTFMaterialProxyFactory.cpp
+++ b/Engine/Plugins/Enterprise/GLTFExporter/Source/GLTFExporter/Private/Converters/GLTFMaterialProxyFactory.cpp
@@ -75,6 +75,13 @@ void FGLTFMaterialProxyFactory::SetBaseProperties(UMaterialInstanceConstant* Pro
 		ProxyMaterial->BasePropertyOverrides.TwoSided = bTwoSided;
 	}
 
+	const bool bIsThinSurface = OriginalMaterial->IsThinSurface();
+	if (bIsThinSurface != BaseMaterial->IsThinSurface())
+	{
+		ProxyMaterial->BasePropertyOverrides.bOverride_bIsThinSurface = true;
+		ProxyMaterial->BasePropertyOverrides.bIsThinSurface = bIsThinSurface;
+	}
+
 	const EBlendMode BlendMode = OriginalMaterial->GetBlendMode();
 	if (BlendMode != BaseMaterial->GetBlendMode())
 	{
diff --git a/Engine/Plugins/Enterprise/GLTFExporter/Source/GLTFExporter/Private/Utilities/GLTFProxyMaterialUtilities.cpp b/Engine/Plugins/Enterprise/GLTFExporter/Source/GLTFExporter/Private/Utilities/GLTFProxyMaterialUtilities.cpp
--- a/Engine/Plugins/Enterprise/GLTFExporter/Source/GLTFExporter/Private/Utilities/GLTFProxyMaterialUtilities.cpp
+++ b/Engine/Plugins/Enterprise/GLTFExporter/Source/GLTFExporter/Private/Utilities/GLTFProxyMaterialUtilities.cpp
@@ -196,6 +196,20 @@ bool FGLTFProxyMaterialUtilities::GetTwoSided(const UMaterialInterface* Material
 	return true;
 }
 
+bool FGLTFProxyMaterialUtilities::GetIsThinSurface(const UMaterialInterface* Material, bool& OutValue, bool NonDefaultOnly)
+{
+	const bool DefaultValue = Material->GetMaterial()->IsThinSurface();
+	const bool Value = Material->IsThinSurface();
+
+	if (NonDefaultOnly && Value == DefaultValue)
+	{
+		return false;
+	}
+
+	OutValue = Value;
+	return true;
+}
+
 bool FGLTFProxyMaterialUtilities::GetBlendMode(const UMaterialInterface* Material, EBlendMode& OutValue, bool NonDefaultOnly)
 {
 	const EBlendMode DefaultValue = Material->GetMaterial()->GetBlendMode();
@@ -241,6 +255,22 @@ void FGLTFProxyMaterialUtilities::SetTwoSided(UMaterialInstanceConstant* Materia
 	Material->TwoSided = Value;
 }
 
+void FGLTFProxyMaterialUtilities::SetIsThinSurface(UMaterialInstanceConstant* Material, bool Value, bool NonDefaultOnly)
+{
+	if (NonDefaultOnly)
+	{
+		const bool DefaultValue = Material->Parent->IsThinSurface();
+		if (DefaultValue == Value)
+		{
+			return;
+		}
+	}
+
+	Material->BasePropertyOverrides.bOverride_bIsThinSurface = true;
+	Material->BasePropertyOverrides.bIsThinSurface = Value;
+	Material->bIsThinSurface = Value;
+}
+
 void FGLTFProxyMaterialUtilities::SetBlendMode(UMaterialInstanceConstant* Material, EBlendMode Value, bool NonDefaultOnly)
 {
 	if (NonDefaultOnly)
diff --git a/Engine/Plugins/Enterprise/GLTFExporter/Source/GLTFExporter/Public/Utilities/GLTFProxyMaterialUtilities.h b/Engine/Plugins/Enterprise/GLTFExporter/Source/GLTFExporter/Public/Utilities/GLTFProxyMaterialUtilities.h
--- a/Engine/Plugins/Enterprise/GLTFExporter/Source/GLTFExporter/Public/Utilities/GLTFProxyMaterialUtilities.h
+++ b/Engine/Plugins/Enterprise/GLTFExporter/Source/GLTFExporter/Public/Utilities/GLTFProxyMaterialUtilities.h
@@ -40,11 +40,13 @@ public:
 #endif
 
 	static bool GetTwoSided(const UMaterialInterface* Material, bool& OutValue, bool NonDefaultOnly = false);
+	static bool GetIsThinSurface(const UMaterialInterface* Material, bool& OutValue, bool NonDefaultOnly = false);
 	static bool GetBlendMode(const UMaterialInterface* Material, EBlendMode& OutValue, bool NonDefaultOnly = false);
 	static bool GetOpacityMaskClipValue(const UMaterialInterface* Material, float& OutValue, bool NonDefaultOnly = false);
 
 #if WITH_EDITOR
 	static void SetTwoSided(UMaterialInstanceConstant* Material, bool Value, bool NonDefaultOnly = false);
+	static void SetIsThinSurface(UMaterialInstanceConstant* Material, bool Value, bool NonDefaultOnly = false);
 	static void SetBlendMode(UMaterialInstanceConstant* Material, EBlendMode Value, bool NonDefaultOnly = false);
 	static void SetOpacityMaskClipValue(UMaterialInstanceConstant* Material, float Value, bool NonDefaultOnly = false);
 #endif
diff --git a/Engine/Plugins/Enterprise/GLTFExporter/Source/GLTFMaterialBaking/Private/GLTFExportMaterialProxy.h b/Engine/Plugins/Enterprise/GLTFExporter/Source/GLTFMaterialBaking/Private/GLTFExportMaterialProxy.h
--- a/Engine/Plugins/Enterprise/GLTFExporter/Source/GLTFMaterialBaking/Private/GLTFExportMaterialProxy.h
+++ b/Engine/Plugins/Enterprise/GLTFExporter/Source/GLTFMaterialBaking/Private/GLTFExportMaterialProxy.h
@@ -449,6 +449,14 @@ public:
 		}
 		return false;
 	}
+	virtual bool IsThinSurface() const  override
+	{
+		if (MaterialInterface)
+		{
+			return MaterialInterface->IsThinSurface();
+		}
+		return false;
+	}
 	virtual bool IsDitheredLODTransition() const  override
 	{
 		if (MaterialInterface)
diff --git a/Engine/Shaders/Private/DebugProbes.usf b/Engine/Shaders/Private/DebugProbes.usf
--- a/Engine/Shaders/Private/DebugProbes.usf
+++ b/Engine/Shaders/Private/DebugProbes.usf
@@ -133,6 +133,7 @@ void WriteOutMaterial(
 		/*float FuzzAmount						*/0.0f,
 		/*float3 FuzzColor						*/0.0f,
 		/*float Thickness						*/0.0f,
+		/*bool bIsThinSurface					*/false,
 		/*uint SharedLocalBasisIndex			*/SharedLocalBasisIndex,
 		/*inout uint SharedLocalBasisTypes		*/StrataPixelHeader.SharedLocalBases.Types);
 
diff --git a/Engine/Shaders/Private/Definitions.usf b/Engine/Shaders/Private/Definitions.usf
--- a/Engine/Shaders/Private/Definitions.usf
+++ b/Engine/Shaders/Private/Definitions.usf
@@ -10,6 +10,10 @@
 #define	MATERIAL_TWOSIDED								0
 #endif
 
+#ifndef MATERIAL_ISTHINSURFACE
+#define	MATERIAL_ISTHINSURFACE							0
+#endif
+
 #ifndef MATERIAL_TANGENTSPACENORMAL
 #define	MATERIAL_TANGENTSPACENORMAL						0
 #endif
diff --git a/Engine/Shaders/Private/Strata/Strata.ush b/Engine/Shaders/Private/Strata/Strata.ush
--- a/Engine/Shaders/Private/Strata/Strata.ush
+++ b/Engine/Shaders/Private/Strata/Strata.ush
@@ -1346,10 +1346,16 @@ FStrataData GetStrataSlabBSDF(
 	float SecondRoughness, float SecondRoughnessWeight, float SecondRoughnessAsSimpleClearCoat,
 	float FuzzAmount, float3 FuzzColor,
 	float Thickness,
+	bool bIsThinSurface,
 	uint SharedLocalBasisIndex, inout uint SharedLocalBasisTypes)
 {
 	FStrataData StrataData = GetInitialisedStrataData();
 
+	if (bIsThinSurface)
+	{
+		Thickness = STRATA_LAYER_ISTHIN_THICKNESS_THRESHOLD_CM - 1e-5f;
+	}
+
 	SSSMFPScale = saturate(SSSMFPScale);
 	const bool bSSSProfilePlugged = SSSProfileID > 0.0f;
 	const bool bMFPPlugged = any(SSSMFP > 0);
@@ -3428,6 +3434,7 @@ void StrataConvertFromDBuffer(
 		0.f,				// FuzzAmount
 		0.f,				// FuzzColor
 		DefaultThickness,
+		false,				// bIsThinSurface,
 		SharedLocalBasisIndex,
 		SharedLocalBasisTypes);
 #if STRATA_INLINE_SHADING
diff --git a/Engine/Shaders/Private/Strata/StrataExport.ush b/Engine/Shaders/Private/Strata/StrataExport.ush
--- a/Engine/Shaders/Private/Strata/StrataExport.ush
+++ b/Engine/Shaders/Private/Strata/StrataExport.ush
@@ -216,7 +216,6 @@ void PackStrataOut(
 						TopLayerData.Material = TOP_LAYER_MATERIAL_VALID;
 
 						const bool bHasSSS = BSDF_GETSSSTYPE(BSDF) != SSS_TYPE_INVALID;
-						const bool bIsThin = BSDF_GETISTHIN(BSDF);
 						const bool bIsSimpleVolume = BSDF_GETISSIMPLEVOLUME(BSDF);
 						bStrataSubsurfaceEnable = bStrataSubsurfaceEnable || (bHasSSS && BSDF.Coverage > 0.0f && !bIsSimpleVolume); // Should it be coverage>0 or any(LuminanceWeight>0) ?
 
diff --git a/Engine/Shaders/Private/Strata/StrataLegacyConversion.ush b/Engine/Shaders/Private/Strata/StrataLegacyConversion.ush
--- a/Engine/Shaders/Private/Strata/StrataLegacyConversion.ush
+++ b/Engine/Shaders/Private/Strata/StrataLegacyConversion.ush
@@ -43,6 +43,7 @@ struct FStrataLegacyParameters
 	float  FuzzAmount;
 	float3 FuzzColor;
 	float  Thickness;
+	bool   bIsThinSurface;
 	uint   SharedLocalBasisIndex;
 
 	float Weight;
@@ -76,6 +77,7 @@ FStrataLegacyParameters InitStrataLegacyParameters(FStrataPixelFootprint InPixel
 	Out.FuzzAmount = 0.f;
 	Out.FuzzColor = Zeros;
 	Out.Thickness = STRATA_LAYER_DEFAULT_THICKNESS_CM;
+	Out.bIsThinSurface = false;
 	Out.SharedLocalBasisIndex = InSharedLocalBasisIndex;
 	Out.Weight = InWeight;
 	Out.PixelFootprint = InPixelFootprint;
@@ -118,6 +120,7 @@ FStrataData CreateLegacySlab(
 		In.FuzzAmount,										// FuzzAmount
 		In.FuzzColor,										// FuzzColor
 		In.Thickness,
+		In.bIsThinSurface,
 		In.SharedLocalBasisIndex, SharedLocalBasisTypes);
 }
 
@@ -356,6 +359,7 @@ FStrataData StrataConvertLegacyMaterialDynamic(
 		LegacySlab.Emissive = Emissive;
 		LegacySlab.UseSSSDiffusion = false;
 		LegacySlab.Thickness = Thickness;
+		LegacySlab.bIsThinSurface = true;
 
 		FinalWeight = 1.0f;
 #if STRATA_TRANSLUCENT_MATERIAL
diff --git a/Engine/Source/Developer/MaterialBaking/Private/ExportMaterialProxy.h b/Engine/Source/Developer/MaterialBaking/Private/ExportMaterialProxy.h
--- a/Engine/Source/Developer/MaterialBaking/Private/ExportMaterialProxy.h
+++ b/Engine/Source/Developer/MaterialBaking/Private/ExportMaterialProxy.h
@@ -459,6 +459,14 @@ public:
 		}
 		return false;
 	}
+	virtual bool IsThinSurface() const  override
+	{
+		if (MaterialInterface)
+		{
+			return MaterialInterface->IsThinSurface();
+		}
+		return false;
+	}
 	virtual bool IsDitheredLODTransition() const  override
 	{
 		if (MaterialInterface)
diff --git a/Engine/Source/Developer/MaterialUtilities/Private/MaterialUtilities.cpp b/Engine/Source/Developer/MaterialUtilities/Private/MaterialUtilities.cpp
--- a/Engine/Source/Developer/MaterialUtilities/Private/MaterialUtilities.cpp
+++ b/Engine/Source/Developer/MaterialUtilities/Private/MaterialUtilities.cpp
@@ -108,6 +108,7 @@ UMaterialInterface* FMaterialUtilities::CreateProxyMaterialAndTextures(UPackage*
 	
 	Material->BasePropertyOverrides.TwoSided = MaterialData.Material->IsTwoSided();
 	Material->BasePropertyOverrides.bOverride_TwoSided = MaterialData.Material->IsTwoSided();
+	Material->BasePropertyOverrides.bOverride_bIsThinSurface = MaterialData.Material->IsThinSurface();
 	Material->BasePropertyOverrides.DitheredLODTransition = MaterialData.Material->IsDitheredLODTransition();
 	Material->BasePropertyOverrides.bOverride_DitheredLODTransition = MaterialData.Material->IsDitheredLODTransition();
 
@@ -648,6 +649,14 @@ public:
 		}
 		return false;
 	}
+	virtual bool IsThinSurface() const  override
+	{
+		if (MaterialInterface)
+		{
+			return MaterialInterface->IsThinSurface();
+		}
+		return false;
+	}
 	virtual bool IsDitheredLODTransition() const  override
 	{ 
 		if (MaterialInterface)
diff --git a/Engine/Source/Developer/MaterialUtilities/Private/ProxyMaterialUtilities.cpp b/Engine/Source/Developer/MaterialUtilities/Private/ProxyMaterialUtilities.cpp
--- a/Engine/Source/Developer/MaterialUtilities/Private/ProxyMaterialUtilities.cpp
+++ b/Engine/Source/Developer/MaterialUtilities/Private/ProxyMaterialUtilities.cpp
@@ -284,6 +284,8 @@ UMaterialInstanceConstant* FMaterialUtilities::CreateFlattenMaterialInstance(UPa
 
 	OutMaterial->BasePropertyOverrides.TwoSided = FlattenMaterial.bTwoSided && InMaterialProxySettings.bAllowTwoSidedMaterial;
 	OutMaterial->BasePropertyOverrides.bOverride_TwoSided = (FlattenMaterial.bTwoSided != false) && InMaterialProxySettings.bAllowTwoSidedMaterial;
+	OutMaterial->BasePropertyOverrides.bIsThinSurface = (FlattenMaterial.bIsThinSurface != false);
+	OutMaterial->BasePropertyOverrides.bOverride_bIsThinSurface = (FlattenMaterial.bIsThinSurface != false);
 	OutMaterial->BasePropertyOverrides.DitheredLODTransition = FlattenMaterial.bDitheredLODTransition;
 	OutMaterial->BasePropertyOverrides.bOverride_DitheredLODTransition = FlattenMaterial.bDitheredLODTransition != false;
 
diff --git a/Engine/Source/Developer/MaterialUtilities/Public/MaterialUtilities.h b/Engine/Source/Developer/MaterialUtilities/Public/MaterialUtilities.h
--- a/Engine/Source/Developer/MaterialUtilities/Public/MaterialUtilities.h
+++ b/Engine/Source/Developer/MaterialUtilities/Public/MaterialUtilities.h
@@ -56,6 +56,7 @@ struct FFlattenMaterial
 	FFlattenMaterial()
 		: RenderSize(0, 0)
 		, bTwoSided(false)
+		, bIsThinSurface(false)
 		, bDitheredLODTransition(false)
 		, BlendMode(BLEND_Opaque)
 		, EmissiveScale(1.0f)
@@ -106,6 +107,8 @@ struct FFlattenMaterial
 
 	/** Flag whether or not the material will have to be two-sided */
 	bool			bTwoSided;
+	/** Flag whether or not the material will have to be thin  */
+	bool			bIsThinSurface;
 	/** Flag whether or not the material will use dithered LOD transitions */
 	bool			bDitheredLODTransition;
 	/** Blend mode for the new material */
diff --git a/Engine/Source/Editor/MaterialEditor/Private/MaterialEditor.h b/Engine/Source/Editor/MaterialEditor/Private/MaterialEditor.h
--- a/Engine/Source/Editor/MaterialEditor/Private/MaterialEditor.h
+++ b/Engine/Source/Editor/MaterialEditor/Private/MaterialEditor.h
@@ -116,6 +116,7 @@ public:
 	virtual EMaterialDomain GetMaterialDomain() const override { return MD_Surface; }
 	virtual FString GetMaterialUsageDescription() const override { return FString::Printf(TEXT("FMatExpressionPreview %s"), Expression.IsValid() ? *Expression->GetName() : TEXT("NULL")); }
 	virtual bool IsTwoSided() const override { return false; }
+	virtual bool IsThinSurface() const override { return false; }
 	virtual bool IsDitheredLODTransition() const override { return false; }
 	virtual bool IsLightFunction() const override { return false; }
 	virtual bool IsDeferredDecal() const override { return false; }
diff --git a/Engine/Source/Editor/MaterialEditor/Private/MaterialEditorInstanceDetailCustomization.cpp b/Engine/Source/Editor/MaterialEditor/Private/MaterialEditorInstanceDetailCustomization.cpp
--- a/Engine/Source/Editor/MaterialEditor/Private/MaterialEditorInstanceDetailCustomization.cpp
+++ b/Engine/Source/Editor/MaterialEditor/Private/MaterialEditorInstanceDetailCustomization.cpp
@@ -1223,6 +1223,7 @@ void FMaterialInstanceParameterDetails::CreateBasePropertyOverrideWidgets(IDetai
 	TAttribute<bool> IsOverrideBlendModeEnabled = TAttribute<bool>::Create(TAttribute<bool>::FGetter::CreateSP(this, &FMaterialInstanceParameterDetails::OverrideBlendModeEnabled));
 	TAttribute<bool> IsOverrideShadingModelEnabled = TAttribute<bool>::Create(TAttribute<bool>::FGetter::CreateSP(this, &FMaterialInstanceParameterDetails::OverrideShadingModelEnabled));
 	TAttribute<bool> IsOverrideTwoSidedEnabled = TAttribute<bool>::Create(TAttribute<bool>::FGetter::CreateSP(this, &FMaterialInstanceParameterDetails::OverrideTwoSidedEnabled));
+	TAttribute<bool> IsOverrideIsThinSurfaceEnabled = TAttribute<bool>::Create(TAttribute<bool>::FGetter::CreateSP(this, &FMaterialInstanceParameterDetails::OverrideIsThinSurfaceEnabled));
 	TAttribute<bool> IsOverrideDitheredLODTransitionEnabled = TAttribute<bool>::Create(TAttribute<bool>::FGetter::CreateSP(this, &FMaterialInstanceParameterDetails::OverrideDitheredLODTransitionEnabled));
 	TAttribute<bool> IsOverrideOutputTranslucentVelocityEnabled = TAttribute<bool>::Create(TAttribute<bool>::FGetter::CreateSP(this, &FMaterialInstanceParameterDetails::OverrideOutputTranslucentVelocityEnabled));
 
@@ -1231,6 +1232,7 @@ void FMaterialInstanceParameterDetails::CreateBasePropertyOverrideWidgets(IDetai
 	TSharedPtr<IPropertyHandle> BlendModeProperty = BasePropertyOverridePropery->GetChildHandle("BlendMode");
 	TSharedPtr<IPropertyHandle> ShadingModelProperty = BasePropertyOverridePropery->GetChildHandle("ShadingModel");
 	TSharedPtr<IPropertyHandle> TwoSidedProperty = BasePropertyOverridePropery->GetChildHandle("TwoSided");
+	TSharedPtr<IPropertyHandle> IsThinSurfaceProperty = BasePropertyOverridePropery->GetChildHandle("IsThinSurface");
 	TSharedPtr<IPropertyHandle> DitheredLODTransitionProperty = BasePropertyOverridePropery->GetChildHandle("DitheredLODTransition");
 	TSharedPtr<IPropertyHandle> OutputTranslucentVelocityProperty = BasePropertyOverridePropery->GetChildHandle("bOutputTranslucentVelocity");
 
@@ -1331,6 +1333,25 @@ void FMaterialInstanceParameterDetails::CreateBasePropertyOverrideWidgets(IDetai
 			.Visibility(TAttribute<EVisibility>::Create(TAttribute<EVisibility>::FGetter::CreateSP(this, &FMaterialInstanceParameterDetails::IsOverriddenAndVisible, IsOverrideTwoSidedEnabled)))
 			.OverrideResetToDefault(ResetTwoSidedPropertyOverride);
 	}
+	{
+		FIsResetToDefaultVisible IsThinSurfacePropertyResetVisible = FIsResetToDefaultVisible::CreateLambda([this](TSharedPtr<IPropertyHandle> InHandle) {
+			return MaterialEditorInstance->Parent != nullptr ? MaterialEditorInstance->BasePropertyOverrides.bIsThinSurface != MaterialEditorInstance->Parent->IsThinSurface() : false;
+			});
+		FResetToDefaultHandler ResetIsThinSurfaceValuePropertyHandler = FResetToDefaultHandler::CreateLambda([this](TSharedPtr<IPropertyHandle> InHandle) {
+			if (MaterialEditorInstance->Parent != nullptr)
+			{
+				MaterialEditorInstance->BasePropertyOverrides.bIsThinSurface = MaterialEditorInstance->Parent->IsThinSurface();
+			}
+			});
+		FResetToDefaultOverride ResetIsThinSurfacePropertyOverride = FResetToDefaultOverride::Create(IsThinSurfacePropertyResetVisible, ResetIsThinSurfaceValuePropertyHandler);
+		IDetailPropertyRow& IsThinSurfacePropertyRow = BasePropertyOverrideGroup.AddPropertyRow(IsThinSurfaceProperty.ToSharedRef());
+		IsThinSurfacePropertyRow
+			.DisplayName(IsThinSurfaceProperty->GetPropertyDisplayName())
+			.ToolTip(IsThinSurfaceProperty->GetToolTipText())
+			.EditCondition(IsOverrideIsThinSurfaceEnabled, FOnBooleanValueChanged::CreateSP(this, &FMaterialInstanceParameterDetails::OnOverrideIsThinSurfaceChanged))
+			.Visibility(TAttribute<EVisibility>::Create(TAttribute<EVisibility>::FGetter::CreateSP(this, &FMaterialInstanceParameterDetails::IsOverriddenAndVisible, IsOverrideIsThinSurfaceEnabled)))
+			.OverrideResetToDefault(ResetIsThinSurfacePropertyOverride);
+	}
 	{
 		FIsResetToDefaultVisible IsDitheredLODTransitionPropertyResetVisible = FIsResetToDefaultVisible::CreateLambda([this](TSharedPtr<IPropertyHandle> InHandle) {
 			return MaterialEditorInstance->Parent != nullptr ? MaterialEditorInstance->BasePropertyOverrides.DitheredLODTransition != MaterialEditorInstance->Parent->IsDitheredLODTransition() : false;
@@ -1401,6 +1422,11 @@ bool FMaterialInstanceParameterDetails::OverrideTwoSidedEnabled() const
 	return MaterialEditorInstance->BasePropertyOverrides.bOverride_TwoSided;
 }
 
+bool FMaterialInstanceParameterDetails::OverrideIsThinSurfaceEnabled() const
+{
+	return MaterialEditorInstance->BasePropertyOverrides.bOverride_bIsThinSurface;
+}
+
 bool FMaterialInstanceParameterDetails::OverrideDitheredLODTransitionEnabled() const
 {
 	return MaterialEditorInstance->BasePropertyOverrides.bOverride_DitheredLODTransition;
@@ -1439,6 +1465,13 @@ void FMaterialInstanceParameterDetails::OnOverrideTwoSidedChanged(bool NewValue)
 	FEditorSupportDelegates::RedrawAllViewports.Broadcast();
 }
 
+void FMaterialInstanceParameterDetails::OnOverrideIsThinSurfaceChanged(bool NewValue)
+{
+	MaterialEditorInstance->BasePropertyOverrides.bOverride_bIsThinSurface = NewValue;
+	MaterialEditorInstance->PostEditChange();
+	FEditorSupportDelegates::RedrawAllViewports.Broadcast();
+}
+
 void FMaterialInstanceParameterDetails::OnOverrideDitheredLODTransitionChanged(bool NewValue)
 {
 	MaterialEditorInstance->BasePropertyOverrides.bOverride_DitheredLODTransition = NewValue;
diff --git a/Engine/Source/Editor/MaterialEditor/Private/MaterialEditorInstanceDetailCustomization.h b/Engine/Source/Editor/MaterialEditor/Private/MaterialEditorInstanceDetailCustomization.h
--- a/Engine/Source/Editor/MaterialEditor/Private/MaterialEditorInstanceDetailCustomization.h
+++ b/Engine/Source/Editor/MaterialEditor/Private/MaterialEditorInstanceDetailCustomization.h
@@ -111,12 +111,14 @@ private:
 	bool OverrideBlendModeEnabled() const;
 	bool OverrideShadingModelEnabled() const;
 	bool OverrideTwoSidedEnabled() const;
+	bool OverrideIsThinSurfaceEnabled() const;
 	bool OverrideDitheredLODTransitionEnabled() const;
 	bool OverrideOutputTranslucentVelocityEnabled() const;
 	void OnOverrideOpacityClipMaskValueChanged(bool NewValue);
 	void OnOverrideBlendModeChanged(bool NewValue);
 	void OnOverrideShadingModelChanged(bool NewValue);
 	void OnOverrideTwoSidedChanged(bool NewValue);
+	void OnOverrideIsThinSurfaceChanged(bool NewValue);
 	void OnOverrideDitheredLODTransitionChanged(bool NewValue);
 	void OnOverrideOutputTranslucentVelocityChanged(bool NewValue);
 
diff --git a/Engine/Source/Editor/UnrealEd/Private/Lightmass/Lightmass.cpp b/Engine/Source/Editor/UnrealEd/Private/Lightmass/Lightmass.cpp
--- a/Engine/Source/Editor/UnrealEd/Private/Lightmass/Lightmass.cpp
+++ b/Engine/Source/Editor/UnrealEd/Private/Lightmass/Lightmass.cpp
@@ -1421,6 +1421,7 @@ void FLightmassExporter::ExportMaterial(UMaterialInterface* Material, const FLig
 		FMemory::Memzero(&MaterialData,sizeof(MaterialData));
 		UMaterial* BaseMaterial = Material->GetMaterial();
 		MaterialData.bTwoSided = (uint32)Material->IsTwoSided();
+		MaterialData.bIsThinSurface = (uint32)Material->IsThinSurface();
 		MaterialData.EmissiveBoost = Material->GetEmissiveBoost();
 		MaterialData.DiffuseBoost = Material->GetDiffuseBoost();
 
diff --git a/Engine/Source/Editor/UnrealEd/Private/Lightmass/LightmassRender.cpp b/Engine/Source/Editor/UnrealEd/Private/Lightmass/LightmassRender.cpp
--- a/Engine/Source/Editor/UnrealEd/Private/Lightmass/LightmassRender.cpp
+++ b/Engine/Source/Editor/UnrealEd/Private/Lightmass/LightmassRender.cpp
@@ -545,6 +545,14 @@ public:
 		}
 		return false;
 	}
+	virtual bool IsThinSurface() const override
+	{
+		if (MaterialInterface)
+		{
+			return MaterialInterface->IsThinSurface();
+		}
+		return false;
+	}
 	virtual bool IsDitheredLODTransition() const override
 	{
 		if (MaterialInterface)
@@ -1106,6 +1114,7 @@ bool FLightmassMaterialRenderer::GenerateMaterialData(
 
 	// Set the two-sided flag
 	OutMaterialData.bTwoSided = (uint32)InMaterial.IsTwoSided();
+	OutMaterialData.bIsThinSurface = (uint32)InMaterial.IsThinSurface();
 	OutMaterialData.OpacityMaskClipValue = InMaterial.GetOpacityMaskClipValue();
 	// Cast shadow as masked feature need to access transmission texture. Only allow
 	// if transmission/opacity data exists
diff --git a/Engine/Source/Editor/UnrealEd/Private/PreviewMaterial.cpp b/Engine/Source/Editor/UnrealEd/Private/PreviewMaterial.cpp
--- a/Engine/Source/Editor/UnrealEd/Private/PreviewMaterial.cpp
+++ b/Engine/Source/Editor/UnrealEd/Private/PreviewMaterial.cpp
@@ -936,6 +936,10 @@ void UMaterialEditorInstanceConstant::CopyBasePropertiesFromParent()
 	{
 		BasePropertyOverrides.TwoSided = SourceInstance->IsTwoSided();
 	}
+	if (!BasePropertyOverrides.bOverride_bIsThinSurface)
+	{
+		BasePropertyOverrides.bIsThinSurface = SourceInstance->IsThinSurface();
+	}
 	if (!BasePropertyOverrides.bOverride_OutputTranslucentVelocity)
 	{
 		BasePropertyOverrides.bOutputTranslucentVelocity = SourceInstance->IsTranslucencyWritingVelocity();
diff --git a/Engine/Source/Programs/UnrealLightmass/Private/Lighting/LightingMesh.cpp b/Engine/Source/Programs/UnrealLightmass/Private/Lighting/LightingMesh.cpp
--- a/Engine/Source/Programs/UnrealLightmass/Private/Lighting/LightingMesh.cpp
+++ b/Engine/Source/Programs/UnrealLightmass/Private/Lighting/LightingMesh.cpp
@@ -149,6 +149,7 @@ void FStaticLightingMesh::Import( FLightmassImporter& Importer )
 
 		CurrentMaterialElement.bIsMasked = bMasked && CurrentMaterialElement.Material->TransmissionSize > 0;
 		CurrentMaterialElement.bIsTwoSided = CurrentMaterialElement.Material->bTwoSided;
+		CurrentMaterialElement.bIsThinSurface = CurrentMaterialElement.Material->bIsThinSurface;
 		CurrentMaterialElement.bTranslucent = !CurrentMaterialElement.bIsMasked && CurrentMaterialElement.Material->TransmissionSize > 0;
 		CurrentMaterialElement.bCastShadowAsMasked = CurrentMaterialElement.Material->bCastShadowAsMasked;
 		CurrentMaterialElement.bSurfaceDomain = CurrentMaterialElement.Material->bSurfaceDomain;
diff --git a/Engine/Source/Programs/UnrealLightmass/Private/Lighting/LightingMesh.h b/Engine/Source/Programs/UnrealLightmass/Private/Lighting/LightingMesh.h
--- a/Engine/Source/Programs/UnrealLightmass/Private/Lighting/LightingMesh.h
+++ b/Engine/Source/Programs/UnrealLightmass/Private/Lighting/LightingMesh.h
@@ -280,6 +280,8 @@ public:
 	 * It just indicates whether backfaces will be visible, and therefore artifacts on backfaces should be avoided.
 	 */
 	bool bIsTwoSided;
+	/** Whether Material is a thin surface or not. */
+	bool bIsThinSurface;
 	/** Whether Material wants to cast shadows as masked, cached here to avoid dereferencing Material. */
 	bool bCastShadowAsMasked;
 	bool bSurfaceDomain;
@@ -290,6 +292,7 @@ public:
 		bTranslucent(false),
 		bIsMasked(false),
 		bIsTwoSided(false),
+		bIsThinSurface(false),
 		bCastShadowAsMasked(false),
 		bSurfaceDomain(true),
 		Material(NULL)
@@ -348,6 +351,7 @@ public:
 	inline bool UsesTwoSidedLighting(int32 ElementIndex) const { return MaterialElements[ElementIndex].bUseTwoSidedLighting; }
 	/** Whether samples using the given element are going to have backfaces visible, and therefore artifacts on backfaces should be avoided. */
 	inline bool IsTwoSided(int32 ElementIndex) const { return MaterialElements[ElementIndex].bIsTwoSided || MaterialElements[ElementIndex].bUseTwoSidedLighting; }
+	inline bool IsThinSurface(int32 ElementIndex) const { return MaterialElements[ElementIndex].bIsThinSurface; }
 	inline bool IsCastingShadowsAsMasked(int32 ElementIndex) const { return MaterialElements[ElementIndex].bCastShadowAsMasked; }
 	inline bool IsSurfaceDomain(int32 ElementIndex) const { return MaterialElements[ElementIndex].bSurfaceDomain; }
 	inline bool IsCastingShadowAsTwoSided() const { return bCastShadowAsTwoSided; }
diff --git a/Engine/Source/Programs/UnrealLightmass/Public/MaterialExport.h b/Engine/Source/Programs/UnrealLightmass/Public/MaterialExport.h
--- a/Engine/Source/Programs/UnrealLightmass/Public/MaterialExport.h
+++ b/Engine/Source/Programs/UnrealLightmass/Public/MaterialExport.h
@@ -72,6 +72,8 @@ struct FMaterialData
 	EStrataBlendMode StrataBlendMode;
 	/** Whether the material is two-sided or not */
 	uint32 bTwoSided:1;
+	/** Whether the material is thin-surface or not */
+	uint32 bIsThinSurface: 1;
 	/** Whether the material should cast shadows as masked even though it has a translucent blend mode. */
 	uint32 bCastShadowAsMasked:1;
 	uint32 bSurfaceDomain:1;
diff --git a/Engine/Source/Runtime/Core/Public/UObject/UE5MainStreamObjectVersion.h b/Engine/Source/Runtime/Core/Public/UObject/UE5MainStreamObjectVersion.h
--- a/Engine/Source/Runtime/Core/Public/UObject/UE5MainStreamObjectVersion.h
+++ b/Engine/Source/Runtime/Core/Public/UObject/UE5MainStreamObjectVersion.h
@@ -309,7 +309,10 @@ struct CORE_API FUE5MainStreamObjectVersion
 
 		// Store the Visual Logger timestamp as a double
 		VisualLoggerTimeStampAsDouble,
-		
+
+		// Add ThinSurface instance override support
+		MaterialInstanceBasePropertyOverridesThinSurface,
+
 		// -----<new versions can be added above this line>-------------------------------------------------
 		VersionPlusOne,
 		LatestVersion = VersionPlusOne - 1
diff --git a/Engine/Source/Runtime/Datasmith/DatasmithCore/Private/DatasmithMaterialElementsImpl.cpp b/Engine/Source/Runtime/Datasmith/DatasmithCore/Private/DatasmithMaterialElementsImpl.cpp
--- a/Engine/Source/Runtime/Datasmith/DatasmithCore/Private/DatasmithMaterialElementsImpl.cpp
+++ b/Engine/Source/Runtime/Datasmith/DatasmithCore/Private/DatasmithMaterialElementsImpl.cpp
@@ -207,6 +207,7 @@ FDatasmithUEPbrMaterialElementImpl::FDatasmithUEPbrMaterialElementImpl( const TC
 	, MaterialAttributes( MakeShared< FDatasmithExpressionInputImpl >( TEXT("MaterialAttributes") ) )
 	, BlendMode(0)
 	, bTwoSided( false )
+	, bIsThinSurface(false)
 	, bUseMaterialAttributes( false )
 	, bMaterialFunctionOnly ( false )
 	, OpacityMaskClipValue( 0.3333f )
@@ -232,6 +233,7 @@ FDatasmithUEPbrMaterialElementImpl::FDatasmithUEPbrMaterialElementImpl( const TC
 
 	Store.RegisterParameter( BlendMode, "BlendMode" );
 	Store.RegisterParameter( bTwoSided, "bTwoSided" );
+	Store.RegisterParameter( bIsThinSurface, "bIsThinSurface");
 	Store.RegisterParameter( bUseMaterialAttributes, "bUseMaterialAttributes" );
 	Store.RegisterParameter( bMaterialFunctionOnly, "bMaterialFunctionOnly" );
 	Store.RegisterParameter( OpacityMaskClipValue, "OpacityMaskClipValue" );
@@ -250,6 +252,7 @@ FMD5Hash FDatasmithUEPbrMaterialElementImpl::CalculateElementHash(bool bForce)
 	FMD5 MD5;
 	MD5.Update(reinterpret_cast<const uint8*>(&BlendMode), sizeof(BlendMode));
 	MD5.Update(reinterpret_cast<const uint8*>(&bTwoSided), sizeof(bTwoSided));
+	MD5.Update(reinterpret_cast<const uint8*>(&bIsThinSurface), sizeof(bIsThinSurface));
 	MD5.Update(reinterpret_cast<const uint8*>(&bUseMaterialAttributes), sizeof(bUseMaterialAttributes));
 	MD5.Update(reinterpret_cast<const uint8*>(&bMaterialFunctionOnly), sizeof(bMaterialFunctionOnly));
 	MD5.Update(reinterpret_cast<const uint8*>(&OpacityMaskClipValue), sizeof(OpacityMaskClipValue));
diff --git a/Engine/Source/Runtime/Datasmith/DatasmithCore/Private/DatasmithMaterialElementsImpl.h b/Engine/Source/Runtime/Datasmith/DatasmithCore/Private/DatasmithMaterialElementsImpl.h
--- a/Engine/Source/Runtime/Datasmith/DatasmithCore/Private/DatasmithMaterialElementsImpl.h
+++ b/Engine/Source/Runtime/Datasmith/DatasmithCore/Private/DatasmithMaterialElementsImpl.h
@@ -685,6 +685,9 @@ public:
 	virtual bool GetTwoSided() const override {return bTwoSided; }
 	virtual void SetTwoSided( bool bInTwoSided ) override { bTwoSided = bInTwoSided; }
 
+	virtual bool GetIsThinSurface() const override { return bIsThinSurface; }
+	virtual void SetIsThinSurface(bool bInIsThinSurface) override { bIsThinSurface = bInIsThinSurface; }
+
 	virtual bool GetUseMaterialAttributes() const override{ return bUseMaterialAttributes; }
 	virtual void SetUseMaterialAttributes( bool bInUseMaterialAttributes ) override { bUseMaterialAttributes = bInUseMaterialAttributes; }
 
@@ -733,6 +736,7 @@ protected:
 
 	TReflected<int32> BlendMode;
 	TReflected<bool> bTwoSided;
+	TReflected<bool> bIsThinSurface;
 	TReflected<bool> bUseMaterialAttributes;
 	TReflected<bool> bMaterialFunctionOnly;
 	TReflected<float> OpacityMaskClipValue;
diff --git a/Engine/Source/Runtime/Datasmith/DatasmithCore/Public/DatasmithMaterialElements.h b/Engine/Source/Runtime/Datasmith/DatasmithCore/Public/DatasmithMaterialElements.h
--- a/Engine/Source/Runtime/Datasmith/DatasmithCore/Public/DatasmithMaterialElements.h
+++ b/Engine/Source/Runtime/Datasmith/DatasmithCore/Public/DatasmithMaterialElements.h
@@ -233,6 +233,9 @@ public:
 	virtual bool GetTwoSided() const = 0;
 	virtual void SetTwoSided( bool bTwoSided ) = 0;
 
+	virtual bool GetIsThinSurface() const = 0;
+	virtual void SetIsThinSurface(bool bIsThinSurface) = 0;
+
 	virtual bool GetUseMaterialAttributes() const = 0;
 	virtual void SetUseMaterialAttributes( bool bInUseMaterialAttributes ) = 0;
 
diff --git a/Engine/Source/Runtime/Engine/Classes/Materials/Material.h b/Engine/Source/Runtime/Engine/Classes/Materials/Material.h
--- a/Engine/Source/Runtime/Engine/Classes/Materials/Material.h
+++ b/Engine/Source/Runtime/Engine/Classes/Materials/Material.h
@@ -512,6 +512,10 @@ public:
 	UPROPERTY(EditAnywhere, Category=Material)
 	uint8 TwoSided : 1;
 
+	/** Indicates that the material should be rendered as a thin surface (i.e., without inner volume). */
+	UPROPERTY(EditAnywhere, Category = Material)
+	uint8 bIsThinSurface : 1;
+
 	/** Whether meshes rendered with the material should support dithered LOD transitions. */
 	UPROPERTY(EditAnywhere, Category = Material, AdvancedDisplay, meta = (DisplayName = "Dithered LOD Transition"))
 	uint8 DitheredLODTransition : 1;
@@ -1063,6 +1067,7 @@ public:
 	ENGINE_API virtual FMaterialShadingModelField GetShadingModels() const override;
 	ENGINE_API virtual bool IsShadingModelFromMaterialExpression() const override;
 	ENGINE_API virtual bool IsTwoSided() const override;
+	ENGINE_API virtual bool IsThinSurface() const override;
 	ENGINE_API virtual bool IsDitheredLODTransition() const override;
 	ENGINE_API virtual bool IsTranslucencyWritingCustomDepth() const override;
 	ENGINE_API virtual bool IsTranslucencyWritingVelocity() const override;
diff --git a/Engine/Source/Runtime/Engine/Classes/Materials/MaterialInstance.h b/Engine/Source/Runtime/Engine/Classes/Materials/MaterialInstance.h
--- a/Engine/Source/Runtime/Engine/Classes/Materials/MaterialInstance.h
+++ b/Engine/Source/Runtime/Engine/Classes/Materials/MaterialInstance.h
@@ -603,6 +603,7 @@ class UMaterialInstance : public UMaterialInterface
 	uint8 bOverrideSubsurfaceProfile:1;
 
 	uint8 TwoSided : 1;
+	uint8 bIsThinSurface : 1;
 	uint8 DitheredLODTransition : 1;
 	uint8 bCastDynamicShadowAsMasked : 1;
 	uint8 bOutputTranslucentVelocity : 1;
@@ -800,6 +801,7 @@ public:
 	ENGINE_API virtual FMaterialShadingModelField GetShadingModels() const override;
 	ENGINE_API virtual bool IsShadingModelFromMaterialExpression() const override;
 	ENGINE_API virtual bool IsTwoSided() const override;
+	ENGINE_API virtual bool IsThinSurface() const override;
 	ENGINE_API virtual bool IsTranslucencyWritingVelocity() const override;
 	ENGINE_API virtual bool IsDitheredLODTransition() const override;
 	ENGINE_API virtual bool IsMasked() const override;
diff --git a/Engine/Source/Runtime/Engine/Classes/Materials/MaterialInstanceBasePropertyOverrides.h b/Engine/Source/Runtime/Engine/Classes/Materials/MaterialInstanceBasePropertyOverrides.h
--- a/Engine/Source/Runtime/Engine/Classes/Materials/MaterialInstanceBasePropertyOverrides.h
+++ b/Engine/Source/Runtime/Engine/Classes/Materials/MaterialInstanceBasePropertyOverrides.h
@@ -37,6 +37,10 @@ struct ENGINE_API FMaterialInstanceBasePropertyOverrides
 	UPROPERTY(EditAnywhere, Category = Material)
 	uint8 bOverride_TwoSided : 1;
 
+	/** Enables override of the IsThinSurface property. */
+	UPROPERTY(EditAnywhere, Category = Material)
+	uint8 bOverride_bIsThinSurface : 1;
+
 	/** Enables override of the output velocity property. */
 	UPROPERTY(EditAnywhere, Category = Material)
 	uint8 bOverride_OutputTranslucentVelocity : 1;
@@ -45,6 +49,10 @@ struct ENGINE_API FMaterialInstanceBasePropertyOverrides
 	UPROPERTY(EditAnywhere, Category = Material, meta = (editcondition = "bOverride_TwoSided"))
 	uint8 TwoSided : 1;
 
+	/** Indicates that the material should be rendered as. */
+	UPROPERTY(EditAnywhere, Category = Material, meta = (editcondition = "bOverride_bThinSurface"))
+	uint8 bIsThinSurface : 1;
+
 	/** Whether the material should support a dithered LOD transition when used with the foliage system. */
 	UPROPERTY(EditAnywhere, Category = Material, meta = (editcondition = "bOverride_DitheredLODTransition"))
 	uint8 DitheredLODTransition : 1;
diff --git a/Engine/Source/Runtime/Engine/Classes/Materials/MaterialInstanceDynamic.h b/Engine/Source/Runtime/Engine/Classes/Materials/MaterialInstanceDynamic.h
--- a/Engine/Source/Runtime/Engine/Classes/Materials/MaterialInstanceDynamic.h
+++ b/Engine/Source/Runtime/Engine/Classes/Materials/MaterialInstanceDynamic.h
@@ -198,6 +198,7 @@ class ENGINE_API UMaterialInstanceDynamic : public UMaterialInstance
 	virtual EBlendMode GetBlendMode() const override;
 	virtual EStrataBlendMode GetStrataBlendMode() const override;
 	virtual bool IsTwoSided() const override;
+	virtual bool IsThinSurface() const override;
 	virtual bool IsTranslucencyWritingVelocity() const override;
 	virtual bool IsDitheredLODTransition() const override;
 	virtual bool IsMasked() const override;
diff --git a/Engine/Source/Runtime/Engine/Classes/Materials/MaterialInterface.h b/Engine/Source/Runtime/Engine/Classes/Materials/MaterialInterface.h
--- a/Engine/Source/Runtime/Engine/Classes/Materials/MaterialInterface.h
+++ b/Engine/Source/Runtime/Engine/Classes/Materials/MaterialInterface.h
@@ -838,6 +838,7 @@ public:
 	ENGINE_API virtual FMaterialShadingModelField GetShadingModels() const;
 	ENGINE_API virtual bool IsShadingModelFromMaterialExpression() const;
 	ENGINE_API virtual bool IsTwoSided() const;
+	ENGINE_API virtual bool IsThinSurface() const;
 	ENGINE_API virtual bool IsDitheredLODTransition() const;
 	ENGINE_API virtual bool IsTranslucencyWritingCustomDepth() const;
 	ENGINE_API virtual bool IsTranslucencyWritingVelocity() const;
diff --git a/Engine/Source/Runtime/Engine/Private/Materials/HLSLMaterialTranslator.cpp b/Engine/Source/Runtime/Engine/Private/Materials/HLSLMaterialTranslator.cpp
--- a/Engine/Source/Runtime/Engine/Private/Materials/HLSLMaterialTranslator.cpp
+++ b/Engine/Source/Runtime/Engine/Private/Materials/HLSLMaterialTranslator.cpp
@@ -10700,6 +10700,7 @@ int32 FHLSLMaterialTranslator::StrataSlabBSDF(
 {
 	const FString NormalCode = GetParameterCode(Normal);
 	const FString TangentCode = Tangent != INDEX_NONE ? *GetParameterCode(Tangent) : TEXT("NONE");
+	const bool bIsThinSurface = Material->IsThinSurface();
 
 	if (PromoteToOperator)
 	{
@@ -10709,7 +10710,7 @@ int32 FHLSLMaterialTranslator::StrataSlabBSDF(
 			return INDEX_NONE;
 		}
 		return AddCodeChunk(
-			MCT_Strata, TEXT("PromoteParameterBlendedBSDFToOperator(GetStrataSlabBSDF(Parameters.StrataPixelFootprint, %s, %s, %s, %s, %s, %s, %s, %s, %s, %s, %s, %s, %s, %s, %s, %s, %s, %s, Parameters.SharedLocalBases.Types) /* Normal = %s ; Tangent = %s */, Parameters.StrataTree, %u, %u, %u, %u)"),
+			MCT_Strata, TEXT("PromoteParameterBlendedBSDFToOperator(GetStrataSlabBSDF(Parameters.StrataPixelFootprint, %s, %s, %s, %s, %s, %s, %s, %s, %s, %s, %s, %s, %s, %s, %s, %s, %s, %s, %s, Parameters.SharedLocalBases.Types) /* Normal = %s ; Tangent = %s */, Parameters.StrataTree, %u, %u, %u, %u)"),
 			*StrataGetCastParameterCode(DiffuseAlbedo,			MCT_Float3),
 			*StrataGetCastParameterCode(F0,						MCT_Float3),
 			*StrataGetCastParameterCode(F90,					MCT_Float3),
@@ -10727,6 +10728,7 @@ int32 FHLSLMaterialTranslator::StrataSlabBSDF(
 			*StrataGetCastParameterCode(FuzzAmount,				MCT_Float),
 			*StrataGetCastParameterCode(FuzzColor,				MCT_Float3),
 			*StrataGetCastParameterCode(Thickness,				MCT_Float),
+			bIsThinSurface ? TEXT("true") : TEXT("false"),
 			*SharedLocalBasisIndexMacro,
 			*NormalCode,
 			*TangentCode,
@@ -10738,7 +10740,7 @@ int32 FHLSLMaterialTranslator::StrataSlabBSDF(
 	}
 	
 	return AddCodeChunk(
-		MCT_Strata, TEXT("GetStrataSlabBSDF(Parameters.StrataPixelFootprint, %s, %s, %s, %s, %s, %s, %s, %s, %s, %s, %s, %s, %s, %s, %s, %s, %s, %s, Parameters.SharedLocalBases.Types) /* Normal = %s ; Tangent = %s */"),
+		MCT_Strata, TEXT("GetStrataSlabBSDF(Parameters.StrataPixelFootprint, %s, %s, %s, %s, %s, %s, %s, %s, %s, %s, %s, %s, %s, %s, %s, %s, %s, %s, %s, Parameters.SharedLocalBases.Types) /* Normal = %s ; Tangent = %s */"),
 		*StrataGetCastParameterCode(DiffuseAlbedo,			MCT_Float3),
 		*StrataGetCastParameterCode(F0,						MCT_Float3),
 		*StrataGetCastParameterCode(F90,					MCT_Float3),
@@ -10756,6 +10758,7 @@ int32 FHLSLMaterialTranslator::StrataSlabBSDF(
 		*StrataGetCastParameterCode(FuzzAmount,				MCT_Float),
 		*StrataGetCastParameterCode(FuzzColor,				MCT_Float3),
 		*StrataGetCastParameterCode(Thickness,				MCT_Float),
+		bIsThinSurface ? TEXT("true") : TEXT("false"),
 		*SharedLocalBasisIndexMacro,
 		*NormalCode,
 		*TangentCode
diff --git a/Engine/Source/Runtime/Engine/Private/Materials/Material.cpp b/Engine/Source/Runtime/Engine/Private/Materials/Material.cpp
--- a/Engine/Source/Runtime/Engine/Private/Materials/Material.cpp
+++ b/Engine/Source/Runtime/Engine/Private/Materials/Material.cpp
@@ -4278,6 +4278,7 @@ bool UMaterial::CanEditChange(const FProperty* InProperty) const
 		if (PropertyName == GET_MEMBER_NAME_STRING_CHECKED(UMaterial, bFullyRough) ||
 			PropertyName == GET_MEMBER_NAME_STRING_CHECKED(UMaterial, bNormalCurvatureToRoughness) ||
 			PropertyName == GET_MEMBER_NAME_STRING_CHECKED(UMaterial, TwoSided) ||
+			PropertyName == GET_MEMBER_NAME_STRING_CHECKED(UMaterial, bIsThinSurface) ||
 			PropertyName == GET_MEMBER_NAME_STRING_CHECKED(UMaterial, bUseLightmapDirectionality) ||
 			PropertyName == GET_MEMBER_NAME_STRING_CHECKED(UMaterial, bUseHQForwardReflections) ||
 			PropertyName == GET_MEMBER_NAME_STRING_CHECKED(UMaterial, bForwardBlendsSkyLightCubemaps) ||
@@ -6385,6 +6386,11 @@ bool UMaterial::IsTwoSided() const
 	return TwoSided != 0;
 }
 
+bool UMaterial::IsThinSurface() const
+{
+	return bIsThinSurface != 0;
+}
+
 bool UMaterial::IsDitheredLODTransition() const
 {
 	return DitheredLODTransition != 0;
diff --git a/Engine/Source/Runtime/Engine/Private/Materials/MaterialInstance.cpp b/Engine/Source/Runtime/Engine/Private/Materials/MaterialInstance.cpp
--- a/Engine/Source/Runtime/Engine/Private/Materials/MaterialInstance.cpp
+++ b/Engine/Source/Runtime/Engine/Private/Materials/MaterialInstance.cpp
@@ -1651,6 +1651,11 @@ bool UMaterialInstanceDynamic::IsTwoSided() const
 	return Parent ? Parent->IsTwoSided() : false;
 }
 
+bool UMaterialInstanceDynamic::IsThinSurface() const
+{
+	return Parent ? Parent->IsThinSurface() : false;
+}
+
 bool UMaterialInstanceDynamic::IsTranslucencyWritingVelocity() const
 {
 	return Parent ? Parent->IsTranslucencyWritingVelocity() : false;
@@ -2031,6 +2036,7 @@ void UMaterialInstance::UpdateOverridableBaseProperties()
 		BlendMode = BLEND_Opaque;
 		ShadingModels = MSM_DefaultLit;
 		TwoSided = 0;
+		bIsThinSurface = false;
 		DitheredLODTransition = 0;
 		bIsShadingModelFromMaterialExpression = 0;
 		bOutputTranslucentVelocity = false;
@@ -2124,6 +2130,16 @@ void UMaterialInstance::UpdateOverridableBaseProperties()
 		BasePropertyOverrides.TwoSided = TwoSided;
 	}
 
+	if (BasePropertyOverrides.bOverride_bIsThinSurface)
+	{
+		bIsThinSurface = BasePropertyOverrides.bIsThinSurface != 0;
+	}
+	else
+	{
+		bIsThinSurface = Parent->IsThinSurface();
+		BasePropertyOverrides.bIsThinSurface = bIsThinSurface;
+	}
+
 	if (BasePropertyOverrides.bOverride_DitheredLODTransition)
 	{
 		DitheredLODTransition = BasePropertyOverrides.DitheredLODTransition != 0;
@@ -2784,6 +2800,11 @@ void UMaterialInstance::Serialize(FArchive& Ar)
 					FArchive_Serialize_BitfieldBool(Ar, BasePropertyOverrides.bOverride_TwoSided);
 					FArchive_Serialize_BitfieldBool(Ar, BasePropertyOverrides.TwoSided);
 
+					if (Ar.CustomVer(FUE5MainStreamObjectVersion::GUID) >= FUE5MainStreamObjectVersion::MaterialInstanceBasePropertyOverridesThinSurface)
+					{
+						FArchive_Serialize_BitfieldBool(Ar, BasePropertyOverrides.bOverride_bIsThinSurface);
+						FArchive_Serialize_BitfieldBool(Ar, BasePropertyOverrides.bIsThinSurface);
+					}
 					if( Ar.UEVer() >= VER_UE4_MATERIAL_INSTANCE_BASE_PROPERTY_OVERRIDES_DITHERED_LOD_TRANSITION )
 					{
 						FArchive_Serialize_BitfieldBool(Ar, BasePropertyOverrides.bOverride_DitheredLODTransition);
@@ -4155,6 +4176,14 @@ void UMaterialInstance::GetBasePropertyOverridesHash(FSHAHash& OutHash)const
 		Hash.Update((uint8*)&bUsedIsTwoSided, sizeof(bUsedIsTwoSided));
 		bHasOverrides = true;
 	}
+	bool bUsedIsThinSurface = IsThinSurface();
+	if (bUsedIsThinSurface != Mat->IsThinSurface())
+	{
+		const FString HashString = TEXT("bOverride_bIsThinSurface");
+		Hash.UpdateWithString(*HashString, HashString.Len());
+		Hash.Update((uint8*)&bUsedIsThinSurface, sizeof(bUsedIsThinSurface));
+		bHasOverrides = true;
+	}
 	bool bUsedIsDitheredLODTransition = IsDitheredLODTransition();
 	if (bUsedIsDitheredLODTransition != Mat->IsDitheredLODTransition())
 	{
@@ -4188,6 +4217,7 @@ bool UMaterialInstance::HasOverridenBaseProperties()const
 		(GetBlendMode() != Parent->GetBlendMode()) ||
 		(GetShadingModels() != Parent->GetShadingModels()) ||
 		(IsTwoSided() != Parent->IsTwoSided()) ||
+		(IsThinSurface() != Parent->IsThinSurface()) ||
 		(IsDitheredLODTransition() != Parent->IsDitheredLODTransition()) ||
 		(GetCastDynamicShadowAsMasked() != Parent->GetCastDynamicShadowAsMasked()) ||
 		(IsTranslucencyWritingVelocity() != Parent->IsTranslucencyWritingVelocity())
@@ -4209,6 +4239,7 @@ FString UMaterialInstance::GetBasePropertyOverrideString() const
 		BasePropString += FString::Printf(TEXT("bOverride_BlendMode_%d, "), (GetBlendMode() != Parent->GetBlendMode()));
 		BasePropString += FString::Printf(TEXT("bOverride_ShadingModel_%d, "), (GetShadingModels() != Parent->GetShadingModels()));
 		BasePropString += FString::Printf(TEXT("bOverride_TwoSided_%d, "), (IsTwoSided() != Parent->IsTwoSided()));
+		BasePropString += FString::Printf(TEXT("bOverride_bIsThinSurface_%d, "), (IsThinSurface() != Parent->IsThinSurface()));
 		BasePropString += FString::Printf(TEXT("bOverride_DitheredLODTransition_%d, "), (IsDitheredLODTransition() != Parent->IsDitheredLODTransition()));
 		BasePropString += FString::Printf(TEXT("bOverride_CastDynamicShadowAsMasked_%d, "), (GetCastDynamicShadowAsMasked() != Parent->GetCastDynamicShadowAsMasked()));
 		BasePropString += FString::Printf(TEXT("bOverride_OutputTranslucentVelocity_%d "), (IsTranslucencyWritingVelocity() != Parent->IsTranslucencyWritingVelocity()));
@@ -4252,6 +4283,11 @@ bool UMaterialInstance::IsTwoSided() const
 	return TwoSided;
 }
 
+bool UMaterialInstance::IsThinSurface() const
+{
+	return bIsThinSurface;
+}
+
 bool UMaterialInstance::IsTranslucencyWritingVelocity() const
 {
 	return bOutputTranslucentVelocity && IsTranslucentBlendMode(GetBlendMode());
diff --git a/Engine/Source/Runtime/Engine/Private/Materials/MaterialInterface.cpp b/Engine/Source/Runtime/Engine/Private/Materials/MaterialInterface.cpp
--- a/Engine/Source/Runtime/Engine/Private/Materials/MaterialInterface.cpp
+++ b/Engine/Source/Runtime/Engine/Private/Materials/MaterialInterface.cpp
@@ -1156,6 +1156,11 @@ bool UMaterialInterface::IsTwoSided() const
 	return false;
 }
 
+bool UMaterialInterface::IsThinSurface() const
+{
+	return false;
+}
+
 bool UMaterialInterface::IsDitheredLODTransition() const
 {
 	return false;
diff --git a/Engine/Source/Runtime/Engine/Private/Materials/MaterialShared.cpp b/Engine/Source/Runtime/Engine/Private/Materials/MaterialShared.cpp
--- a/Engine/Source/Runtime/Engine/Private/Materials/MaterialShared.cpp
+++ b/Engine/Source/Runtime/Engine/Private/Materials/MaterialShared.cpp
@@ -1776,6 +1776,11 @@ bool FMaterialResource::IsTwoSided() const
 	return MaterialInstance ? MaterialInstance->IsTwoSided() : Material->IsTwoSided();
 }
 
+bool FMaterialResource::IsThinSurface() const
+{
+	return MaterialInstance ? MaterialInstance->IsThinSurface() : Material->IsThinSurface();
+}
+
 bool FMaterialResource::IsDitheredLODTransition() const 
 {
 	if (!AllowDitheredLODTransition(GetFeatureLevel()))
@@ -2335,6 +2340,7 @@ void FMaterial::SetupMaterialEnvironment(
 
 	OutEnvironment.SetDefine(TEXT("USE_DITHERED_LOD_TRANSITION_FROM_MATERIAL"), IsDitheredLODTransition());
 	OutEnvironment.SetDefine(TEXT("MATERIAL_TWOSIDED"), IsTwoSided());
+	OutEnvironment.SetDefine(TEXT("MATERIAL_ISTHINSURFACE"), IsThinSurface());
 	OutEnvironment.SetDefine(TEXT("MATERIAL_TANGENTSPACENORMAL"), IsTangentSpaceNormal());
 	OutEnvironment.SetDefine(TEXT("GENERATE_SPHERICAL_PARTICLE_NORMALS"),ShouldGenerateSphericalParticleNormals());
 	OutEnvironment.SetDefine(TEXT("MATERIAL_USES_SCENE_COLOR_COPY"), RequiresSceneColorCopy_GameThread());
@@ -5224,8 +5230,10 @@ FMaterialInstanceBasePropertyOverrides::FMaterialInstanceBasePropertyOverrides()
 	,bOverride_DitheredLODTransition(false)
 	,bOverride_CastDynamicShadowAsMasked(false)
 	,bOverride_TwoSided(false)
+	,bOverride_bIsThinSurface(false)
 	,bOverride_OutputTranslucentVelocity(false)
 	,TwoSided(0)
+	,bIsThinSurface(false)
 	,DitheredLODTransition(0)
 	,bCastDynamicShadowAsMasked(false)
 	,bOutputTranslucentVelocity(false)
@@ -5242,12 +5250,14 @@ bool FMaterialInstanceBasePropertyOverrides::operator==(const FMaterialInstanceB
 		bOverride_BlendMode == Other.bOverride_BlendMode &&
 		bOverride_ShadingModel == Other.bOverride_ShadingModel &&
 		bOverride_TwoSided == Other.bOverride_TwoSided &&
+		bOverride_bIsThinSurface == Other.bOverride_bIsThinSurface &&
 		bOverride_DitheredLODTransition == Other.bOverride_DitheredLODTransition &&
 		bOverride_OutputTranslucentVelocity == Other.bOverride_OutputTranslucentVelocity &&
 		OpacityMaskClipValue == Other.OpacityMaskClipValue &&
 		BlendMode == Other.BlendMode &&
 		ShadingModel == Other.ShadingModel &&
 		TwoSided == Other.TwoSided &&
+		bIsThinSurface == Other.bIsThinSurface &&
 		DitheredLODTransition == Other.DitheredLODTransition &&
 		bCastDynamicShadowAsMasked == Other.bCastDynamicShadowAsMasked;
 }
@@ -6105,6 +6115,7 @@ FMaterialShaderParameters::FMaterialShaderParameters(const FMaterial* InMaterial
 	bIsSpecialEngineMaterial = InMaterial->IsSpecialEngineMaterial();
 	bIsMasked = InMaterial->IsMasked();
 	bIsTwoSided = InMaterial->IsTwoSided();
+	bIsThinSurface = InMaterial->IsThinSurface();
 	bIsDistorted = InMaterial->IsDistorted();
 	bShouldCastDynamicShadows = InMaterial->ShouldCastDynamicShadows();
 	bWritesEveryPixel = InMaterial->WritesEveryPixel(false);
diff --git a/Engine/Source/Runtime/Engine/Private/ShaderCompiler/ShaderGenerationUtil.cpp b/Engine/Source/Runtime/Engine/Private/ShaderCompiler/ShaderGenerationUtil.cpp
--- a/Engine/Source/Runtime/Engine/Private/ShaderCompiler/ShaderGenerationUtil.cpp
+++ b/Engine/Source/Runtime/Engine/Private/ShaderCompiler/ShaderGenerationUtil.cpp
@@ -185,6 +185,7 @@ void FShaderCompileUtilities::ApplyFetchEnvironment(FShaderMaterialPropertyDefin
 
 	FETCH_COMPILE_BOOL(USE_DITHERED_LOD_TRANSITION_FROM_MATERIAL);
 	FETCH_COMPILE_BOOL(MATERIAL_TWOSIDED);
+	FETCH_COMPILE_BOOL(MATERIAL_ISTHINSURFACE);
 	FETCH_COMPILE_BOOL(MATERIAL_TANGENTSPACENORMAL);
 	FETCH_COMPILE_BOOL(GENERATE_SPHERICAL_PARTICLE_NORMALS);
 	FETCH_COMPILE_BOOL(MATERIAL_USE_PREINTEGRATED_GF);
diff --git a/Engine/Source/Runtime/Engine/Private/StaticMeshRender.cpp b/Engine/Source/Runtime/Engine/Private/StaticMeshRender.cpp
--- a/Engine/Source/Runtime/Engine/Private/StaticMeshRender.cpp
+++ b/Engine/Source/Runtime/Engine/Private/StaticMeshRender.cpp
@@ -1262,6 +1262,7 @@ void FStaticMeshSceneProxy::DrawStaticElements(FStaticPrimitiveDrawInterface* PD
 							!(bAnySectionUsesDitheredLODTransition && !bAllSectionsUseDitheredLODTransition) // can't use a single section if they are not homogeneous
 							&& Material.WritesEveryPixel()
 							&& !Material.IsTwoSided()
+							&& !Material.IsThinSurface()
 							&& !IsTranslucentBlendMode(Material.GetBlendMode())
 							&& !Material.MaterialModifiesMeshPosition_RenderThread()
 							&& Material.GetMaterialDomain() == MD_Surface
diff --git a/Engine/Source/Runtime/Engine/Public/MaterialShared.h b/Engine/Source/Runtime/Engine/Public/MaterialShared.h
--- a/Engine/Source/Runtime/Engine/Public/MaterialShared.h
+++ b/Engine/Source/Runtime/Engine/Public/MaterialShared.h
@@ -1781,6 +1781,7 @@ public:
 #endif // WITH_EDITORONLY_DATA
 	virtual EMaterialDomain GetMaterialDomain() const = 0; // See EMaterialDomain.
 	virtual bool IsTwoSided() const = 0;
+	virtual bool IsThinSurface() const = 0;
 	virtual bool IsDitheredLODTransition() const = 0;
 	virtual bool IsTranslucencyWritingCustomDepth() const { return false; }
 	virtual bool IsTranslucencyWritingVelocity() const { return false; }
@@ -2832,6 +2833,7 @@ public:
 #endif // WITH_EDITOR
 	ENGINE_API virtual EMaterialDomain GetMaterialDomain() const override;
 	ENGINE_API virtual bool IsTwoSided() const override;
+	ENGINE_API virtual bool IsThinSurface() const override;
 	ENGINE_API virtual bool IsDitheredLODTransition() const override;
 	ENGINE_API virtual bool IsTranslucencyWritingCustomDepth() const override;
 	ENGINE_API virtual bool IsTranslucencyWritingVelocity() const override;
@@ -3478,6 +3480,7 @@ struct FMaterialShaderParameters
 			uint64 bIsSpecialEngineMaterial : 1;
 			uint64 bIsMasked : 1;
 			uint64 bIsTwoSided : 1;
+			uint64 bIsThinSurface : 1;
 			uint64 bIsDistorted : 1;
 			uint64 bShouldCastDynamicShadows : 1;
 			uint64 bWritesEveryPixel : 1;
diff --git a/Engine/Source/Runtime/RenderCore/Public/ShaderMaterial.h b/Engine/Source/Runtime/RenderCore/Public/ShaderMaterial.h
--- a/Engine/Source/Runtime/RenderCore/Public/ShaderMaterial.h
+++ b/Engine/Source/Runtime/RenderCore/Public/ShaderMaterial.h
@@ -139,6 +139,7 @@ struct FShaderMaterialPropertyDefines
 
 	uint8 USE_DITHERED_LOD_TRANSITION_FROM_MATERIAL : 1;
 	uint8 MATERIAL_TWOSIDED : 1;
+	uint8 MATERIAL_ISTHINSURFACE : 1;
 	uint8 MATERIAL_TANGENTSPACENORMAL : 1;
 	uint8 GENERATE_SPHERICAL_PARTICLE_NORMALS : 1;
 	uint8 MATERIAL_USE_PREINTEGRATED_GF : 1;
