Fixed bug causing warning error dialog not appearing when the user tries to save a material with compilation errors.

diff --git a/Engine/Source/Editor/MaterialEditor/Private/MaterialEditor.cpp b/Engine/Source/Editor/MaterialEditor/Private/MaterialEditor.cpp
--- a/Engine/Source/Editor/MaterialEditor/Private/MaterialEditor.cpp
+++ b/Engine/Source/Editor/MaterialEditor/Private/MaterialEditor.cpp
@@ -2546,7 +2546,7 @@ void FMaterialEditor::UpdatePreviewMaterial( bool bForce )
 bool FMaterialEditor::UpdateOriginalMaterial()
 {
 	// If the Material has compilation errors, warn the user
-	for (int32 i = ERHIFeatureLevel::SM5; i >= 0; --i)
+	for (int32 i = ERHIFeatureLevel::Num - 1; i >= 0; --i)
 	{
 		ERHIFeatureLevel::Type FeatureLevel = (ERHIFeatureLevel::Type)i;
 		FMaterialResource* CurrentResource = Material->GetMaterialResource(FeatureLevel);
