Fix crash when opening material instance due to IsThinSurface being null for asset create prior to the addition of IsThinSurface property.

diff --git a/Engine/Source/Editor/MaterialEditor/Private/MaterialEditorInstanceDetailCustomization.cpp b/Engine/Source/Editor/MaterialEditor/Private/MaterialEditorInstanceDetailCustomization.cpp
--- a/Engine/Source/Editor/MaterialEditor/Private/MaterialEditorInstanceDetailCustomization.cpp
+++ b/Engine/Source/Editor/MaterialEditor/Private/MaterialEditorInstanceDetailCustomization.cpp
@@ -1333,6 +1333,7 @@ void FMaterialInstanceParameterDetails::CreateBasePropertyOverrideWidgets(IDetai
 			.Visibility(TAttribute<EVisibility>::Create(TAttribute<EVisibility>::FGetter::CreateSP(this, &FMaterialInstanceParameterDetails::IsOverriddenAndVisible, IsOverrideTwoSidedEnabled)))
 			.OverrideResetToDefault(ResetTwoSidedPropertyOverride);
 	}
+	if (IsThinSurfaceProperty)
 	{
 		FIsResetToDefaultVisible IsThinSurfacePropertyResetVisible = FIsResetToDefaultVisible::CreateLambda([this](TSharedPtr<IPropertyHandle> InHandle) {
 			return MaterialEditorInstance->Parent != nullptr ? MaterialEditorInstance->BasePropertyOverrides.bIsThinSurface != MaterialEditorInstance->Parent->IsThinSurface() : false;
