Fixed bug in SMaterialParametersOverviewTree causing Material Editor to crash when the user creates a Param node and views the UI panel containing the material parameters.

diff --git a/Engine/Source/Editor/MaterialEditor/Private/SMaterialParametersOverviewWidget.cpp b/Engine/Source/Editor/MaterialEditor/Private/SMaterialParametersOverviewWidget.cpp
--- a/Engine/Source/Editor/MaterialEditor/Private/SMaterialParametersOverviewWidget.cpp
+++ b/Engine/Source/Editor/MaterialEditor/Private/SMaterialParametersOverviewWidget.cpp
@@ -800,22 +800,10 @@ void SMaterialParametersOverviewTree::ShowSubParameters()
 			}
 
 			UDEditorRuntimeVirtualTextureParameterValue* VTParameter = Cast<UDEditorRuntimeVirtualTextureParameterValue>(Parameter);
-			//Don't show VT samples here
-			if (!VTParameter)
-			{
-				for (TSharedPtr<struct FSortedParamData> GroupChild : SortedParameters)
-				{
-					if (GroupChild->Group.GroupName == Property.ParameterGroup.GroupName
-						&& GroupChild->ParameterInfo.Association == ChildProperty->ParameterInfo.Association
-						&&  GroupChild->ParameterInfo.Index == ChildProperty->ParameterInfo.Index)
-					{
-						GroupChild->Children.Add(ChildProperty);
-					}
-				}
-			}
-
 			UDEditorSparseVolumeTextureParameterValue* SVTParameter = Cast<UDEditorSparseVolumeTextureParameterValue>(Parameter);
-			if (!SVTParameter)
+
+			// Don't add child property to this group if parameter is of type 'Virtual Texture' or 'Sparse Volume Texture'
+			if (!VTParameter && !SVTParameter)
 			{
 				for (TSharedPtr<struct FSortedParamData> GroupChild : SortedParameters)
 				{
