Fixed crash happening after the user creates a Constant4Vector in a MaterialFunction and hits undo.

diff --git a/Engine/Source/Editor/MaterialEditor/Private/MaterialEditor.cpp b/Engine/Source/Editor/MaterialEditor/Private/MaterialEditor.cpp
--- a/Engine/Source/Editor/MaterialEditor/Private/MaterialEditor.cpp
+++ b/Engine/Source/Editor/MaterialEditor/Private/MaterialEditor.cpp
@@ -6384,10 +6384,12 @@ void FMaterialEditor::PostUndo(bool bSuccess)
 			FocusedGraphEd->ClearSelectionSet();
 		}
 		
-		// After an undo operation, the reference contained in graph node to its material expression is lost
-		// (the serializer sets it to null). Here we relink the material expressions graph nodes back.
+		// After an undo operation, objects material expressions reference have their back-references to the material expressions
+		// cleared (the serializer sets it to null). Loop through all material expressions in the material and fix-up these secondary references.
 		for (UMaterialExpression* MaterialExpression : Material->GetExpressions())
 		{
+			MaterialExpression->Material = Material;
+
 			UMaterialGraphNode* GraphNode = Cast<UMaterialGraphNode>(MaterialExpression->GraphNode);
 			if (GraphNode)
 			{
