Strata - editor tab - only use to preview fully simplified materials for now. It will be extended later on according to simplifications and other possible platforms/targets.

diff --git a/Engine/Source/Editor/MaterialEditor/Private/MaterialEditor.cpp b/Engine/Source/Editor/MaterialEditor/Private/MaterialEditor.cpp
--- a/Engine/Source/Editor/MaterialEditor/Private/MaterialEditor.cpp
+++ b/Engine/Source/Editor/MaterialEditor/Private/MaterialEditor.cpp
@@ -155,6 +155,7 @@
 #include "MaterialEditorModes.h"
 #include "Materials/MaterialExpression.h"
 #include "MaterialCachedHLSLTree.h"
+#include "SMaterialEditorStrataWidget.h"
 
 #include "SMaterialParametersOverviewWidget.h"
 #include "SMaterialEditorCustomPrimitiveDataWidget.h"
@@ -444,6 +445,11 @@ void FMaterialEditor::RegisterToolbarTab(const TSharedRef<class FTabManager>& In
 		.SetGroup(WorkspaceMenuCategoryRef)
 		.SetIcon(FSlateIcon(FAppStyle::GetAppStyleSetName(), "LevelEditor.Tabs.Layers"));
 
+	InTabManager->RegisterTabSpawner(FMaterialEditorTabs::StrataTabId, FOnSpawnTab::CreateSP(this, &FMaterialEditor::SpawnTab_Strata))
+		.SetDisplayName(LOCTEXT("StrataTab", "Strata"))
+		.SetGroup(WorkspaceMenuCategoryRef)
+		.SetIcon(FSlateIcon(FAppStyle::GetAppStyleSetName(), "Kismet.Tabs.Palette"));	// STRATA_TODO a strata icon
+
 	MaterialStatsManager->RegisterTabs();
 
 	OnRegisterTabSpawners().Broadcast(InTabManager);
@@ -467,6 +473,7 @@ void FMaterialEditor::UnregisterTabSpawners(const TSharedRef<class FTabManager>&
 	InTabManager->UnregisterTabSpawner(FMaterialEditorTabs::ParameterDefaultsTabId);
 	InTabManager->UnregisterTabSpawner(FMaterialEditorTabs::CustomPrimitiveTabId);
 	InTabManager->UnregisterTabSpawner(FMaterialEditorTabs::LayerPropertiesTabId);
+	InTabManager->UnregisterTabSpawner(FMaterialEditorTabs::StrataTabId);
 
 	MaterialStatsManager->UnregisterTabs();
 
@@ -1590,6 +1597,8 @@ void FMaterialEditor::CreateInternalWidgets()
 	FindResults =
 		SNew(SFindInMaterial, SharedThis(this));
 
+	StrataWidget = SNew(SMaterialEditorStrataWidget, SharedThis(this));
+
 	RegenerateCodeView();
 }
 END_SLATE_FUNCTION_BUILD_OPTIMIZATION
@@ -5173,6 +5182,22 @@ TSharedRef<SDockTab> FMaterialEditor::SpawnTab_LayerProperties(const FSpawnTabAr
 	return SpawnedTab;
 }
 
+TSharedRef<SDockTab> FMaterialEditor::SpawnTab_Strata(const FSpawnTabArgs& Args)
+{
+	check(Args.GetTabId() == FMaterialEditorTabs::StrataTabId);
+
+	TSharedRef<SDockTab> StrataTab = SNew(SDockTab)
+		.Label(LOCTEXT("MaterialStrataTabTitle", "Strata"))
+		[
+			SNew(SBox)
+			[
+				StrataWidget.ToSharedRef()
+			]
+		];
+
+	return StrataTab;
+}
+
 void FMaterialEditor::SetPreviewExpression(UMaterialExpression* NewPreviewExpression)
 {
 	UMaterialExpressionFunctionOutput* FunctionOutput = Cast<UMaterialExpressionFunctionOutput>(NewPreviewExpression);
diff --git a/Engine/Source/Editor/MaterialEditor/Private/MaterialEditor.h b/Engine/Source/Editor/MaterialEditor/Private/MaterialEditor.h
--- a/Engine/Source/Editor/MaterialEditor/Private/MaterialEditor.h
+++ b/Engine/Source/Editor/MaterialEditor/Private/MaterialEditor.h
@@ -46,6 +46,8 @@ struct FGraphAppearanceInfo;
 class UMaterialFunctionInstance;
 class FMaterialCachedHLSLTree;
 struct FMaterialCachedExpressionData;
+class SMaterialEditorStrataWidget;
+
 
 /**
  * Class for rendering previews of material expressions in the material editor's linked object viewport.
@@ -910,6 +912,7 @@ private:
 	TSharedRef<SDockTab> SpawnTab_ParameterDefaults(const FSpawnTabArgs& Args);
 	TSharedRef<SDockTab> SpawnTab_CustomPrimitiveData(const FSpawnTabArgs& Args);
 	TSharedRef<SDockTab> SpawnTab_LayerProperties(const FSpawnTabArgs& Args);
+	TSharedRef<SDockTab> SpawnTab_Strata(const FSpawnTabArgs& Args);
 
 	void OnFinishedChangingProperties(const FPropertyChangedEvent& PropertyChangedEvent);
 	void OnFinishedChangingParametersFromOverview(const FPropertyChangedEvent& PropertyChangedEvent);
@@ -936,6 +939,9 @@ private:
 	/** Palette of Material Expressions and functions */
 	TSharedPtr<class SMaterialPalette> Palette;
 
+	/** The strata control tab */
+	TSharedPtr<class SMaterialEditorStrataWidget> StrataWidget;
+
 	/** Stats log, with the log listing that it reflects */
 	TSharedPtr<class SWidget> Stats;
 	TSharedPtr<class IMessageLogListing> StatsListing;
diff --git a/Engine/Source/Editor/MaterialEditor/Private/MaterialEditorModes.cpp b/Engine/Source/Editor/MaterialEditor/Private/MaterialEditorModes.cpp
--- a/Engine/Source/Editor/MaterialEditor/Private/MaterialEditorModes.cpp
+++ b/Engine/Source/Editor/MaterialEditor/Private/MaterialEditorModes.cpp
@@ -61,8 +61,8 @@ TSharedPtr<FTabManager::FLayout> FMaterialEditorApplicationModes::GetDefaultEdit
 				(
 					FTabManager::NewStack()
 					->SetSizeCoefficient( 0.20f )
-						->AddTab( FMaterialStats::GetGridStatsTabName(), ETabState::ClosedTab )
-						->AddTab( FMaterialStats::GetGridOldStatsTabName(), ETabState::ClosedTab )
+					->AddTab( FMaterialStats::GetGridStatsTabName(), ETabState::ClosedTab )
+					->AddTab( FMaterialStats::GetGridOldStatsTabName(), ETabState::ClosedTab )
 					->AddTab( FMaterialEditorTabs::FindTabId, ETabState::ClosedTab )
 				)
 			)
@@ -73,6 +73,7 @@ TSharedPtr<FTabManager::FLayout> FMaterialEditorApplicationModes::GetDefaultEdit
 				(
 					FTabManager::NewStack()
 					->AddTab( FMaterialEditorTabs::PaletteTabId, ETabState::SidebarTab, ESidebarLocation::Right, 0.14f )
+					->AddTab( FMaterialEditorTabs::StrataTabId,  ETabState::SidebarTab, ESidebarLocation::Right, 0.14f)
 				)
 			)
 		)
diff --git a/Engine/Source/Editor/MaterialEditor/Private/MaterialEditorTabs.cpp b/Engine/Source/Editor/MaterialEditor/Private/MaterialEditorTabs.cpp
--- a/Engine/Source/Editor/MaterialEditor/Private/MaterialEditorTabs.cpp
+++ b/Engine/Source/Editor/MaterialEditor/Private/MaterialEditorTabs.cpp
@@ -13,4 +13,5 @@ const FName FMaterialEditorTabs::GraphEditor(TEXT("Document"));
 const FName FMaterialEditorTabs::PreviewSettingsTabId(TEXT("MaterialEditor_PreviewSettings"));
 const FName FMaterialEditorTabs::ParameterDefaultsTabId(TEXT("MaterialEditor_ParameterDefaults"));
 const FName FMaterialEditorTabs::CustomPrimitiveTabId(TEXT("MaterialEditor_CustomPrimitiveData"));
-const FName FMaterialEditorTabs::LayerPropertiesTabId(TEXT("MaterialInstanceEditor_MaterialLayerProperties"));
\ No newline at end of file
+const FName FMaterialEditorTabs::LayerPropertiesTabId(TEXT("MaterialInstanceEditor_MaterialLayerProperties"));
+const FName FMaterialEditorTabs::StrataTabId(TEXT("MaterialEditor_Strata"));
\ No newline at end of file
diff --git a/Engine/Source/Editor/MaterialEditor/Private/MaterialEditorTabs.h b/Engine/Source/Editor/MaterialEditor/Private/MaterialEditorTabs.h
--- a/Engine/Source/Editor/MaterialEditor/Private/MaterialEditorTabs.h
+++ b/Engine/Source/Editor/MaterialEditor/Private/MaterialEditorTabs.h
@@ -17,6 +17,7 @@ struct FMaterialEditorTabs
 	static const FName ParameterDefaultsTabId;
 	static const FName CustomPrimitiveTabId;
 	static const FName LayerPropertiesTabId;
+	static const FName StrataTabId;
 
 	// Disable default constructor
 	FMaterialEditorTabs() = delete;
diff --git a/Engine/Source/Editor/MaterialEditor/Private/SMaterialEditorStrataWidget.cpp b/Engine/Source/Editor/MaterialEditor/Private/SMaterialEditorStrataWidget.cpp
new file mode 100644
--- /dev/null
+++ b/Engine/Source/Editor/MaterialEditor/Private/SMaterialEditorStrataWidget.cpp
@@ -0,0 +1,126 @@
+// Copyright Epic Games, Inc. All Rights Reserved.
+
+#include "SMaterialEditorStrataWidget.h"
+
+#include "EditorWidgetsModule.h"
+#include "AssetRegistry/AssetRegistryModule.h"
+#include "MaterialEditor.h"
+
+#include "Widgets/Input/SButton.h"
+#include "Widgets/Input/SCheckBox.h"
+#include "Widgets/Layout/SWrapBox.h"
+#include "Widgets/Text/STextBlock.h"
+
+#define LOCTEXT_NAMESPACE "SMaterialEditorStrataWidget"
+
+DEFINE_LOG_CATEGORY_STATIC(LogMaterialEditorStrataWidget, Log, All);
+
+bool Editor_IsStrataEnabled();
+
+void SMaterialEditorStrataWidget::Construct(const FArguments& InArgs, TWeakPtr<FMaterialEditor> InMaterialEditorPtr)
+{
+	MaterialEditorPtr = InMaterialEditorPtr;
+
+	ButtonApplyToPreview = SNew(SButton)
+		.HAlign(HAlign_Center)
+		.OnClicked(this, &SMaterialEditorStrataWidget::OnButtonApplyToPreview)
+		.Text(LOCTEXT("ButtonApplyToPreview", "Apply to preview"));
+
+	CheckBoxForceFullSimplification = SNew(SCheckBox)
+		.Padding(5.0f)
+		.ToolTipText(LOCTEXT("CheckBoxForceFullSimplificationToolTip", "This will force full simplification of the material."));	// Just a test, needs to be more explicit
+
+	if (Editor_IsStrataEnabled())
+	{
+		this->ChildSlot
+		[
+			SNew(SVerticalBox)
+			+SVerticalBox::Slot()
+			.AutoHeight()
+			.Padding(0.0f, 5.0f, 0.0f, 0.0f)
+			[
+				SNew(SWrapBox)
+				.UseAllottedSize(true)
+				+SWrapBox::Slot()
+				.Padding(5.0f)
+				.HAlign(HAlign_Center)
+				.VAlign(VAlign_Center)
+				[
+					CheckBoxForceFullSimplification->AsShared()
+				]
+			]
+			+SVerticalBox::Slot()
+			.AutoHeight()
+			.Padding(0.0f, 5.0f, 0.0f, 0.0f)
+			[
+				SNew(SWrapBox)
+				.UseAllottedSize(true)
+				+SWrapBox::Slot()
+				.Padding(5.0f)
+				.HAlign(HAlign_Center)
+				.VAlign(VAlign_Center)
+				[
+					ButtonApplyToPreview->AsShared()
+				]
+			]
+		];
+	}
+	else
+	{
+		this->ChildSlot
+		[
+			SNew(SVerticalBox)
+			+SVerticalBox::Slot()
+			.AutoHeight()
+			.Padding(0.0f, 5.0f, 0.0f, 0.0f)
+			[
+				SNew(SWrapBox)
+				.UseAllottedSize(true)
+				+SWrapBox::Slot()
+				.Padding(5.0f)
+				.HAlign(HAlign_Center)
+				.VAlign(VAlign_Center)
+				[
+					SNew(STextBlock)
+					.ColorAndOpacity(FLinearColor::Yellow)
+					.ShadowColorAndOpacity(FLinearColor::Black)
+					.ShadowOffset(FVector2D::UnitVector)
+					.Text(LOCTEXT("StrataWidgetNotEnable", "Strata is not enabled."))
+				]
+			]
+		];
+	}
+}
+
+TSharedRef<SWidget> SMaterialEditorStrataWidget::GetContent()
+{
+	return SharedThis(this);
+}
+
+SMaterialEditorStrataWidget::~SMaterialEditorStrataWidget()
+{
+}
+
+void SMaterialEditorStrataWidget::Tick(const FGeometry& AllottedGeometry, const double InCurrentTime, const float InDeltaTime)
+{
+}
+
+FReply SMaterialEditorStrataWidget::OnButtonApplyToPreview()
+{
+	if (MaterialEditorPtr.IsValid())
+	{
+		UMaterialInterface* MaterialInterface = MaterialEditorPtr.Pin()->GetMaterialInterface();
+
+		FStrataCompilationConfig StrataCompilationConfig;
+		StrataCompilationConfig.bFullSimplify = CheckBoxForceFullSimplification->IsChecked();
+		MaterialInterface->SetStrataCompilationConfig(StrataCompilationConfig);
+
+		MaterialInterface->ForceRecompileForRendering();
+	}
+
+
+
+	return FReply::Handled();
+}
+
+#undef LOCTEXT_NAMESPACE
diff --git a/Engine/Source/Editor/MaterialEditor/Private/SMaterialEditorStrataWidget.h b/Engine/Source/Editor/MaterialEditor/Private/SMaterialEditorStrataWidget.h
new file mode 100644
--- /dev/null
+++ b/Engine/Source/Editor/MaterialEditor/Private/SMaterialEditorStrataWidget.h
@@ -0,0 +1,44 @@
+// Copyright Epic Games, Inc. All Rights Reserved.
+
+#pragma once
+
+#include "CoreMinimal.h"
+#include "Widgets/SWidget.h"
+
+#include "Widgets/DeclarativeSyntaxSupport.h"
+#include "Widgets/SCompoundWidget.h"
+#include "Widgets/Input/SButton.h"
+#include "Widgets/Input/SCheckBox.h"
+
+class FMaterialEditor;
+
+class SMaterialEditorStrataWidget : public SCompoundWidget
+{
+public:
+	SLATE_BEGIN_ARGS(SMaterialEditorStrataWidget)
+	{
+	}
+	SLATE_END_ARGS()
+
+	void Construct(const FArguments& InArgs, TWeakPtr<FMaterialEditor> InMaterialEditorPtr);
+
+	/** Gets the widget contents of the app */
+	virtual TSharedRef<SWidget> GetContent();
+
+	virtual ~SMaterialEditorStrataWidget();
+
+	/** SWidget interface */
+	virtual void Tick(const FGeometry& AllottedGeometry, const double InCurrentTime, const float InDeltaTime) override;
+
+private:
+
+	TSharedPtr<class SCheckBox> CheckBoxForceFullSimplification;
+
+	TSharedPtr<class SButton> ButtonApplyToPreview;
+
+	FReply OnButtonApplyToPreview();
+
+	/** Pointer back to the material editor that owns this */
+	TWeakPtr<FMaterialEditor> MaterialEditorPtr;
+};
+
diff --git a/Engine/Source/Runtime/Engine/Classes/Materials/Material.h b/Engine/Source/Runtime/Engine/Classes/Materials/Material.h
--- a/Engine/Source/Runtime/Engine/Classes/Materials/Material.h
+++ b/Engine/Source/Runtime/Engine/Classes/Materials/Material.h
@@ -994,6 +994,9 @@ public:
 	UPROPERTY(EditAnywhere, Category=Refraction)
 	float RefractionDepthBias;
 
+	/** Not a UPROPERTY, used to propagate editor only strata material simplification options for preview. */
+	FStrataCompilationConfig StrataCompilationConfig;
+
 	/** 
 	 * Guid that uniquely identifies this material. 
 	 * Any changes to the state of the material that do not appear separately in the shadermap DDC keys must cause this guid to be regenerated!
diff --git a/Engine/Source/Runtime/Engine/Classes/Materials/MaterialInterface.h b/Engine/Source/Runtime/Engine/Classes/Materials/MaterialInterface.h
--- a/Engine/Source/Runtime/Engine/Classes/Materials/MaterialInterface.h
+++ b/Engine/Source/Runtime/Engine/Classes/Materials/MaterialInterface.h
@@ -63,6 +63,7 @@ class FMaterialCachedHLSLTree;
 struct FParameterChannelNames;
 #endif
 enum EShaderPlatform : uint16;
+struct FStrataCompilationConfig;
 
 typedef TArray<FMaterialResource*> FMaterialResourceDeferredDeletionArray;
 
@@ -396,6 +397,9 @@ public:
 	ENGINE_API bool IsUsingNewHLSLGenerator() const;
 	ENGINE_API bool IsUsingControlFlow() const;
 
+	ENGINE_API const FStrataCompilationConfig& GetStrataCompilationConfig() const;
+	ENGINE_API void SetStrataCompilationConfig(FStrataCompilationConfig& StrataCompilationConfig);
+
 	/**
 	* Test this material for dependency on a given material.
 	* @param	TestDependency - The material to test for dependency upon.
diff --git a/Engine/Source/Runtime/Engine/Private/Materials/HLSLMaterialTranslator.cpp b/Engine/Source/Runtime/Engine/Private/Materials/HLSLMaterialTranslator.cpp
--- a/Engine/Source/Runtime/Engine/Private/Materials/HLSLMaterialTranslator.cpp
+++ b/Engine/Source/Runtime/Engine/Private/Materials/HLSLMaterialTranslator.cpp
@@ -270,7 +270,8 @@ FHLSLMaterialTranslator::FHLSLMaterialTranslator(FMaterial* InMaterial,
 	EShaderPlatform InPlatform,
 	EMaterialQualityLevel::Type InQualityLevel,
 	ERHIFeatureLevel::Type InFeatureLevel,
-	const ITargetPlatform* InTargetPlatform) //if InTargetPlatform is nullptr, we use the current active
+	const ITargetPlatform* InTargetPlatform, //if InTargetPlatform is nullptr, we use the current active
+	const FStrataCompilationConfig* InStrataCompilationConfig)
 :	ShaderFrequency(SF_Pixel)
 ,	MaterialProperty(MP_EmissiveColor)
 ,	CurrentScopeChunks(nullptr)
@@ -331,6 +332,7 @@ FHLSLMaterialTranslator::FHLSLMaterialTranslator(FMaterial* InMaterial,
 ,	FinalUsedSharedLocalBasesCount(0)
 ,	NumVtSamples(0)
 ,	TargetPlatform(InTargetPlatform)
+,	StrataCompilationConfig()
 {
 	FMemory::Memzero(SharedPixelProperties);
 
@@ -392,6 +394,11 @@ FHLSLMaterialTranslator::FHLSLMaterialTranslator(FMaterial* InMaterial,
 
 	// Default value used as the root of the tree for the first path (when a node parent==nullptr).
 	StrataNodeIdentifierStack.Push(FGuid(0x7AEE, 0xBAD, 0xDEAD, 0xBEEF));
+
+	if (InStrataCompilationConfig)
+	{
+		StrataCompilationConfig = *InStrataCompilationConfig;
+	}
 }
 
 FHLSLMaterialTranslator::~FHLSLMaterialTranslator()
@@ -10230,6 +10237,8 @@ bool FHLSLMaterialTranslator::StrataGenerateDerivedMaterialOperatorData()
 		}
 	}
 
+	StrataSimplificationStatus.bRunFullSimplification = StrataCompilationConfig.bFullSimplify && !bStrataUsesConversionFromLegacy;
+
 	const uint32 StrataBytePerPixel = GetStrataBytePerPixel(Platform);
 	do 
 	{
diff --git a/Engine/Source/Runtime/Engine/Private/Materials/HLSLMaterialTranslator.h b/Engine/Source/Runtime/Engine/Private/Materials/HLSLMaterialTranslator.h
--- a/Engine/Source/Runtime/Engine/Private/Materials/HLSLMaterialTranslator.h
+++ b/Engine/Source/Runtime/Engine/Private/Materials/HLSLMaterialTranslator.h
@@ -469,6 +469,9 @@ protected:
 	uint32 NumVtSamples;
 
 	const ITargetPlatform* TargetPlatform;
+
+	/** Strata material compilation and simplification configuration. */
+	FStrataCompilationConfig StrataCompilationConfig;
 public: 
 
 	FHLSLMaterialTranslator(FMaterial* InMaterial,
@@ -477,7 +480,8 @@ public:
 		EShaderPlatform InPlatform,
 		EMaterialQualityLevel::Type InQualityLevel,
 		ERHIFeatureLevel::Type InFeatureLevel,
-		const ITargetPlatform* InTargetPlatform = nullptr);
+		const ITargetPlatform* InTargetPlatform = nullptr,
+		const FStrataCompilationConfig* InStrataCompilationConfig = nullptr);
 
 	~FHLSLMaterialTranslator();
 
diff --git a/Engine/Source/Runtime/Engine/Private/Materials/Material.cpp b/Engine/Source/Runtime/Engine/Private/Materials/Material.cpp
--- a/Engine/Source/Runtime/Engine/Private/Materials/Material.cpp
+++ b/Engine/Source/Runtime/Engine/Private/Materials/Material.cpp
@@ -996,6 +996,8 @@ UMaterial::UMaterial(const FObjectInitializer& ObjectInitializer)
 	RefractionDepthBias = 0.0f;
 	MaterialDecalResponse = MDR_ColorNormalRoughness;
 
+	StrataCompilationConfig = FStrataCompilationConfig();
+
 	bAllowDevelopmentShaderCompile = true;
 	bIsMaterialEditorStatsMaterial = false;
 
diff --git a/Engine/Source/Runtime/Engine/Private/Materials/MaterialInterface.cpp b/Engine/Source/Runtime/Engine/Private/Materials/MaterialInterface.cpp
--- a/Engine/Source/Runtime/Engine/Private/Materials/MaterialInterface.cpp
+++ b/Engine/Source/Runtime/Engine/Private/Materials/MaterialInterface.cpp
@@ -263,6 +263,22 @@ bool UMaterialInterface::IsUsingNewHLSLGenerator() const
 	return BaseMaterial ? BaseMaterial->bEnableNewHLSLGenerator : false;
 }
 
+const FStrataCompilationConfig& UMaterialInterface::GetStrataCompilationConfig() const
+{
+	const UMaterial* BaseMaterial = GetMaterial_Concurrent();
+	static FStrataCompilationConfig DefaultFStrataCompilationConfig = FStrataCompilationConfig();
+	return BaseMaterial ? BaseMaterial->StrataCompilationConfig : DefaultFStrataCompilationConfig;
+}
+
+ENGINE_API void UMaterialInterface::SetStrataCompilationConfig(FStrataCompilationConfig& StrataCompilationConfig)
+{
+	UMaterial* BaseMaterial = GetMaterial();
+	if (BaseMaterial)
+	{
+		BaseMaterial->StrataCompilationConfig = StrataCompilationConfig;
+	}
+}
+
 void UMaterialInterface::GetQualityLevelUsage(TArray<bool, TInlineAllocator<EMaterialQualityLevel::Num> >& OutQualityLevelsUsed, EShaderPlatform ShaderPlatform, bool bCooking)
 {
 	OutQualityLevelsUsed = GetCachedExpressionData().QualityLevelsUsed;
diff --git a/Engine/Source/Runtime/Engine/Private/Materials/MaterialShader.cpp b/Engine/Source/Runtime/Engine/Private/Materials/MaterialShader.cpp
--- a/Engine/Source/Runtime/Engine/Private/Materials/MaterialShader.cpp
+++ b/Engine/Source/Runtime/Engine/Private/Materials/MaterialShader.cpp
@@ -566,6 +566,28 @@ void FStaticParameterSet::SetStaticComponentMaskParameterValue(const FMaterialPa
 #endif // WITH_EDITORONLY_DATA
 
 
+
+#if WITH_EDITOR
+FString FStrataCompilationConfig::GetShaderMapKeyString() const
+{
+	if (bFullSimplify)
+	{
+		return TEXT("_STRTFS");
+	}
+	return TEXT("");
+}
+
+void FStrataCompilationConfig::UpdateHash(FSHA1& Hasher) const
+{
+	Hasher.Update((const uint8*)(&bFullSimplify), sizeof(bFullSimplify));
+}
+
+void FStrataCompilationConfig::Serialize(FArchive& Ar)
+{
+	Ar << bFullSimplify;
+}
+#endif
+
 void FMaterialShaderMapId::Serialize(FArchive& Ar, bool bLoadedByCookedMaterial)
 {
 	SCOPED_LOADTIMER(FMaterialShaderMapId_Serialize);
@@ -692,6 +714,11 @@ void FMaterialShaderMapId::Serialize(FArchive& Ar, bool bLoadedByCookedMaterial)
 		{
 			bUsingNewHLSLGenerator = false;
 		}
+
+		// STRATA_TODO We do not need to serialize FStrataCompilationConfig for now since this is only used when debugging in the editor.
+		// However we might want to do that when compilation config will change between raster and path tracing for instance.
+		// So currently, the shader map DDC key string won't be changing, but if the user toggles simplification on via the Material Editor it will cache a new map. 
+		// In other words we only cache the simplified shader map version of the material in the editor (and not during cooks).
 	}
 	else
 	{
@@ -787,6 +814,8 @@ void FMaterialShaderMapId::GetMaterialHash(FSHAHash& OutHash, bool bWithStaticPa
 
 	HashState.Update((const uint8*)&bUsingNewHLSLGenerator, sizeof(bUsingNewHLSLGenerator));
 
+	StrataCompilationConfig.UpdateHash(HashState);
+
 	HashState.Final();
 	HashState.GetHash(&OutHash.Hash[0]);
 }
@@ -814,6 +843,11 @@ bool FMaterialShaderMapId::Equals(const FMaterialShaderMapId& ReferenceSet, bool
 		return false;
 	}
 
+	if (StrataCompilationConfig != ReferenceSet.StrataCompilationConfig)
+	{
+		return false;
+	}
+
 	if (!IsCookedId())
 	{
 		if (BaseMaterialId != ReferenceSet.BaseMaterialId)
@@ -1100,6 +1134,8 @@ void FMaterialShaderMapId::AppendKeyString(FString& KeyString, bool bIncludeSour
 	{
 		KeyString += TEXT("_NewHLSL");
 	}
+
+	KeyString += StrataCompilationConfig.GetShaderMapKeyString();
 }
 
 void FMaterialShaderMapId::SetShaderDependencies(const TArray<FShaderType*>& ShaderTypes, const TArray<const FShaderPipelineType*>& ShaderPipelineTypes, const TArray<FVertexFactoryType*>& VFTypes, EShaderPlatform ShaderPlatform)
diff --git a/Engine/Source/Runtime/Engine/Private/Materials/MaterialShared.cpp b/Engine/Source/Runtime/Engine/Private/Materials/MaterialShared.cpp
--- a/Engine/Source/Runtime/Engine/Private/Materials/MaterialShared.cpp
+++ b/Engine/Source/Runtime/Engine/Private/Materials/MaterialShared.cpp
@@ -747,6 +747,8 @@ void FMaterial::GetShaderMapId(EShaderPlatform Platform, const ITargetPlatform*
 		OutId.SetShaderDependencies(ShaderTypes, ShaderPipelineTypes, VFTypes, Platform);
 		GetReferencedTexturesHash(Platform, OutId.TextureReferencesHash);
 
+		OutId.StrataCompilationConfig = GetStrataCompilationConfig();
+
 #else
 		OutId.QualityLevel = GetQualityLevel();
 		OutId.FeatureLevel = GetFeatureLevel();
@@ -865,6 +867,22 @@ bool FMaterial::IsUsingNewHLSLGenerator() const
 	return MaterialInterface ? MaterialInterface->IsUsingNewHLSLGenerator() : false;
 }
 
+const FStrataCompilationConfig& FMaterial::GetStrataCompilationConfig() const
+{
+	const UMaterialInterface* MaterialInterface = GetMaterialInterface();
+	static FStrataCompilationConfig DefaultFStrataCompilationConfig = FStrataCompilationConfig();
+	return MaterialInterface ? MaterialInterface->GetStrataCompilationConfig() : DefaultFStrataCompilationConfig;
+}
+
+void FMaterial::SetStrataCompilationConfig(FStrataCompilationConfig& StrataCompilationConfig)
+{
+	UMaterialInterface* MaterialInterface = GetMaterialInterface();
+	if (MaterialInterface)
+	{
+		MaterialInterface->SetStrataCompilationConfig(StrataCompilationConfig);
+	}
+}
+
 #endif // WITH_EDITOR
 
 bool FMaterial::HasValidGameThreadShaderMap() const
@@ -3096,7 +3114,7 @@ bool FMaterial::Translate_Legacy(const FMaterialShaderMapId& ShaderMapId,
 	FMaterialCompilationOutput& OutCompilationOutput,
 	TRefCountPtr<FSharedShaderCompilerEnvironment>& OutMaterialEnvironment)
 {
-	FHLSLMaterialTranslator MaterialTranslator(this, OutCompilationOutput, InStaticParameters, InPlatform, GetQualityLevel(), ShaderMapId.FeatureLevel, InTargetPlatform);
+	FHLSLMaterialTranslator MaterialTranslator(this, OutCompilationOutput, InStaticParameters, InPlatform, GetQualityLevel(), ShaderMapId.FeatureLevel, InTargetPlatform, &ShaderMapId.StrataCompilationConfig);
 	const bool bSuccess = MaterialTranslator.Translate();
 	if (bSuccess)
 	{
diff --git a/Engine/Source/Runtime/Engine/Public/MaterialShared.h b/Engine/Source/Runtime/Engine/Public/MaterialShared.h
--- a/Engine/Source/Runtime/Engine/Public/MaterialShared.h
+++ b/Engine/Source/Runtime/Engine/Public/MaterialShared.h
@@ -36,6 +36,7 @@
 #include "Templates/UnrealTemplate.h"
 #include "ShaderCompilerCore.h"
 #include "PSOPrecache.h"
+#include "UObject/ObjectMacros.h"
 
 #if UE_ENABLE_INCLUDE_ORDER_DEPRECATED_IN_5_2
 #include "RHI.h"
@@ -872,6 +873,10 @@ public:
 
 	/** Is the material using the new HLSL generator? */
 	bool bUsingNewHLSLGenerator;
+
+	/** The Strata configuration used when compiling this material, can be tweaked in the material editor for live visualization of simplification. */
+	FStrataCompilationConfig StrataCompilationConfig;
+
 #endif // WITH_EDITOR
 
 	/*
@@ -886,6 +891,7 @@ public:
 		, Usage(EMaterialShaderMapUsage::Default)
 		, bIsCookedId(false)
 		, bUsingNewHLSLGenerator(false)
+		, StrataCompilationConfig()
 #endif
 	{ }
 
@@ -1935,6 +1941,12 @@ public:
 
 	/** Is the material using the new (WIP) HLSL generator? */
 	ENGINE_API virtual bool IsUsingNewHLSLGenerator() const;
+
+	/** Get to the strata compilation config */
+	ENGINE_API virtual const FStrataCompilationConfig& GetStrataCompilationConfig() const;
+	/** Set the strata compilation config */
+	ENGINE_API virtual void SetStrataCompilationConfig(FStrataCompilationConfig& StrataCompilationConfig);
+
 #endif // WITH_EDITOR
 
 	/**
diff --git a/Engine/Source/Runtime/Engine/Public/MaterialTypes.h b/Engine/Source/Runtime/Engine/Public/MaterialTypes.h
--- a/Engine/Source/Runtime/Engine/Public/MaterialTypes.h
+++ b/Engine/Source/Runtime/Engine/Public/MaterialTypes.h
@@ -448,4 +448,26 @@ struct FMaterialParameterMetadata
 	/** Is the parameter overriden on the material it was queried from? */
 	bool bOverride = false;
 #endif // WITH_EDITORONLY_DATA
-};
\ No newline at end of file
+};
+
+class FSHA1;
+struct FStrataCompilationConfig
+{
+	bool bFullSimplify = false;
+
+	FString GetShaderMapKeyString() const;
+
+	void UpdateHash(FSHA1& Hasher) const;
+
+	void Serialize(FArchive& Ar);
+
+	friend inline bool operator==(const FStrataCompilationConfig& Lhs, const FStrataCompilationConfig& Rhs)
+	{
+		return Lhs.bFullSimplify == Rhs.bFullSimplify;
+	}
+
+	friend inline bool operator!=(const FStrataCompilationConfig& Lhs, const FStrataCompilationConfig& Rhs)
+	{
+		return !operator==(Lhs, Rhs);
+	}
+};
