Strata - Editor widget - added label to tickbox.

diff --git a/Engine/Source/Editor/MaterialEditor/Private/SMaterialEditorStrataWidget.cpp b/Engine/Source/Editor/MaterialEditor/Private/SMaterialEditorStrataWidget.cpp
--- a/Engine/Source/Editor/MaterialEditor/Private/SMaterialEditorStrataWidget.cpp
+++ b/Engine/Source/Editor/MaterialEditor/Private/SMaterialEditorStrataWidget.cpp
@@ -39,14 +39,32 @@ void SMaterialEditorStrataWidget::Construct(const FArguments& InArgs, TWeakPtr<F
 			.AutoHeight()
 			.Padding(0.0f, 5.0f, 0.0f, 0.0f)
 			[
-				SNew(SWrapBox)
-				.UseAllottedSize(true)
-				+SWrapBox::Slot()
-				.Padding(5.0f)
-				.HAlign(HAlign_Center)
+				SNew(SHorizontalBox)
+				+SHorizontalBox::Slot()
+				.AutoWidth()
 				.VAlign(VAlign_Center)
 				[
-					CheckBoxForceFullSimplification->AsShared()
+					SNew(SWrapBox)
+					.UseAllottedSize(true)
+					+SWrapBox::Slot()
+					.Padding(5.0f)
+					.HAlign(HAlign_Center)
+					.VAlign(VAlign_Center)
+					[
+						CheckBoxForceFullSimplification->AsShared()
+					]
+				]
+				+SHorizontalBox::Slot()
+				.FillWidth(1.0f)
+				.Padding(16.0f, 0.0f)
+				.HAlign(HAlign_Left)
+				.VAlign(VAlign_Center)
+				[
+					SNew(STextBlock)
+					.ColorAndOpacity(FLinearColor::White)
+					.ShadowColorAndOpacity(FLinearColor::Black)
+					.ShadowOffset(FVector2D::UnitVector)
+					.Text(LOCTEXT("FullsimplificationLabel", "Full simplification"))
 				]
 			]
 			+SVerticalBox::Slot()
