Fixed crash caused by renaming material graph node to name 1024 character or more long.

diff --git a/Engine/Source/Editor/MaterialEditor/Private/MaterialEditor.cpp b/Engine/Source/Editor/MaterialEditor/Private/MaterialEditor.cpp
--- a/Engine/Source/Editor/MaterialEditor/Private/MaterialEditor.cpp
+++ b/Engine/Source/Editor/MaterialEditor/Private/MaterialEditor.cpp
@@ -7417,9 +7417,15 @@ void FMaterialEditor::OnNodeTitleCommitted(const FText& NewText, ETextCommit::Ty
 {
 	if (NodeBeingChanged)
 	{
+		FString NewName = NewText.ToString();
+		if (NewName.Len() >= NAME_SIZE) {
+			UE_LOG(LogMaterialEditor, Warning, TEXT("New material graph node name '%s...' exceeds maximum length of %d and thus was truncated."), *NewName.Left(8), NAME_SIZE - 1);
+			NewName = NewName.Left(NAME_SIZE - 1);
+		}
+
 		const FScopedTransaction Transaction( LOCTEXT( "RenameNode", "Rename Node" ) );
 		NodeBeingChanged->Modify();
-		NodeBeingChanged->OnRenameNode(NewText.ToString());
+		NodeBeingChanged->OnRenameNode(NewName);
 		UpdateGenerator();
 		MaterialCustomPrimitiveDataWidget->UpdateEditorInstance(MaterialEditorInstance);
 	}
