Conditionally generate material platform stats for GPUSkinVertexFactory only when UsedWithSkeletalMesh, UsedWithMorphTarget, UsedWithCloth boxes are ticked.
It fixes GPUSkinVertexFactory throwing shader errors when material is not used with skeletal mesh

diff --git a/Engine/Source/Editor/MaterialEditor/Private/MaterialStatsCommon.cpp b/Engine/Source/Editor/MaterialEditor/Private/MaterialStatsCommon.cpp
--- a/Engine/Source/Editor/MaterialEditor/Private/MaterialStatsCommon.cpp
+++ b/Engine/Source/Editor/MaterialEditor/Private/MaterialStatsCommon.cpp
@@ -190,6 +190,10 @@ FString FMaterialStatsUtils::RepresentativeShaderTypeToString(const ERepresentat
 			return TEXT("Skeletal Mesh");
 		break;
 
+		case ERepresentativeShader::SkinnedCloth:
+			return TEXT("Skinned Cloth");
+		break;
+
 		case ERepresentativeShader::UIDefaultFragmentShader:
 			return TEXT("UI Pixel Shader");
 		break;
@@ -287,6 +291,7 @@ void FMaterialStatsUtils::GetRepresentativeShaderTypesAndDescriptions(TMap<FName
 
 	static const FName FLocalVertexFactoryName = FLocalVertexFactory::StaticType.GetFName();
 	static const FName FGPUFactoryName = TEXT("TGPUSkinVertexFactoryDefault");
+	static const FName FClothVertexFactoryName = TEXT("TGPUSkinAPEXClothVertexFactoryDefault");
 
 	if (TargetMaterial->IsUIMaterial())
 	{
@@ -340,8 +345,16 @@ void FMaterialStatsUtils::GetRepresentativeShaderTypesAndDescriptions(TMap<FName
 		ShaderTypeNamesAndDescriptions.FindOrAdd(FLocalVertexFactoryName)
 			.Add(FRepresentativeShaderInfo(ERepresentativeShader::StaticMesh, TBasePassVSFNoLightMapPolicyName, TEXT("Base pass vertex shader")));
 
-		ShaderTypeNamesAndDescriptions.FindOrAdd(FGPUFactoryName)
-			.Add(FRepresentativeShaderInfo(ERepresentativeShader::SkeletalMesh, TBasePassVSFNoLightMapPolicyName, TEXT("Base pass vertex shader")));
+		if (TargetMaterial->IsUsedWithSkeletalMesh() || TargetMaterial->IsUsedWithMorphTargets())
+		{	
+			ShaderTypeNamesAndDescriptions.FindOrAdd(FGPUFactoryName)
+				.Add(FRepresentativeShaderInfo(ERepresentativeShader::SkeletalMesh, TBasePassVSFNoLightMapPolicyName, TEXT("Base pass vertex shader")));
+		}
+		if (TargetMaterial->IsUsedWithAPEXCloth())
+		{
+			ShaderTypeNamesAndDescriptions.FindOrAdd(FClothVertexFactoryName)
+				.Add(FRepresentativeShaderInfo(ERepresentativeShader::SkinnedCloth, TBasePassVSFNoLightMapPolicyName, TEXT("Base pass vertex shader")));
+		}
 
 		// Add the shader type with the most sampler usages so we can accurately report the worst case scenario.
 		// This is ad-hoc, and ideally we have a better way for finding this shader type in the future.
@@ -378,9 +391,18 @@ void FMaterialStatsUtils::GetRepresentativeShaderTypesAndDescriptions(TMap<FName
 				.Add(FRepresentativeShaderInfo(ERepresentativeShader::StaticMesh, FName(ShaderNameStr),
 					FString::Printf(TEXT("Mobile base pass vertex shader%s"), DescSuffix)));
 
-			ShaderTypeNamesAndDescriptions.FindOrAdd(FGPUFactoryName)
-				.Add(FRepresentativeShaderInfo(ERepresentativeShader::SkeletalMesh, FName(ShaderNameStr),
-					FString::Printf(TEXT("Mobile base pass vertex shader%s"), DescSuffix)));
+			if (TargetMaterial->IsUsedWithSkeletalMesh() || TargetMaterial->IsUsedWithMorphTargets())
+			{
+				ShaderTypeNamesAndDescriptions.FindOrAdd(FGPUFactoryName)
+					.Add(FRepresentativeShaderInfo(ERepresentativeShader::SkeletalMesh, FName(ShaderNameStr),
+						FString::Printf(TEXT("Mobile base pass vertex shader%s"), DescSuffix)));
+			}
+			if (TargetMaterial->IsUsedWithAPEXCloth())
+			{
+				ShaderTypeNamesAndDescriptions.FindOrAdd(FClothVertexFactoryName)
+					.Add(FRepresentativeShaderInfo(ERepresentativeShader::SkinnedCloth, FName(ShaderNameStr),
+						FString::Printf(TEXT("Mobile base pass vertex shader%s"), DescSuffix)));
+			}
 		}
 		else
 		{
@@ -465,9 +487,18 @@ void FMaterialStatsUtils::GetRepresentativeShaderTypesAndDescriptions(TMap<FName
 				.Add(FRepresentativeShaderInfo(ERepresentativeShader::StaticMesh, FName(ShaderNameStr),
 				FString::Printf(TEXT("Mobile base pass vertex shader%s"), DescSuffix)));
 
-			ShaderTypeNamesAndDescriptions.FindOrAdd(FGPUFactoryName)
-				.Add(FRepresentativeShaderInfo(ERepresentativeShader::SkeletalMesh, FName(ShaderNameStr),
-				FString::Printf(TEXT("Mobile base pass vertex shader%s"), DescSuffix)));
+			if (TargetMaterial->IsUsedWithSkeletalMesh() || TargetMaterial->IsUsedWithMorphTargets())
+			{
+				ShaderTypeNamesAndDescriptions.FindOrAdd(FGPUFactoryName)
+					.Add(FRepresentativeShaderInfo(ERepresentativeShader::SkeletalMesh, FName(ShaderNameStr),
+						FString::Printf(TEXT("Mobile base pass vertex shader%s"), DescSuffix)));
+			}
+			if (TargetMaterial->IsUsedWithAPEXCloth())
+			{
+				ShaderTypeNamesAndDescriptions.FindOrAdd(FClothVertexFactoryName)
+					.Add(FRepresentativeShaderInfo(ERepresentativeShader::SkinnedCloth, FName(ShaderNameStr),
+						FString::Printf(TEXT("Mobile base pass vertex shader%s"), DescSuffix)));
+			}
 
 			// dynamically lit shader FMobileDirectionalLightAndCSMPolicy
 			MobileBasePassShaderName(false, TEXT("FMobileDirectionalLightAndCSMPolicy"), bEnableLocalLights, bMobileHDR, bOnlySkyPermutation, ShaderNameStr);
@@ -480,9 +511,18 @@ void FMaterialStatsUtils::GetRepresentativeShaderTypesAndDescriptions(TMap<FName
 				.Add(FRepresentativeShaderInfo(ERepresentativeShader::StaticMesh, FName(ShaderNameStr),
 				FString::Printf(TEXT("Mobile base pass vertex shader%s"), DescSuffix)));
 
-			ShaderTypeNamesAndDescriptions.FindOrAdd(FGPUFactoryName)
-				.Add(FRepresentativeShaderInfo(ERepresentativeShader::SkeletalMesh, FName(ShaderNameStr),
-				FString::Printf(TEXT("Mobile base pass vertex shader%s"), DescSuffix)));
+			if (TargetMaterial->IsUsedWithSkeletalMesh() || TargetMaterial->IsUsedWithMorphTargets())
+			{
+				ShaderTypeNamesAndDescriptions.FindOrAdd(FGPUFactoryName)
+					.Add(FRepresentativeShaderInfo(ERepresentativeShader::SkeletalMesh, FName(ShaderNameStr),
+						FString::Printf(TEXT("Mobile base pass vertex shader%s"), DescSuffix)));
+			}
+			if (TargetMaterial->IsUsedWithAPEXCloth())
+			{
+				ShaderTypeNamesAndDescriptions.FindOrAdd(FClothVertexFactoryName)
+					.Add(FRepresentativeShaderInfo(ERepresentativeShader::SkinnedCloth, FName(ShaderNameStr),
+						FString::Printf(TEXT("Mobile base pass vertex shader%s"), DescSuffix)));
+			}
 		}
 	}
 }
diff --git a/Engine/Source/Editor/MaterialEditor/Public/MaterialStatsCommon.h b/Engine/Source/Editor/MaterialEditor/Public/MaterialStatsCommon.h
--- a/Engine/Source/Editor/MaterialEditor/Public/MaterialStatsCommon.h
+++ b/Engine/Source/Editor/MaterialEditor/Public/MaterialStatsCommon.h
@@ -59,6 +59,7 @@ enum class ERepresentativeShader
 	FirstVertexShader,
 	StaticMesh = FirstVertexShader,
 	SkeletalMesh,
+	SkinnedCloth,
 
 	UIDefaultVertexShader,
 	UIInstancedVertexShader,
