Fix memstomp with color picker. Prepare SColorPicker for API change. Update AssetViewUtilis::PathColors to use FLinearColor instead of shared pointer. The shared pointer was only needed for the SColorPicker.

diff --git a/Engine/Plugins/Animation/ControlRig/Source/ControlRigEditor/Private/ControlRigGraphDetails.cpp b/Engine/Plugins/Animation/ControlRig/Source/ControlRigEditor/Private/ControlRigGraphDetails.cpp
--- a/Engine/Plugins/Animation/ControlRig/Source/ControlRigEditor/Private/ControlRigGraphDetails.cpp
+++ b/Engine/Plugins/Animation/ControlRig/Source/ControlRigEditor/Private/ControlRigGraphDetails.cpp
@@ -1175,8 +1175,6 @@ FLinearColor FControlRigGraphDetails::GetNodeColor() const
 
 void FControlRigGraphDetails::SetNodeColor(FLinearColor InColor, bool bSetupUndoRedo)
 {
-	TargetColor = InColor;
-
 	if (GraphPtr.IsValid() && ControlRigBlueprintPtr.IsValid())
 	{
 		UControlRigBlueprint* Blueprint = ControlRigBlueprintPtr.Get();
@@ -1186,7 +1184,7 @@ void FControlRigGraphDetails::SetNodeColor(FLinearColor InColor, bool bSetupUndo
 			{
 				if (URigVMController* Controller = Blueprint->GetOrCreateController(OuterNode->GetGraph()))
 				{
-					Controller->SetNodeColor(OuterNode, TargetColor, bSetupUndoRedo, bIsPickingColor, true);
+					Controller->SetNodeColor(OuterNode, InColor, bSetupUndoRedo, bIsPickingColor, true);
 				}
 			}
 		}
@@ -1209,19 +1207,12 @@ void FControlRigGraphDetails::OnNodeColorCancelled(FLinearColor OriginalColor)
 
 FReply FControlRigGraphDetails::OnNodeColorClicked()
 {
-	TargetColor = GetNodeColor();
-	TargetColors.Reset();
-	TargetColors.Add(&TargetColor);
-
-	FColorPickerArgs PickerArgs;
+	FColorPickerArgs PickerArgs = FColorPickerArgs(GetNodeColor(), FOnLinearColorValueChanged::CreateSP(this, &FControlRigGraphDetails::SetNodeColor, true));
 	PickerArgs.ParentWidget = ColorBlock;
 	PickerArgs.bUseAlpha = false;
 	PickerArgs.DisplayGamma = false;
-	PickerArgs.InitialColorOverride = TargetColor;
-	PickerArgs.LinearColorArray = &TargetColors;
 	PickerArgs.OnInteractivePickBegin = FSimpleDelegate::CreateSP(this, &FControlRigGraphDetails::OnNodeColorBegin);
 	PickerArgs.OnInteractivePickEnd = FSimpleDelegate::CreateSP(this, &FControlRigGraphDetails::OnNodeColorEnd);
-	PickerArgs.OnColorCommitted = FOnLinearColorValueChanged::CreateSP(this, &FControlRigGraphDetails::SetNodeColor, true);
 	PickerArgs.OnColorPickerCancelled = FOnColorPickerCancelled::CreateSP(this, &FControlRigGraphDetails::OnNodeColorCancelled);
 	OpenColorPicker(PickerArgs);
 	return FReply::Handled();
diff --git a/Engine/Plugins/Animation/ControlRig/Source/ControlRigEditor/Private/ControlRigGraphDetails.h b/Engine/Plugins/Animation/ControlRig/Source/ControlRigEditor/Private/ControlRigGraphDetails.h
--- a/Engine/Plugins/Animation/ControlRig/Source/ControlRigEditor/Private/ControlRigGraphDetails.h
+++ b/Engine/Plugins/Animation/ControlRig/Source/ControlRigEditor/Private/ControlRigGraphDetails.h
@@ -205,12 +205,6 @@ private:
 	/** The color block widget */
 	TSharedPtr<SColorBlock> ColorBlock;
 
-	/** The color to change */
-	FLinearColor TargetColor;
-
-	/** The color array to change */
-	TArray<FLinearColor*> TargetColors;
-
 	/** Set to true if the UI is currently picking a color */
 	bool bIsPickingColor;
 
diff --git a/Engine/Plugins/Runtime/Metasound/Source/MetasoundEditor/Private/SMetasoundPinValueInspector.cpp b/Engine/Plugins/Runtime/Metasound/Source/MetasoundEditor/Private/SMetasoundPinValueInspector.cpp
--- a/Engine/Plugins/Runtime/Metasound/Source/MetasoundEditor/Private/SMetasoundPinValueInspector.cpp
+++ b/Engine/Plugins/Runtime/Metasound/Source/MetasoundEditor/Private/SMetasoundPinValueInspector.cpp
@@ -419,9 +419,6 @@ namespace Metasound
 				return FReply::Unhandled();
 			}
 
-			FLinearColor Color = InInitColor;
-			TArray<FLinearColor*> LinearColorArray;
-			LinearColorArray.Add(&Color);
 			const FMetasoundFrontendEdgeStyle* EdgeStyle = GetEdgeStyle();
 			if (!EdgeStyle)
 			{
@@ -429,9 +426,8 @@ namespace Metasound
 			}
 
 			FColorPickerArgs PickerArgs = InitPickerArgs();
-
-			PickerArgs.LinearColorArray = &LinearColorArray;
-			PickerArgs.OnColorCommitted = FOnLinearColorValueChanged::CreateLambda([this, InIndex](FLinearColor NewColor) { OnColorCommitted(NewColor, InIndex); });
+			PickerArgs.InitialColorOverride = InInitColor;
+			PickerArgs.OnColorCommitted = FOnLinearColorValueChanged::CreateSP(this, &FMetasoundNumericDebugLineItem::OnColorCommitted, InIndex);
 
 			OpenColorPicker(PickerArgs);
 
diff --git a/Engine/Source/Developer/AssetTools/Private/AssetViewUtils.cpp b/Engine/Source/Developer/AssetTools/Private/AssetViewUtils.cpp
--- a/Engine/Source/Developer/AssetTools/Private/AssetViewUtils.cpp
+++ b/Engine/Source/Developer/AssetTools/Private/AssetViewUtils.cpp
@@ -59,7 +59,7 @@ namespace AssetViewUtils
 	static FOnFolderPathChanged OnFolderPathChangedDelegate;
 
 	/** Keep a map of all the paths that have custom colors, so updating the color in one location updates them all */
-	static TMap< FString, TSharedPtr< FLinearColor > > PathColors;
+	static TMap< FString, FLinearColor > PathColors;
 
 	/** Internal function to delete a folder from disk, but only if it is empty. InPathToDelete is in FPackageName format. */
 	bool DeleteEmptyFolderFromDisk(const FString& InPathToDelete);
@@ -578,11 +578,11 @@ bool AssetViewUtils::RenameFolder(const FString& DestPath, const FString& Source
 	}
 
 	// set color of folder to new path
-	const TSharedPtr<FLinearColor> FolderColor = LoadColor(SourcePath);
-	if (FolderColor.IsValid())
+	const TOptional<FLinearColor> FolderColor = GetPathColor(SourcePath);
+	if (FolderColor.IsSet())
 	{
-		SaveColor(SourcePath, nullptr);
-		SaveColor(DestPath, FolderColor);
+		SetPathColor(SourcePath, TOptional<FLinearColor>());
+		SetPathColor(DestPath, FolderColor);
 	}
 
 	{
@@ -629,10 +629,10 @@ bool AssetViewUtils::CopyFolders(const TArray<FString>& InSourcePathNames, const
 			ObjectTools::DuplicateObjects( PathIt.Value(), SourcePath, Destination, /*bOpenDialog=*/false );
 		}
 
-		const TSharedPtr<FLinearColor> FolderColor = LoadColor(SourcePath);
-		if (FolderColor.IsValid())
+		const TOptional<FLinearColor> FolderColor = GetPathColor(SourcePath);
+		if (FolderColor.IsSet())
 		{
-			SaveColor(Destination, FolderColor);
+			SetPathColor(Destination, FolderColor);
 		}
 	}
 
@@ -710,11 +710,11 @@ bool AssetViewUtils::MoveFolders(const TArray<FString>& InSourcePathNames, const
 			AssetRegistryModule.Get().RemovePath(SourcePath);
 		}
 
-		const TSharedPtr<FLinearColor> FolderColor = LoadColor(SourcePath);
-		if (FolderColor.IsValid())
+		const TOptional<FLinearColor> FolderColor = GetPathColor(SourcePath);
+		if (FolderColor.IsSet())
 		{
-			SaveColor(SourcePath, nullptr);
-			SaveColor(Destination, FolderColor);
+			SetPathColor(SourcePath, TOptional<FLinearColor>());
+			SetPathColor(Destination, FolderColor);
 		}
 
 		ChangedPaths.Add(MakeTuple(SourcePath, Destination));
@@ -1111,13 +1111,23 @@ bool AssetViewUtils::DoesFolderExist(const FString& FolderPath)
 
 const TSharedPtr<FLinearColor> AssetViewUtils::LoadColor(const FString& FolderPath)
 {
-	auto LoadColorInternal = [](const FString& InPath) -> TSharedPtr<FLinearColor>
+	TOptional<FLinearColor> FoundColor = GetPathColor(FolderPath);
+	if (FoundColor.IsSet())
+	{
+		return MakeShared<FLinearColor>(FoundColor.GetValue());
+	}
+	return TSharedPtr<FLinearColor>();
+}
+
+TOptional<FLinearColor> AssetViewUtils::GetPathColor(const FString& FolderPath)
+{
+	auto GetPathColorInternal = [](const FString& InPath) -> TOptional<FLinearColor>
 	{
 		// See if we have a value cached first
-		TSharedPtr<FLinearColor> CachedColor = PathColors.FindRef(InPath);
-		if(CachedColor.IsValid())
+		FLinearColor* CachedColor = PathColors.Find(InPath);
+		if(CachedColor)
 		{
-			return CachedColor;
+			return *CachedColor;
 		}
 		
 		// Loads the color of folder at the given path from the config
@@ -1130,21 +1140,21 @@ const TSharedPtr<FLinearColor> AssetViewUtils::LoadColor(const FString& FolderPa
 				FLinearColor Color;
 				if(Color.InitFromString(ColorStr) && !Color.Equals(GetDefaultColor()))
 				{
-					return PathColors.Add(InPath, MakeShareable(new FLinearColor(Color)));
+					return PathColors.Add(InPath, FLinearColor(Color));
 				}
 			}
 			else
 			{
-				return PathColors.Add(InPath, MakeShareable(new FLinearColor(GetDefaultColor())));
+				return PathColors.Add(InPath, FLinearColor(GetDefaultColor()));
 			}
 		}
 
-		return nullptr;
+		return TOptional<FLinearColor>();
 	};
 
 	// First try and find the color using the given path, as this works correctly for both assets and classes
-	TSharedPtr<FLinearColor> FoundColor = LoadColorInternal(FolderPath);
-	if(FoundColor.IsValid())
+	TOptional<FLinearColor> FoundColor = GetPathColorInternal(FolderPath);
+	if(FoundColor.IsSet())
 	{
 		return FoundColor;
 	}
@@ -1154,21 +1164,35 @@ const TSharedPtr<FLinearColor> AssetViewUtils::LoadColor(const FString& FolderPa
 		FString RelativePath;
 		if (FPackageName::TryConvertLongPackageNameToFilename(FolderPath / FString(), RelativePath))
 		{
-			return LoadColorInternal(RelativePath);
+			return GetPathColorInternal(RelativePath);
 		}
 	}
 
-	return nullptr;
+	return TOptional<FLinearColor>();
 }
 
 void AssetViewUtils::SaveColor(const FString& FolderPath, const TSharedPtr<FLinearColor>& FolderColor, bool bForceAdd)
 {
-	auto SaveColorInternal = [](const FString& InPath, const TSharedPtr<FLinearColor>& InFolderColor)
+	TOptional<FLinearColor> OptionalFolderColor;
+	if (FolderColor)
+	{
+		OptionalFolderColor = *FolderColor.Get();
+	}
+	else if (bForceAdd)
+	{
+		OptionalFolderColor = GetDefaultColor();
+	}
+	SetPathColor(FolderPath, OptionalFolderColor);
+}
+
+void AssetViewUtils::SetPathColor(const FString& FolderPath, TOptional<FLinearColor> FolderColor)
+{
+	auto SetPathColorInternal = [](const FString& InPath, FLinearColor InFolderColor)
 	{
 		// Saves the color of the folder to the config
 		if(FPaths::FileExists(GEditorPerProjectIni))
 		{
-			GConfig->SetString(TEXT("PathColor"), *InPath, *InFolderColor->ToString(), GEditorPerProjectIni);
+			GConfig->SetString(TEXT("PathColor"), *InPath, *InFolderColor.ToString(), GEditorPerProjectIni);
 		}
 
 		// Update the map too
@@ -1188,15 +1212,14 @@ void AssetViewUtils::SaveColor(const FString& FolderPath, const TSharedPtr<FLine
 	};
 
 	// Remove the color if it's invalid or default
-	const bool bRemove = !FolderColor.IsValid() || (!bForceAdd && FolderColor->Equals(GetDefaultColor()));
-
+	const bool bRemove = !FolderColor.IsSet() || FolderColor->Equals(GetDefaultColor());
 	if(bRemove)
 	{
 		RemoveColorInternal(FolderPath);
 	}
 	else
 	{
-		SaveColorInternal(FolderPath, FolderColor);
+		SetPathColorInternal(FolderPath, FolderColor.GetValue());
 	}
 
 	// Make sure and remove any colors using the legacy path format
diff --git a/Engine/Source/Developer/AssetTools/Public/AssetViewUtils.h b/Engine/Source/Developer/AssetTools/Public/AssetViewUtils.h
--- a/Engine/Source/Developer/AssetTools/Public/AssetViewUtils.h
+++ b/Engine/Source/Developer/AssetTools/Public/AssetViewUtils.h
@@ -144,8 +144,17 @@ namespace AssetViewUtils
 	 * @param FolderPath - The path to the folder
 	 * @return The color the folder should appear as, will be NULL if not customized
 	 */
+	UE_DEPRECATED(5.0, "LoadColor deprecated, please use GetPathColor")
 	ASSETTOOLS_API const TSharedPtr<FLinearColor> LoadColor(const FString& FolderPath);
 
+	/**
+	 * Returns the color of this path from the config
+	 *
+	 * @param FolderPath - The path to the folder
+	 * @return The color the folder should appear as, will be invalid if not customized
+	 */
+	ASSETTOOLS_API TOptional<FLinearColor> GetPathColor(const FString& FolderPath);
+
 	/**
 	 * Saves the color of the path to the config
 	 *
@@ -153,8 +162,17 @@ namespace AssetViewUtils
 	 * @param FolderColor - The color the folder should appear as
 	 * @param bForceAdd - If true, force the color to be added for the path
 	 */
+	UE_DEPRECATED(5.0, "SaveColor deprecated, please use SetPathColor")
 	ASSETTOOLS_API void SaveColor(const FString& FolderPath, const TSharedPtr<FLinearColor>& FolderColor, bool bForceAdd = false);
 
+	/**
+	 * Saves the color of the path to the config
+	 *
+	 * @param FolderPath - The path to the folder
+	 * @param FolderColor - The color the folder should appear as
+	 */
+	ASSETTOOLS_API void SetPathColor(const FString& FolderPath, TOptional<FLinearColor> FolderColor);
+
 	/**
 	 * Checks to see if any folder has a custom color, optionally outputs them to a list
 	 *
diff --git a/Engine/Source/Developer/ToolWidgets/Private/Filters/SCustomTextFilterDialog.cpp b/Engine/Source/Developer/ToolWidgets/Private/Filters/SCustomTextFilterDialog.cpp
--- a/Engine/Source/Developer/ToolWidgets/Private/Filters/SCustomTextFilterDialog.cpp
+++ b/Engine/Source/Developer/ToolWidgets/Private/Filters/SCustomTextFilterDialog.cpp
@@ -207,14 +207,9 @@ FReply SCustomTextFilterDialog::ColorBlock_OnMouseButtonDown(const FGeometry& My
 {
 	if (MouseEvent.GetEffectingButton() == EKeys::LeftMouseButton)
 	{
-		TArray<FLinearColor*> LinearColorArray;
-		LinearColorArray.Add(&(FilterData.FilterColor));
-
-		FColorPickerArgs PickerArgs;
+		FColorPickerArgs PickerArgs = FColorPickerArgs(FilterData.FilterColor, FOnLinearColorValueChanged::CreateSP(this, &SCustomTextFilterDialog::HandleColorValueChanged));
 		PickerArgs.bIsModal = true;
 		PickerArgs.ParentWidget = ColorBlock;
-		PickerArgs.LinearColorArray = &LinearColorArray;
-
 		OpenColorPicker(PickerArgs);
 
 		return FReply::Handled();
@@ -222,8 +217,12 @@ FReply SCustomTextFilterDialog::ColorBlock_OnMouseButtonDown(const FGeometry& My
 	else
 	{
 		return FReply::Unhandled();
-	}
-		
+	}		
+}
+
+void SCustomTextFilterDialog::HandleColorValueChanged(FLinearColor NewValue)
+{
+	FilterData.FilterColor = NewValue;
 }
 
 bool SCustomTextFilterDialog::CheckFilterValidity() const
diff --git a/Engine/Source/Developer/ToolWidgets/Public/Filters/SCustomTextFilterDialog.h b/Engine/Source/Developer/ToolWidgets/Public/Filters/SCustomTextFilterDialog.h
--- a/Engine/Source/Developer/ToolWidgets/Public/Filters/SCustomTextFilterDialog.h
+++ b/Engine/Source/Developer/ToolWidgets/Public/Filters/SCustomTextFilterDialog.h
@@ -50,6 +50,8 @@ protected:
 
 	/* Handler for when the color block is clicked to open the color picker */
 	FReply ColorBlock_OnMouseButtonDown(const FGeometry& MyGeometry, const FPointerEvent& MouseEvent);
+
+	void HandleColorValueChanged(FLinearColor NewValue);
 	
 	FReply OnDeleteButtonClicked() const;
 	
diff --git a/Engine/Source/Editor/Cascade/Private/Cascade.cpp b/Engine/Source/Editor/Cascade/Private/Cascade.cpp
--- a/Engine/Source/Editor/Cascade/Private/Cascade.cpp
+++ b/Engine/Source/Editor/Cascade/Private/Cascade.cpp
@@ -4069,16 +4069,21 @@ void FCascade::OnBackgroundColor()
 {
 	if (PreviewViewport.IsValid() && PreviewViewport->GetViewportClient().IsValid())
 	{
-		TArray<FColor*> FColorArray;
-
-		FColorArray.Add(&GetParticleSystem()->BackgroundColor);
-
-		FColorPickerArgs PickerArgs;
-		PickerArgs.ParentWidget = PreviewViewport;
-		PickerArgs.DisplayGamma = TAttribute<float>::Create( TAttribute<float>::FGetter::CreateUObject(GEngine, &UEngine::GetDisplayGamma) );
-		PickerArgs.ColorArray = &FColorArray;
+		if (UParticleSystem* LocalParticleSystem = GetParticleSystem())
+		{
+			TWeakObjectPtr<UParticleSystem> WeakParticleSystem = LocalParticleSystem;
+			FColorPickerArgs PickerArgs = FColorPickerArgs(LocalParticleSystem->BackgroundColor, FOnLinearColorValueChanged::CreateLambda([WeakParticleSystem](FLinearColor NewValue)
+				{
+					if (UParticleSystem* PinnedParticleSystem = WeakParticleSystem.Get())
+					{
+						PinnedParticleSystem->BackgroundColor = NewValue.ToFColorSRGB();
+					}
+				}));
+			PickerArgs.ParentWidget = PreviewViewport;
+			PickerArgs.DisplayGamma = TAttribute<float>::Create(TAttribute<float>::FGetter::CreateUObject(GEngine, &UEngine::GetDisplayGamma));
 
-		OpenColorPicker(PickerArgs);
+			OpenColorPicker(PickerArgs);
+		}
 	}
 }
 
diff --git a/Engine/Source/Editor/Cascade/Private/CascadeEmitterCanvasClient.cpp b/Engine/Source/Editor/Cascade/Private/CascadeEmitterCanvasClient.cpp
--- a/Engine/Source/Editor/Cascade/Private/CascadeEmitterCanvasClient.cpp
+++ b/Engine/Source/Editor/Cascade/Private/CascadeEmitterCanvasClient.cpp
@@ -363,10 +363,12 @@ bool FCascadeEmitterCanvasClient::InputKey(const FInputKeyEventArgs& EventArgs)
 
 						if (Module || Emitter)
 						{
-							TArray<FColor*> FColorArray;
+							TWeakObjectPtr<UObject> WeakTarget;
+							TOptional<FLinearColor> InitialColor;
 							if (Module)
 							{
-								FColorArray.Add(&Module->ModuleEditorColor);
+								InitialColor = Module->ModuleEditorColor;
+								WeakTarget = Module;
 							}
 							else
 							{
@@ -374,19 +376,33 @@ bool FCascadeEmitterCanvasClient::InputKey(const FInputKeyEventArgs& EventArgs)
 								UParticleLODLevel* LODLevel = CascadePtr.Pin()->GetCurrentlySelectedLODLevel(Emitter);
 								if ( LODLevel )
 								{
-									FColorArray.Add(&Emitter->EmitterEditorColor);
+									InitialColor = Emitter->EmitterEditorColor;
+									WeakTarget = Emitter;
 								}
 							}
 
-							if ( FColorArray.Num() > 0 )
+							if (InitialColor.IsSet())
 							{
 								// Let go of the mouse lock...
 								EventArgs.Viewport->LockMouseToViewport(false);
 								EventArgs.Viewport->CaptureMouse(false);
 
-								FColorPickerArgs PickerArgs;
+								FColorPickerArgs PickerArgs = FColorPickerArgs(InitialColor.GetValue(), FOnLinearColorValueChanged::CreateLambda([WeakTarget](FLinearColor NewValue)
+									{
+										if (UObject* Target = WeakTarget.Get())
+										{
+											if (UParticleModule* Module = Cast<UParticleModule>(Target))
+											{
+												Module->ModuleEditorColor = NewValue.ToFColorSRGB();
+											}
+											else if (UParticleEmitter* Emitter = CastChecked<UParticleEmitter>(Target))
+											{
+												Emitter->EmitterEditorColor = NewValue.ToFColorSRGB();
+											}
+
+										}
+									}));
 								PickerArgs.DisplayGamma = TAttribute<float>::Create( TAttribute<float>::FGetter::CreateUObject(GEngine, &UEngine::GetDisplayGamma) );
-								PickerArgs.ColorArray = &FColorArray;
 
 								OpenColorPicker(PickerArgs);
 							}
diff --git a/Engine/Source/Editor/ContentBrowser/Private/AssetViewWidgets.cpp b/Engine/Source/Editor/ContentBrowser/Private/AssetViewWidgets.cpp
--- a/Engine/Source/Editor/ContentBrowser/Private/AssetViewWidgets.cpp
+++ b/Engine/Source/Editor/ContentBrowser/Private/AssetViewWidgets.cpp
@@ -1495,9 +1495,9 @@ FSlateColor SAssetViewItem::GetAssetColor() const
 			}
 			else
 			{
-				if (TSharedPtr<FLinearColor> Color = ContentBrowserUtils::LoadColor(AssetItem->GetItem().GetInvariantPath().ToString()))
+				if (TOptional<FLinearColor> Color = ContentBrowserUtils::GetPathColor(AssetItem->GetItem().GetInvariantPath().ToString()))
 				{
-					return *Color;
+					return Color.GetValue();
 				}
 			}
 		}
diff --git a/Engine/Source/Editor/ContentBrowser/Private/PathContextMenu.cpp b/Engine/Source/Editor/ContentBrowser/Private/PathContextMenu.cpp
--- a/Engine/Source/Editor/ContentBrowser/Private/PathContextMenu.cpp
+++ b/Engine/Source/Editor/ContentBrowser/Private/PathContextMenu.cpp
@@ -368,35 +368,33 @@ void FPathContextMenu::ExecuteResetColor()
 
 void FPathContextMenu::ExecutePickColor()
 {
+	if (SelectedFolders.Num() == 0)
+	{
+		return;
+	}
+
 	// Spawn a color picker, so the user can select which color they want
-	TArray<FLinearColor*> LinearColorArray;
-	FColorPickerArgs PickerArgs;
-	PickerArgs.bIsModal = false;
-	PickerArgs.ParentWidget = ParentContent.Pin();
+	FLinearColor InitialColor = ContentBrowserUtils::GetDefaultColor();
 	if (SelectedFolders.Num() > 0)
 	{
 		// Make sure an color entry exists for all the paths, otherwise they won't update in realtime with the widget color
-		for (int32 FolderIndex = SelectedFolders.Num() - 1; FolderIndex >= 0; --FolderIndex)
+		for (const FContentBrowserItem& SelectedItem : SelectedFolders)
 		{
-			const FString Path = SelectedFolders[FolderIndex].GetInvariantPath().ToString();
+			const FString Path = SelectedItem.GetInvariantPath().ToString();
 
-			TSharedPtr<FLinearColor> Color = ContentBrowserUtils::LoadColor(Path);
-			if (!Color.IsValid())
-			{
-				Color = MakeShareable(new FLinearColor(ContentBrowserUtils::GetDefaultColor()));
-				ContentBrowserUtils::SaveColor(Path, Color, true);
-			}
-			else
+			TOptional<FLinearColor> Color = ContentBrowserUtils::GetPathColor(Path);
+			if (Color.IsSet())
 			{
 				// Default the color to the first valid entry
-				PickerArgs.InitialColorOverride = *Color.Get();
+				InitialColor = Color.GetValue();
+				break;
 			}
-			LinearColorArray.Add(Color.Get());
 		}
-		PickerArgs.LinearColorArray = &LinearColorArray;
 	}
 
-	PickerArgs.OnColorPickerWindowClosed = FOnWindowClosed::CreateSP(this, &FPathContextMenu::NewColorComplete);
+	FColorPickerArgs PickerArgs = FColorPickerArgs(InitialColor, FOnLinearColorValueChanged::CreateSP(this, &FPathContextMenu::OnLinearColorValueChanged));
+	PickerArgs.bIsModal = false;
+	PickerArgs.ParentWidget = ParentContent.Pin();
 
 	OpenColorPicker(PickerArgs);
 }
@@ -423,16 +421,9 @@ void FPathContextMenu::ExecutePrivateContentEdit()
 	OnPrivateContentEditToggled.ExecuteIfBound(PathsToUpdate);
 }
 
-void FPathContextMenu::NewColorComplete(const TSharedRef<SWindow>& Window)
+void FPathContextMenu::OnLinearColorValueChanged(const FLinearColor InColor)
 {
-	// Save the colors back in the config (ptr should have already updated by the widget)
-	for (const FContentBrowserItem& SelectedItem : SelectedFolders)
-	{
-		const FString Path = SelectedItem.GetInvariantPath().ToString();
-		const TSharedPtr<FLinearColor> Color = ContentBrowserUtils::LoadColor(Path);
-		check(Color.IsValid());
-		ContentBrowserUtils::SaveColor(Path, Color);
-	}
+	OnColorClicked(InColor);
 }
 
 FReply FPathContextMenu::OnColorClicked( const FLinearColor InColor )
@@ -441,13 +432,7 @@ FReply FPathContextMenu::OnColorClicked( const FLinearColor InColor )
 	for (const FContentBrowserItem& SelectedItem : SelectedFolders)
 	{
 		const FString Path = SelectedItem.GetInvariantPath().ToString();
-		TSharedPtr<FLinearColor> Color = ContentBrowserUtils::LoadColor(Path);
-		if (!Color.IsValid())
-		{
-			Color = MakeShareable(new FLinearColor());
-		}
-		*Color.Get() = InColor;
-		ContentBrowserUtils::SaveColor(Path, Color);
+		ContentBrowserUtils::SetPathColor(Path, InColor);
 	}
 
 	// Dismiss the menu here, as we can't make the 'clear' option appear if a folder has just had a color set for the first time
@@ -461,7 +446,7 @@ void FPathContextMenu::ResetColors()
 	// Clear the custom colors for all the selected paths
 	for (const FContentBrowserItem& SelectedItem : SelectedFolders)
 	{
-		ContentBrowserUtils::SaveColor(SelectedItem.GetInvariantPath().ToString(), nullptr);
+		ContentBrowserUtils::SetPathColor(SelectedItem.GetInvariantPath().ToString(), TOptional<FLinearColor>());
 	}
 }
 
@@ -610,11 +595,13 @@ bool FPathContextMenu::SelectedHasCustomColors() const
 {
 	for (const FContentBrowserItem& SelectedItem : SelectedFolders)
 	{
-		// Ignore any that are the default color
-		const TSharedPtr<FLinearColor> Color = ContentBrowserUtils::LoadColor(SelectedItem.GetInvariantPath().ToString());
-		if (Color.IsValid() && !Color->Equals(ContentBrowserUtils::GetDefaultColor()))
+		if (const TOptional<FLinearColor> Color = ContentBrowserUtils::GetPathColor(SelectedItem.GetInvariantPath().ToString()))
 		{
-			return true;
+			// Ignore any that are the default color, in case the user used the deprecated SaveColor with bForce
+			if (!Color->Equals(ContentBrowserUtils::GetDefaultColor()))
+			{
+				return true;
+			}
 		}
 	}
 	return false;
diff --git a/Engine/Source/Editor/ContentBrowser/Private/PathContextMenu.h b/Engine/Source/Editor/ContentBrowser/Private/PathContextMenu.h
--- a/Engine/Source/Editor/ContentBrowser/Private/PathContextMenu.h
+++ b/Engine/Source/Editor/ContentBrowser/Private/PathContextMenu.h
@@ -108,8 +108,8 @@ private:
 	/** Checks to see if any of the selected paths use custom colors */
 	bool SelectedHasCustomColors() const;
 
-	/** Callback when the color picker dialog has been closed */
-	void NewColorComplete(const TSharedRef<SWindow>& Window);
+	/** Callback when the color picker dialog changed the color. */
+	void OnLinearColorValueChanged(const FLinearColor InColor);
 
 	/** Callback when the color is picked from the set color submenu */
 	FReply OnColorClicked( const FLinearColor InColor );
diff --git a/Engine/Source/Editor/ContentBrowser/Private/SourcesViewWidgets.cpp b/Engine/Source/Editor/ContentBrowser/Private/SourcesViewWidgets.cpp
--- a/Engine/Source/Editor/ContentBrowser/Private/SourcesViewWidgets.cpp
+++ b/Engine/Source/Editor/ContentBrowser/Private/SourcesViewWidgets.cpp
@@ -241,9 +241,10 @@ FSlateColor SAssetTreeItem::GetFolderColor() const
 		}
 		else
 		{
-			if (TSharedPtr<FLinearColor> Color = ContentBrowserUtils::LoadColor(TreeItemPin->GetItem().GetInvariantPath().ToString()))
+			TOptional<FLinearColor> Color = ContentBrowserUtils::GetPathColor(TreeItemPin->GetItem().GetInvariantPath().ToString());
+			if (Color.IsSet())
 			{
-				FoundColor = *Color;
+				FoundColor = Color.GetValue();
 			}
 		}
 
diff --git a/Engine/Source/Editor/CurveEditor/Private/CurveEditor.cpp b/Engine/Source/Editor/CurveEditor/Private/CurveEditor.cpp
--- a/Engine/Source/Editor/CurveEditor/Private/CurveEditor.cpp
+++ b/Engine/Source/Editor/CurveEditor/Private/CurveEditor.cpp
@@ -1910,21 +1910,27 @@ void FCurveEditor::SetCurveColorsForSelected()
 		return;
 	}
 
+	TWeakPtr<FCurveEditor> WeakSelf = AsShared();
+
 	FColorPickerArgs PickerArgs;
 	PickerArgs.bUseAlpha = false;
 	PickerArgs.InitialColorOverride = FindCurve(*CurveModelIDs.CreateIterator())->GetColor();
-	PickerArgs.OnColorCommitted.BindLambda([this, CurveModelIDs](FLinearColor NewColor) {
-		for (const FCurveModelID& CurveModelID : CurveModelIDs)
+	PickerArgs.OnColorCommitted.BindLambda([WeakSelf, CurveModelIDs](FLinearColor NewColor)
+	{
+		if (TSharedPtr<FCurveEditor> Self = WeakSelf.Pin())
 		{
-			if (FCurveModel* Curve = FindCurve(CurveModelID))
+			for (const FCurveModelID& CurveModelID : CurveModelIDs)
 			{
-				UObject* Object = nullptr;
-				FString Name;
-				Curve->GetCurveColorObjectAndName(&Object, Name);
-				if (Object)
+				if (FCurveModel* Curve = Self->FindCurve(CurveModelID))
 				{
-					Settings->SetCustomColor(Object->GetClass(), Name, NewColor);
-					Curve->SetColor(NewColor);
+					UObject* Object = nullptr;
+					FString Name;
+					Curve->GetCurveColorObjectAndName(&Object, Name);
+					if (Object)
+					{
+						Self->Settings->SetCustomColor(Object->GetClass(), Name, NewColor);
+						Curve->SetColor(NewColor);
+					}
 				}
 			}
 		}
diff --git a/Engine/Source/Editor/DistCurveEditor/Private/SDistributionCurveEditor.cpp b/Engine/Source/Editor/DistCurveEditor/Private/SDistributionCurveEditor.cpp
--- a/Engine/Source/Editor/DistCurveEditor/Private/SDistributionCurveEditor.cpp
+++ b/Engine/Source/Editor/DistCurveEditor/Private/SDistributionCurveEditor.cpp
@@ -650,14 +650,9 @@ void SDistributionCurveEditor::OnSetColor()
 	}
 
 	//since the data isn't stored in standard colors, a temp color is used
-	FColor TempColor = InputColor;
-
-	TArray<FColor*> FColorArray;
-	FColorArray.Add(&TempColor);
-
-	FColorPickerArgs PickerArgs;
+	FLinearColor TempColor = InputColor;
+	FColorPickerArgs PickerArgs = FColorPickerArgs(InputColor, FOnLinearColorValueChanged::CreateLambda([&TempColor](FLinearColor NewValue){ TempColor = NewValue.ToFColorSRGB(); }));
 	PickerArgs.bIsModal = true;
-	PickerArgs.ColorArray = &FColorArray;
 	PickerArgs.DisplayGamma = TAttribute<float>::Create( TAttribute<float>::FGetter::CreateUObject(GEngine, &UEngine::GetDisplayGamma) );
 
 	if (OpenColorPicker(PickerArgs))
diff --git a/Engine/Source/Editor/FontEditor/Private/FontEditor.cpp b/Engine/Source/Editor/FontEditor/Private/FontEditor.cpp
--- a/Engine/Source/Editor/FontEditor/Private/FontEditor.cpp
+++ b/Engine/Source/Editor/FontEditor/Private/FontEditor.cpp
@@ -922,21 +922,23 @@ bool FFontEditor::OnExportAllEnabled() const
 
 void FFontEditor::OnBackgroundColor()
 {
-	FColor Color = FontPreviewWidget->GetPreviewBackgroundColor();
-	TArray<FColor*> FColorArray;
-	FColorArray.Add(&Color);
+	TWeakPtr<SFontEditorViewport> WeakFontPreviewWidget = FontPreviewWidget;
 
 	FColorPickerArgs PickerArgs;
 	PickerArgs.bIsModal = true;
 	PickerArgs.ParentWidget = FontPreview;
 	PickerArgs.bUseAlpha = true;
 	PickerArgs.DisplayGamma = TAttribute<float>::Create( TAttribute<float>::FGetter::CreateUObject(GEngine, &UEngine::GetDisplayGamma) );
-	PickerArgs.ColorArray = &FColorArray;
+	PickerArgs.InitialColorOverride = FontPreviewWidget->GetPreviewBackgroundColor();
+	PickerArgs.OnColorCommitted = FOnLinearColorValueChanged::CreateLambda([WeakFontPreviewWidget](FLinearColor NewValue)
+		{
+			if (TSharedPtr<SFontEditorViewport> PinnedWidget = WeakFontPreviewWidget.Pin())
+			{
+				PinnedWidget->SetPreviewBackgroundColor(NewValue.ToFColorSRGB());
+			}
+		});
 
-	if (OpenColorPicker(PickerArgs))
-	{
-		FontPreviewWidget->SetPreviewBackgroundColor(Color);
-	}
+	OpenColorPicker(PickerArgs);
 }
 
 bool FFontEditor::OnBackgroundColorEnabled() const
@@ -947,21 +949,23 @@ bool FFontEditor::OnBackgroundColorEnabled() const
 
 void FFontEditor::OnForegroundColor()
 {
-	FColor Color = FontPreviewWidget->GetPreviewForegroundColor();
-	TArray<FColor*> FColorArray;
-	FColorArray.Add(&Color);
+	TWeakPtr<SFontEditorViewport> WeakFontPreviewWidget = FontPreviewWidget;
 
 	FColorPickerArgs PickerArgs;
 	PickerArgs.bIsModal = true;
 	PickerArgs.ParentWidget = FontPreview;
 	PickerArgs.bUseAlpha = true;
 	PickerArgs.DisplayGamma = TAttribute<float>::Create( TAttribute<float>::FGetter::CreateUObject(GEngine, &UEngine::GetDisplayGamma) );
-	PickerArgs.ColorArray = &FColorArray;
+	PickerArgs.InitialColorOverride = FontPreviewWidget->GetPreviewForegroundColor();
+	PickerArgs.OnColorCommitted = FOnLinearColorValueChanged::CreateLambda([WeakFontPreviewWidget](FLinearColor NewValue)
+		{
+			if (TSharedPtr<SFontEditorViewport> PinnedWidget = WeakFontPreviewWidget.Pin())
+			{
+				PinnedWidget->SetPreviewForegroundColor(NewValue.ToFColorSRGB());
+			}
+		});
 
-	if (OpenColorPicker(PickerArgs))
-	{
-		FontPreviewWidget->SetPreviewForegroundColor(Color);
-	}
+	OpenColorPicker(PickerArgs);
 }
 
 bool FFontEditor::OnForegroundColorEnabled() const
diff --git a/Engine/Source/Editor/GraphEditor/Private/KismetPins/SGraphPinColor.cpp b/Engine/Source/Editor/GraphEditor/Private/KismetPins/SGraphPinColor.cpp
--- a/Engine/Source/Editor/GraphEditor/Private/KismetPins/SGraphPinColor.cpp
+++ b/Engine/Source/Editor/GraphEditor/Private/KismetPins/SGraphPinColor.cpp
@@ -58,15 +58,11 @@ FReply SGraphPinColor::OnColorBoxClicked(const FGeometry& MyGeometry, const FPoi
 {
 	if (MouseEvent.GetEffectingButton() == EKeys::LeftMouseButton)
 	{
-		SelectedColor = GetColor();
-		TArray<FLinearColor*> LinearColorArray;
-		LinearColorArray.Add(&SelectedColor);
-
 		FColorPickerArgs PickerArgs;
 		PickerArgs.bIsModal = true;
 		PickerArgs.ParentWidget = DefaultValueWidget;
 		PickerArgs.DisplayGamma = TAttribute<float>::Create(TAttribute<float>::FGetter::CreateUObject(GEngine, &UEngine::GetDisplayGamma));
-		PickerArgs.LinearColorArray = &LinearColorArray;
+		PickerArgs.InitialColorOverride = GetColor();
 		PickerArgs.OnColorCommitted = FOnLinearColorValueChanged::CreateSP(this, &SGraphPinColor::OnColorCommitted);
 		PickerArgs.bUseAlpha = true;
 
diff --git a/Engine/Source/Editor/GraphEditor/Public/KismetPins/SGraphPinColor.h b/Engine/Source/Editor/GraphEditor/Public/KismetPins/SGraphPinColor.h
--- a/Engine/Source/Editor/GraphEditor/Public/KismetPins/SGraphPinColor.h
+++ b/Engine/Source/Editor/GraphEditor/Public/KismetPins/SGraphPinColor.h
@@ -40,7 +40,5 @@ protected:
 
 
 private:
-	/** Current selected color */
-	FLinearColor SelectedColor;
 	TSharedPtr<SWidget> DefaultValueWidget;
 };
diff --git a/Engine/Source/Editor/Kismet/Private/BlueprintDetailsCustomization.cpp b/Engine/Source/Editor/Kismet/Private/BlueprintDetailsCustomization.cpp
--- a/Engine/Source/Editor/Kismet/Private/BlueprintDetailsCustomization.cpp
+++ b/Engine/Source/Editor/Kismet/Private/BlueprintDetailsCustomization.cpp
@@ -5565,14 +5565,23 @@ FReply FBlueprintGraphActionDetails::ColorBlock_OnMouseButtonDown(const FGeometr
 	{
 		if (FKismetUserDeclaredFunctionMetadata* Metadata = GetMetadataBlock())
 		{
-			TArray<FLinearColor*> LinearColorArray;
-			LinearColorArray.Add(&(Metadata->InstanceTitleColor));
+			TWeakPtr<FBlueprintGraphActionDetails> WeakSelf = SharedThis(this);
 
 			FColorPickerArgs PickerArgs;
 			PickerArgs.bIsModal = true;
 			PickerArgs.ParentWidget = ColorBlock;
 			PickerArgs.DisplayGamma = TAttribute<float>::Create( TAttribute<float>::FGetter::CreateUObject(GEngine, &UEngine::GetDisplayGamma) );
-			PickerArgs.LinearColorArray = &LinearColorArray;
+			PickerArgs.InitialColorOverride = Metadata->InstanceTitleColor;
+			PickerArgs.OnColorCommitted = FOnLinearColorValueChanged::CreateLambda([WeakSelf](FLinearColor NewValue)
+			{
+				if (TSharedPtr<FBlueprintGraphActionDetails> Self = WeakSelf.Pin())
+				{
+					if (FKismetUserDeclaredFunctionMetadata* Metadata = Self->GetMetadataBlock())
+					{
+						Metadata->InstanceTitleColor = NewValue;
+					}
+				}
+			});
 
 			OpenColorPicker(PickerArgs);
 		}
diff --git a/Engine/Source/Editor/MaterialEditor/Private/MaterialEditor.cpp b/Engine/Source/Editor/MaterialEditor/Private/MaterialEditor.cpp
--- a/Engine/Source/Editor/MaterialEditor/Private/MaterialEditor.cpp
+++ b/Engine/Source/Editor/MaterialEditor/Private/MaterialEditor.cpp
@@ -542,7 +542,6 @@ void FMaterialEditor::InitMaterialEditor( const EToolkitMode::Type Mode, const T
 	EditorOptions = NULL;
 	bMaterialDirty = false;
 	bStatsFromPreviewMaterial = false;
-	ColorPickerObject = NULL;
 
 	// Support undo/redo
 	Material->SetFlags(RF_Transactional);
@@ -6711,42 +6710,61 @@ FMatExpressionPreview* FMaterialEditor::GetExpressionPreview(UMaterialExpression
 	return NULL;
 }
 
-void FMaterialEditor::PreColorPickerCommit(FLinearColor LinearColor)
-{
-	// Begin a property edit transaction.
-	if ( GEditor )
+void FMaterialEditor::OnColorPickerCommitted(FLinearColor LinearColor, TWeakObjectPtr<UObject> ColorPickerObject)
+{	
+	if (UObject* Object = ColorPickerObject.Get())
 	{
-		GEditor->BeginTransaction( LOCTEXT("ModifyColorPicker", "Modify Color Picker Value") );
-	}
+		// Begin a property edit transaction.
+		if (GEditor)
+		{
+			GEditor->BeginTransaction(LOCTEXT("ModifyColorPicker", "Modify Color Picker Value"));
+		}
 
-	NotifyPreChange(NULL);
+		NotifyPreChange(NULL);
 
-	UObject* Object = ColorPickerObject.Get(false);
-	if( Object )
-	{
 		Object->PreEditChange(NULL);
-	}
-}
 
-void FMaterialEditor::OnColorPickerCommitted(FLinearColor LinearColor)
-{	
-	UObject* Object = ColorPickerObject.Get(false);
-	if( Object )
-	{
+		FName PropertyName;
+		if (UMaterialExpressionConstant3Vector* Constant3Expression = Cast<UMaterialExpressionConstant3Vector>(Object))
+		{
+			Constant3Expression->Constant = LinearColor;
+			PropertyName = TEXT("Constant");
+		}
+		else if (UMaterialExpressionConstant4Vector* Constant4Expression = Cast<UMaterialExpressionConstant4Vector>(Object))
+		{
+			Constant4Expression->Constant = LinearColor;
+			PropertyName = TEXT("Constant");
+		}
+		else if (UMaterialExpressionFunctionInput* InputExpression = Cast<UMaterialExpressionFunctionInput>(Object))
+		{
+			InputExpression->PreviewValue = LinearColor;
+		}
+		else if (UMaterialExpressionVectorParameter* VectorExpression = Cast<UMaterialExpressionVectorParameter>(Object))
+		{
+			VectorExpression->DefaultValue = LinearColor;
+			PropertyName = TEXT("DefaultValue");
+		}
+		else
+		{
+			checkf(false, TEXT("The expression type is supported in OnNodeDoubleClicked but not in OnColorPickerCommitted"));
+		}
+
 		Object->MarkPackageDirty();
-		FPropertyChangedEvent Event(ColorPickerProperty.Get(false));
+
+		FProperty* ColorPickerProperty = PropertyName.IsNone() ? nullptr : Object->GetClass()->FindPropertyByName(PropertyName);
+		FPropertyChangedEvent Event(ColorPickerProperty);
 		Object->PostEditChangeProperty(Event);
-	}
 
-	NotifyPostChange(NULL,NULL);
+		NotifyPostChange(NULL, NULL);
 
-	if ( GEditor )
-	{
-		GEditor->EndTransaction();
-	}
+		if (GEditor)
+		{
+			GEditor->EndTransaction();
+		}
 
-	RefreshExpressionPreviews();
-	MaterialCustomPrimitiveDataWidget->UpdateEditorInstance(MaterialEditorInstance);
+		RefreshExpressionPreviews();
+		MaterialCustomPrimitiveDataWidget->UpdateEditorInstance(MaterialEditorInstance);
+	}
 }
 
 /** Create new tab for the supplied graph - don't call this directly, instead call OpenDocument to track history.*/
@@ -7265,69 +7283,38 @@ void FMaterialEditor::OnNodeDoubleClicked(class UEdGraphNode* Node)
 	}
 	else if (GraphNode && GraphNode->MaterialExpression)
 	{
-		UMaterialExpressionConstant3Vector* Constant3Expression = Cast<UMaterialExpressionConstant3Vector>(GraphNode->MaterialExpression);
-		UMaterialExpressionConstant4Vector* Constant4Expression = Cast<UMaterialExpressionConstant4Vector>(GraphNode->MaterialExpression);
-		UMaterialExpressionFunctionInput* InputExpression = Cast<UMaterialExpressionFunctionInput>(GraphNode->MaterialExpression);
-		UMaterialExpressionVectorParameter* VectorExpression = Cast<UMaterialExpressionVectorParameter>(GraphNode->MaterialExpression);
-
-		FColorChannels ChannelEditStruct;
-
-		// Reset to default
-		ColorPickerProperty = NULL;
-
-		if( Constant3Expression )
+		TOptional<FLinearColor> InitialColor;
+		bool bCanEditAlpha = true;
+		if(UMaterialExpressionConstant3Vector* Constant3Expression = Cast<UMaterialExpressionConstant3Vector>(GraphNode->MaterialExpression))
 		{
-			ChannelEditStruct.Red = &Constant3Expression->Constant.R;
-			ChannelEditStruct.Green = &Constant3Expression->Constant.G;
-			ChannelEditStruct.Blue = &Constant3Expression->Constant.B;
-			static FName ConstantName = FName(TEXT("Constant"));
-			ColorPickerProperty = Constant3Expression->GetClass()->FindPropertyByName(ConstantName);
+			InitialColor = Constant3Expression->Constant;
+			bCanEditAlpha = false;
 		}
-		else if( Constant4Expression )
+		else if(UMaterialExpressionConstant4Vector* Constant4Expression = Cast<UMaterialExpressionConstant4Vector>(GraphNode->MaterialExpression))
 		{
-			ChannelEditStruct.Red = &Constant4Expression->Constant.R;
-			ChannelEditStruct.Green = &Constant4Expression->Constant.G;
-			ChannelEditStruct.Blue = &Constant4Expression->Constant.B;
-			ChannelEditStruct.Alpha = &Constant4Expression->Constant.A;
-			static FName ConstantName = FName(TEXT("Constant"));
-			ColorPickerProperty = Constant4Expression->GetClass()->FindPropertyByName(ConstantName);
+			InitialColor = Constant4Expression->Constant;
 		}
-		else if (InputExpression)
+		else if (UMaterialExpressionFunctionInput* InputExpression = Cast<UMaterialExpressionFunctionInput>(GraphNode->MaterialExpression))
 		{
-			ChannelEditStruct.Red = &InputExpression->PreviewValue.X;
-			ChannelEditStruct.Green = &InputExpression->PreviewValue.Y;
-			ChannelEditStruct.Blue = &InputExpression->PreviewValue.Z;
-			ChannelEditStruct.Alpha = &InputExpression->PreviewValue.W;
+			InitialColor = InputExpression->PreviewValue;
 		}
-		else if (VectorExpression)
+		else if (UMaterialExpressionVectorParameter* VectorExpression = Cast<UMaterialExpressionVectorParameter>(GraphNode->MaterialExpression))
 		{
-			ChannelEditStruct.Red = &VectorExpression->DefaultValue.R;
-			ChannelEditStruct.Green = &VectorExpression->DefaultValue.G;
-			ChannelEditStruct.Blue = &VectorExpression->DefaultValue.B;
-			ChannelEditStruct.Alpha = &VectorExpression->DefaultValue.A;
-			static FName DefaultValueName = FName(TEXT("DefaultValue"));
-			// Store off the property the color picker will be manipulating, so we can construct a useful PostEditChangeProperty later
-			ColorPickerProperty = VectorExpression->GetClass()->FindPropertyByName(DefaultValueName);
+			InitialColor = VectorExpression->DefaultValue;
 		}
 
-		if (ChannelEditStruct.Red || ChannelEditStruct.Green || ChannelEditStruct.Blue || ChannelEditStruct.Alpha)
+		if (InitialColor.IsSet())
 		{
-			TArray<FColorChannels> Channels;
-			Channels.Add(ChannelEditStruct);
-
-			ColorPickerObject = GraphNode->MaterialExpression;
+			TWeakObjectPtr<UObject> ColorPickerObject = GraphNode->MaterialExpression;
 
 			// Open a color picker 
-			FColorPickerArgs PickerArgs;
+			FColorPickerArgs PickerArgs = FColorPickerArgs(InitialColor.GetValue(), FOnLinearColorValueChanged::CreateSP(this, &FMaterialEditor::OnColorPickerCommitted, ColorPickerObject));
 			PickerArgs.ParentWidget = FocusedGraphEdPtr.Pin();
-			PickerArgs.bUseAlpha = ChannelEditStruct.Alpha != nullptr;
+			PickerArgs.bUseAlpha = bCanEditAlpha;
 			PickerArgs.bOnlyRefreshOnOk = false;
 			PickerArgs.bOnlyRefreshOnMouseUp = true;
 			PickerArgs.bExpandAdvancedSection = true;
 			PickerArgs.DisplayGamma = TAttribute<float>::Create( TAttribute<float>::FGetter::CreateUObject(GEngine, &UEngine::GetDisplayGamma) );
-			PickerArgs.ColorChannelsArray = &Channels;
-			PickerArgs.OnColorCommitted = FOnLinearColorValueChanged::CreateSP(this, &FMaterialEditor::OnColorPickerCommitted);
-			PickerArgs.PreColorCommitted = FOnLinearColorValueChanged::CreateSP(this, &FMaterialEditor::PreColorPickerCommit);
 			PickerArgs.OptionalOwningDetailsView = MaterialDetailsView;
 			OpenColorPicker(PickerArgs);
 		}
diff --git a/Engine/Source/Editor/MaterialEditor/Private/MaterialEditor.h b/Engine/Source/Editor/MaterialEditor/Private/MaterialEditor.h
--- a/Engine/Source/Editor/MaterialEditor/Private/MaterialEditor.h
+++ b/Engine/Source/Editor/MaterialEditor/Private/MaterialEditor.h
@@ -858,15 +858,9 @@ private:
 	 */
 	FMatExpressionPreview* GetExpressionPreview(UMaterialExpression* MaterialExpression, bool& bNewlyCreated);
 
-	/** Pointer to the object that the current color picker is working on. Can be NULL and stale. */
-	TWeakObjectPtr<UObject> ColorPickerObject;
-	TWeakFieldPtr<FProperty> ColorPickerProperty;
-
-	/** Called before the color picker commits a change. */
-	void PreColorPickerCommit(FLinearColor LinearColor);
 
 	/** Called whenever the color picker is used and accepted. */
-	void OnColorPickerCommitted(FLinearColor LinearColor);
+	void OnColorPickerCommitted(FLinearColor LinearColor, TWeakObjectPtr<UObject> ColorPickerObject);
 
 	/** Create new graph editor widget */
 	TSharedRef<class SGraphEditor> CreateGraphEditorWidget(TSharedRef<class FTabInfo> InTabInfo, class UEdGraph* InGraph);
diff --git a/Engine/Source/Editor/PropertyEditor/Private/SDetailsViewBase.cpp b/Engine/Source/Editor/PropertyEditor/Private/SDetailsViewBase.cpp
--- a/Engine/Source/Editor/PropertyEditor/Private/SDetailsViewBase.cpp
+++ b/Engine/Source/Editor/PropertyEditor/Private/SDetailsViewBase.cpp
@@ -387,79 +387,77 @@ void SDetailsViewBase::EnqueueDeferredAction(FSimpleDelegate& DeferredAction)
 	DeferredActions.Add(DeferredAction);
 }
 
-/**
-* Creates the color picker window for this property view.
-*
-* @param Node				The slate property node to edit.
-* @param bUseAlpha			Whether or not alpha is supported
-*/
-void SDetailsViewBase::CreateColorPickerWindow(const TSharedRef< FPropertyEditor >& PropertyEditor, bool bUseAlpha)
+namespace UE::PropertyEditor::Private
+{
+/** Set the color for the property node */
+void SDetailsViewBase_SetColor(FLinearColor NewColor, TWeakPtr<IPropertyHandle> WeakPropertyHandle)
 {
-	const TSharedRef< FPropertyNode > PinnedColorPropertyNode = PropertyEditor->GetPropertyNode();
-	ColorPropertyNode = PinnedColorPropertyNode;
+	if (TSharedPtr<IPropertyHandle> PropertyHandle = WeakPropertyHandle.Pin())
+	{
+		const FProperty* NodeProperty = PropertyHandle->GetProperty();
+		check(NodeProperty);
 
-	FProperty* Property = PinnedColorPropertyNode->GetProperty();
+		const UScriptStruct* Struct = CastField<FStructProperty>(NodeProperty)->Struct;
+		if (Struct->GetFName() == NAME_Color)
+		{
+			const bool bSRGB = true;
+			FColor NewFColor = NewColor.ToFColor(bSRGB);
+			ensure(PropertyHandle->SetValueFromFormattedString(NewFColor.ToString(), EPropertyValueSetFlags::DefaultFlags) == FPropertyAccess::Result::Success);
+		}
+		else
+		{
+			check(Struct->GetFName() == NAME_LinearColor);
+			ensure(PropertyHandle->SetValueFromFormattedString(NewColor.ToString(), EPropertyValueSetFlags::DefaultFlags) == FPropertyAccess::Result::Success);
+		}
+	}
+}
+}//namespace
+
+void SDetailsViewBase::CreateColorPickerWindow(const TSharedRef< FPropertyEditor >& PropertyEditor, bool bUseAlpha)
+{
+	const FProperty* Property = PropertyEditor->GetProperty();
 	check(Property);
 
 	FReadAddressList ReadAddresses;
-	PinnedColorPropertyNode->GetReadAddress(false, ReadAddresses, false);
+	PropertyEditor->GetPropertyNode()->GetReadAddress(false, ReadAddresses, false);
 
-	TArray<FLinearColor*> LinearColor;
-	TArray<FColor*> DWORDColor;
-	for (int32 ColorIndex = 0; ColorIndex < ReadAddresses.Num(); ++ColorIndex)
+	TOptional<FLinearColor> DefaultColor;
+	if (ReadAddresses.Num())
 	{
-		const uint8* Addr = ReadAddresses.GetAddress(ColorIndex);
-		if (Addr)
+		for (int32 ColorIndex = 0; ColorIndex < ReadAddresses.Num(); ++ColorIndex)
 		{
-			if (CastField<FStructProperty>(Property)->Struct->GetFName() == NAME_Color)
-			{
-				DWORDColor.Add((FColor*)Addr);
-			}
-			else
+			const uint8* Addr = ReadAddresses.GetAddress(ColorIndex);
+			if (Addr)
 			{
-				check(CastField<FStructProperty>(Property)->Struct->GetFName() == NAME_LinearColor);
-				LinearColor.Add((FLinearColor*)Addr);
+				if (CastField<FStructProperty>(Property)->Struct->GetFName() == NAME_Color)
+				{
+					DefaultColor = *reinterpret_cast<const FColor*>(Addr);
+				}
+				else
+				{
+					check(CastField<FStructProperty>(Property)->Struct->GetFName() == NAME_LinearColor);
+					DefaultColor = *reinterpret_cast<const FLinearColor*>(Addr);
+				}
 			}
 		}
 	}
 
-	bHasOpenColorPicker = true;
-
-	FColorPickerArgs PickerArgs;
-	PickerArgs.ParentWidget = AsShared();
-	PickerArgs.bUseAlpha = bUseAlpha;
-	PickerArgs.DisplayGamma = TAttribute<float>::Create(TAttribute<float>::FGetter::CreateUObject(GEngine, &UEngine::GetDisplayGamma));
-	PickerArgs.ColorArray = &DWORDColor;
-	PickerArgs.LinearColorArray = &LinearColor;
-	PickerArgs.OnColorCommitted = FOnLinearColorValueChanged::CreateSP(this, &SDetailsViewBase::SetColorPropertyFromColorPicker);
-	PickerArgs.OnColorPickerWindowClosed = FOnWindowClosed::CreateSP(this, &SDetailsViewBase::OnColorPickerWindowClosed);
-	PickerArgs.OptionalOwningDetailsView = AsShared();
-	OpenColorPicker(PickerArgs);
-}
-
-void SDetailsViewBase::SetColorPropertyFromColorPicker(FLinearColor NewColor)
-{
-	const TSharedPtr< FPropertyNode > PinnedColorPropertyNode = ColorPropertyNode.Pin();
-	if (ensure(PinnedColorPropertyNode.IsValid()))
+	if (DefaultColor.IsSet())
 	{
-		FProperty* Property = PinnedColorPropertyNode->GetProperty();
-		check(Property);
-
-		FObjectPropertyNode* ObjectNode = PinnedColorPropertyNode->FindObjectItemParent();
-
-		if (ObjectNode && ObjectNode->GetNumObjects())
-		{
-			FScopedTransaction Transaction(NSLOCTEXT("UnrealEd", "SetColorProperty", "Set Color Property"));
-
-			PinnedColorPropertyNode->NotifyPreChange(Property, GetNotifyHook());
+		bHasOpenColorPicker = true;
+		ColorPropertyNode = PropertyEditor->GetPropertyNode();
 
-			FPropertyChangedEvent ChangeEvent(Property, EPropertyChangeType::ValueSet);
-			PinnedColorPropertyNode->NotifyPostChange(ChangeEvent, GetNotifyHook());
-		}
+		TWeakPtr<IPropertyHandle> WeakPropertyHandle = PropertyEditor->GetPropertyHandle();
+		FColorPickerArgs PickerArgs = FColorPickerArgs(DefaultColor.GetValue(), FOnLinearColorValueChanged::CreateStatic(&UE::PropertyEditor::Private::SDetailsViewBase_SetColor, WeakPropertyHandle));
+		PickerArgs.ParentWidget = AsShared();
+		PickerArgs.bUseAlpha = bUseAlpha;
+		PickerArgs.DisplayGamma = TAttribute<float>::Create(TAttribute<float>::FGetter::CreateUObject(GEngine, &UEngine::GetDisplayGamma));
+		PickerArgs.OnColorPickerWindowClosed = FOnWindowClosed::CreateSP(this, &SDetailsViewBase::OnColorPickerWindowClosed);
+		PickerArgs.OptionalOwningDetailsView = AsShared();
+		OpenColorPicker(PickerArgs);
 	}
 }
 
-
 void SDetailsViewBase::UpdatePropertyMaps()
 {
 	TRACE_CPUPROFILER_EVENT_SCOPE(SDetailsView::UpdatePropertyMaps);
@@ -566,18 +564,6 @@ void SDetailsViewBase::UpdateSinglePropertyMap(TSharedPtr<FComplexPropertyNode>
 
 void SDetailsViewBase::OnColorPickerWindowClosed(const TSharedRef<SWindow>& Window)
 {
-	const TSharedPtr< FPropertyNode > PinnedColorPropertyNode = ColorPropertyNode.Pin();
-	if (ensure(PinnedColorPropertyNode.IsValid()))
-	{
-		FProperty* Property = PinnedColorPropertyNode->GetProperty();
-		if (Property && PropertyUtilities.IsValid())
-		{
-			FPropertyChangedEvent ChangeEvent(Property, EPropertyChangeType::ArrayAdd);
-			PinnedColorPropertyNode->FixPropertiesInEvent(ChangeEvent);
-			PropertyUtilities->NotifyFinishedChangingProperties(ChangeEvent);
-		}
-	}
-
 	// A color picker window is no longer open
 	bHasOpenColorPicker = false;
 	ColorPropertyNode.Reset();
diff --git a/Engine/Source/Editor/PropertyEditor/Private/SPropertyTreeViewImpl.cpp b/Engine/Source/Editor/PropertyEditor/Private/SPropertyTreeViewImpl.cpp
--- a/Engine/Source/Editor/PropertyEditor/Private/SPropertyTreeViewImpl.cpp
+++ b/Engine/Source/Editor/PropertyEditor/Private/SPropertyTreeViewImpl.cpp
@@ -959,7 +959,6 @@ void SPropertyTreeViewImpl::PostSetObject()
 	check( !bNodeTreeExternallyManaged );
 
 	DestroyColorPicker();
-	ColorPropertyNode = NULL;
 
 	// Reconstruct the property tree so we don't have a tree filled with data we are about to destroy
 	ConstructPropertyTree();
@@ -1195,52 +1194,66 @@ void SPropertyTreeViewImpl::SaveFavorites()
 	}
 }
 
-/** 
- * Creates the color picker window for this property view.
- *
- * @param Node				The slate property node to edit.
- * @param bUseAlpha			Whether or not alpha is supported
- */
-void SPropertyTreeViewImpl::CreateColorPickerWindow(const TSharedRef< class FPropertyEditor >& PropertyEditor, bool bUseAlpha)
+namespace UE::PropertyEditor::Private
 {
-	const TSharedRef< FPropertyNode > PropertyNode = PropertyEditor->GetPropertyNode();
-	ColorPropertyNode = &PropertyNode.Get();
+/** Set the color for the property node */
+void CreateColorPickerWindow_SetColor(FLinearColor NewColor, TWeakPtr<IPropertyHandle> WeakPropertyHandle)
+{
+	if (TSharedPtr<IPropertyHandle> PropertyHandle = WeakPropertyHandle.Pin())
+	{
+		FProperty* NodeProperty = PropertyHandle->GetProperty();
+		check(NodeProperty);
 
-	check(ColorPropertyNode);
-	FProperty* Property = ColorPropertyNode->GetProperty();
+		const UScriptStruct* Struct = CastField<FStructProperty>(NodeProperty)->Struct;
+		if (Struct->GetFName() == NAME_Color)
+		{
+			const bool bSRGB = true;
+			FColor NewFColor = NewColor.ToFColor(bSRGB);
+			ensure(PropertyHandle->SetValueFromFormattedString(NewFColor.ToString(), EPropertyValueSetFlags::DefaultFlags) == FPropertyAccess::Result::Success);
+		}
+		else
+		{
+			check(Struct->GetFName() == NAME_LinearColor);
+			ensure(PropertyHandle->SetValueFromFormattedString(NewColor.ToString(), EPropertyValueSetFlags::DefaultFlags) == FPropertyAccess::Result::Success);
+		}
+	}
+}
+}//namespace
+
+void SPropertyTreeViewImpl::CreateColorPickerWindow(const TSharedRef< class FPropertyEditor >& PropertyEditor, bool bUseAlpha)
+{
+	const FProperty* Property = PropertyEditor->GetProperty();
 	check(Property);
 
 	FReadAddressList ReadAddresses;
-	ColorPropertyNode->GetReadAddress( false, ReadAddresses, false );
+	PropertyEditor->GetPropertyNode()->GetReadAddress(false, ReadAddresses, false);
 
-	TArray<FLinearColor*> LinearColor;
-	TArray<FColor*> DWORDColor;
-	if( ReadAddresses.Num() ) 
+	// Use the first address for the initial color
+	TOptional<FLinearColor> DefaultColor;
+	if (ReadAddresses.Num())
 	{
 		const uint8* Addr = ReadAddresses.GetAddress(0);
-		if( Addr )
+		if (Addr)
 		{
-			if( CastField<FStructProperty>(Property)->Struct->GetFName() == NAME_Color )
+			if (CastField<FStructProperty>(Property)->Struct->GetFName() == NAME_Color)
 			{
-				DWORDColor.Add((FColor*)Addr);
+				DefaultColor = *reinterpret_cast<const FColor*>(Addr);
 			}
 			else
 			{
-				check( CastField<FStructProperty>(Property)->Struct->GetFName() == NAME_LinearColor );
-				LinearColor.Add((FLinearColor*)Addr);
+				check(CastField<FStructProperty>(Property)->Struct->GetFName() == NAME_LinearColor);
+				DefaultColor = *reinterpret_cast<const FLinearColor*>(Addr);
 			}
 		}
 	}
 
-	if( DWORDColor.Num() || LinearColor.Num() )
+	if (DefaultColor.IsSet())
 	{
-		FColorPickerArgs PickerArgs;
+		TWeakPtr<IPropertyHandle> WeakPropertyHandle = PropertyEditor->GetPropertyHandle();
+		FColorPickerArgs PickerArgs = FColorPickerArgs(DefaultColor.GetValue(), FOnLinearColorValueChanged::CreateStatic(&UE::PropertyEditor::Private::CreateColorPickerWindow_SetColor, WeakPropertyHandle));
 		PickerArgs.ParentWidget = AsShared();
 		PickerArgs.bUseAlpha = bUseAlpha;
-		PickerArgs.DisplayGamma = TAttribute<float>::Create( TAttribute<float>::FGetter::CreateUObject(GEngine, &UEngine::GetDisplayGamma) );
-		PickerArgs.ColorArray = &DWORDColor;
-		PickerArgs.LinearColorArray = &LinearColor;
-		PickerArgs.OnColorCommitted = FOnLinearColorValueChanged::CreateSP( this, &SPropertyTreeViewImpl::SetColor);
+		PickerArgs.DisplayGamma = TAttribute<float>::Create(TAttribute<float>::FGetter::CreateUObject(GEngine, &UEngine::GetDisplayGamma));
 
 		OpenColorPicker(PickerArgs);
 	}
@@ -1268,27 +1281,6 @@ void SPropertyTreeViewImpl::SetIsPropertyVisible(FIsPropertyVisible IsPropertyVi
 	}
 }
 
-/** Set the color for the property node */
-void SPropertyTreeViewImpl::SetColor(FLinearColor NewColor)
-{
-	check(ColorPropertyNode);
-	FProperty* NodeProperty = ColorPropertyNode->GetProperty();
-	check(NodeProperty);
-	FObjectPropertyNode* ObjectNode = ColorPropertyNode->FindObjectItemParent();
-	
-	// If more than one object is selected, an empty field indicates their values for this property differ.
-	// Don't send it to the objects value in this case (if we did, they would all get set to None which isn't good).
-	if( ObjectNode && ObjectNode->GetNumObjects() == 1 )
-	{
-		FScopedTransaction Transaction( NSLOCTEXT("UnrealEd", "SetColorProperty", "Set Color Property") );
-
-		ColorPropertyNode->NotifyPreChange(NodeProperty, GetNotifyHook());
-
-		FPropertyChangedEvent ChangeEvent(NodeProperty, EPropertyChangeType::ValueSet);
-		ColorPropertyNode->NotifyPostChange( ChangeEvent, GetNotifyHook() );
-	}
-}
-
 void SPropertyTreeViewImpl::RequestItemExpanded( TSharedPtr<FPropertyNode> PropertyNode, bool bExpand, bool bRecursiveExpansion )
 {
 	// Don't change expansion state if its already in that state
diff --git a/Engine/Source/Editor/PropertyEditor/Private/SPropertyTreeViewImpl.h b/Engine/Source/Editor/PropertyEditor/Private/SPropertyTreeViewImpl.h
--- a/Engine/Source/Editor/PropertyEditor/Private/SPropertyTreeViewImpl.h
+++ b/Engine/Source/Editor/PropertyEditor/Private/SPropertyTreeViewImpl.h
@@ -282,9 +282,6 @@ private:
 	 * Saves favorites to INI                   
 	 */
 	void SaveFavorites();
-	
-	/** Set the color for the property node */
-	void SetColor(FLinearColor NewColor);
 
 	/**
 	 * Requests that a property node's expansion state be changed due to being filtered
@@ -323,8 +320,6 @@ private:
 	FString CurrentFilterText;
 	/** The title of the property view (for a window title) */
 	FString Title;
-	/** The property node that the color picker is currently editing. */
-	FPropertyNode* ColorPropertyNode;
 	/** Notify hook to call when properties are changed */
 	FNotifyHook* NotifyHook;
 	/** True if there is an active filter (text in the filter box) */
diff --git a/Engine/Source/Editor/PropertyEditor/Private/SSingleProperty.cpp b/Engine/Source/Editor/PropertyEditor/Private/SSingleProperty.cpp
--- a/Engine/Source/Editor/PropertyEditor/Private/SSingleProperty.cpp
+++ b/Engine/Source/Editor/PropertyEditor/Private/SSingleProperty.cpp
@@ -310,7 +310,7 @@ void SSingleProperty::CreateColorPickerWindow( const TSharedRef< class FProperty
 {
 	if( HasValidProperty() )
 	{
-		auto Node = PropertyEditor->GetPropertyNode();
+		TSharedRef<FPropertyNode> Node = PropertyEditor->GetPropertyNode();
 		check( &Node.Get() == ValueNode.Get() );
 		FProperty* Property = Node->GetProperty();
 		check(Property);
@@ -318,8 +318,8 @@ void SSingleProperty::CreateColorPickerWindow( const TSharedRef< class FProperty
 		FReadAddressList ReadAddresses;
 		Node->GetReadAddress( false, ReadAddresses, false );
 
-		TArray<FLinearColor*> LinearColor;
-		TArray<FColor*> DWORDColor;
+		// Use the first address for the initial color
+		TOptional<FLinearColor> DefaultColor;
 		if( ReadAddresses.Num() ) 
 		{
 			const uint8* Addr = ReadAddresses.GetAddress(0);
@@ -327,25 +327,25 @@ void SSingleProperty::CreateColorPickerWindow( const TSharedRef< class FProperty
 			{
 				if( CastField<FStructProperty>(Property)->Struct->GetFName() == NAME_Color )
 				{
-					DWORDColor.Add((FColor*)Addr);
+					DefaultColor = *reinterpret_cast<const FColor*>(Addr);
 				}
 				else
 				{
 					check( CastField<FStructProperty>(Property)->Struct->GetFName() == NAME_LinearColor );
-					LinearColor.Add((FLinearColor*)Addr);
+					DefaultColor = *reinterpret_cast<const FLinearColor*>(Addr);
 				}
 			}
 		}
 
-		FColorPickerArgs PickerArgs;
-		PickerArgs.ParentWidget = AsShared();
-		PickerArgs.bUseAlpha = bUseAlpha;
-		PickerArgs.DisplayGamma = TAttribute<float>::Create( TAttribute<float>::FGetter::CreateUObject(GEngine, &UEngine::GetDisplayGamma) );
-		PickerArgs.ColorArray = &DWORDColor;
-		PickerArgs.LinearColorArray = &LinearColor;
-		PickerArgs.OnColorCommitted = FOnLinearColorValueChanged::CreateSP( this, &SSingleProperty::SetColorPropertyFromColorPicker);
+		if (DefaultColor.IsSet())
+		{
+			FColorPickerArgs PickerArgs = FColorPickerArgs(DefaultColor.GetValue(), FOnLinearColorValueChanged::CreateSP(this, &SSingleProperty::SetColorPropertyFromColorPicker));
+			PickerArgs.ParentWidget = AsShared();
+			PickerArgs.bUseAlpha = bUseAlpha;
+			PickerArgs.DisplayGamma = TAttribute<float>::Create(TAttribute<float>::FGetter::CreateUObject(GEngine, &UEngine::GetDisplayGamma));
 
-		OpenColorPicker(PickerArgs);
+			OpenColorPicker(PickerArgs);
+		}
 	}
 }
 
@@ -355,11 +355,17 @@ void SSingleProperty::SetColorPropertyFromColorPicker(FLinearColor NewColor)
 	{
 		FProperty* NodeProperty = ValueNode->GetProperty();
 		check(NodeProperty);
+		check(GetPropertyHandle());
 
-		//@todo if multiple objects we need to iterate
-		ValueNode->NotifyPreChange(NodeProperty, GetNotifyHook());
-
-		FPropertyChangedEvent ChangeEvent(NodeProperty, EPropertyChangeType::ValueSet);
-		ValueNode->NotifyPostChange( ChangeEvent, GetNotifyHook() );
+		if (CastField<FStructProperty>(NodeProperty)->Struct->GetFName() == NAME_Color)
+		{
+			const bool bSRGB = true;
+			FColor NewFColor = NewColor.ToFColor(bSRGB);
+			ensure(GetPropertyHandle()->SetValueFromFormattedString(NewFColor.ToString(), EPropertyValueSetFlags::DefaultFlags) == FPropertyAccess::Result::Success);
+		}
+		else
+		{
+			ensure(GetPropertyHandle()->SetValueFromFormattedString(NewColor.ToString(), EPropertyValueSetFlags::DefaultFlags) == FPropertyAccess::Result::Success);
+		}
 	}
 }
diff --git a/Engine/Source/Editor/PropertyEditor/Private/SStructureDetailsView.cpp b/Engine/Source/Editor/PropertyEditor/Private/SStructureDetailsView.cpp
--- a/Engine/Source/Editor/PropertyEditor/Private/SStructureDetailsView.cpp
+++ b/Engine/Source/Editor/PropertyEditor/Private/SStructureDetailsView.cpp
@@ -230,8 +230,14 @@ void SStructureDetailsView::SetStructureData(TSharedPtr<FStructOnScope> InStruct
 	}
 	
 	//POST SET
-	DestroyColorPicker();
-	ColorPropertyNode = NULL;
+	TSharedPtr<SColorPicker> ExistingColorPicker = GetColorPicker();
+	if (ExistingColorPicker.IsValid()
+		&& (!ExistingColorPicker->GetOptionalOwningDetailsView().IsValid()
+			|| ExistingColorPicker->GetOptionalOwningDetailsView().Get() == this))
+	{
+		DestroyColorPicker();
+		bHasOpenColorPicker = false;
+	}
 
 	FPropertyNodeInitParams InitParams;
 	InitParams.ParentNode = NULL;
diff --git a/Engine/Source/Runtime/AppFramework/Private/Widgets/Colors/SColorPicker.cpp b/Engine/Source/Runtime/AppFramework/Private/Widgets/Colors/SColorPicker.cpp
--- a/Engine/Source/Runtime/AppFramework/Private/Widgets/Colors/SColorPicker.cpp
+++ b/Engine/Source/Runtime/AppFramework/Private/Widgets/Colors/SColorPicker.cpp
@@ -555,10 +555,10 @@ void SColorPicker::GenerateDefaultColorPickerContent( bool bAdvancedSectionExpan
 
 										+ SHorizontalBox::Slot()
 											.AutoWidth()
-											.MaxWidth(72.0f)
+											.MaxWidth(80.0f)
 											[
 												SNew(SEditableTextBox)
-													.MinDesiredWidth(72.0f)
+													.MinDesiredWidth(80.0f)
 													.Text(this, &SColorPicker::HandleHexLinearBoxText)
 													.OnTextCommitted(this, &SColorPicker::HandleHexLinearInputTextCommitted)
 											]
@@ -584,10 +584,10 @@ void SColorPicker::GenerateDefaultColorPickerContent( bool bAdvancedSectionExpan
 
 										+ SHorizontalBox::Slot()
 											.AutoWidth()
-											.MaxWidth(72.0f)
+											.MaxWidth(80.0f)
 											[
 												SNew(SEditableTextBox)
-												.MinDesiredWidth(72.0f)
+												.MinDesiredWidth(80.0f)
 												.Text(this, &SColorPicker::HandleHexSRGBBoxText)
 												.OnTextCommitted(this, &SColorPicker::HandleHexSRGBInputTextCommitted)
 											]
diff --git a/Engine/Source/Runtime/AppFramework/Public/Widgets/Colors/SColorPicker.h b/Engine/Source/Runtime/AppFramework/Public/Widgets/Colors/SColorPicker.h
--- a/Engine/Source/Runtime/AppFramework/Public/Widgets/Colors/SColorPicker.h
+++ b/Engine/Source/Runtime/AppFramework/Public/Widgets/Colors/SColorPicker.h
@@ -502,40 +502,40 @@ private:
 struct FColorPickerArgs
 {
 	/** Whether or not the new color picker is modal. */
-	bool bIsModal;
+	bool bIsModal = false;
 
 	/** The parent for the new color picker window */
 	TSharedPtr<SWidget> ParentWidget;
 
 	/** Whether or not to enable the alpha slider. */
-	bool bUseAlpha;
+	bool bUseAlpha = false;
 
 	/** Whether to disable the refresh except on mouse up for performance reasons. */
-	bool bOnlyRefreshOnMouseUp;
+	bool bOnlyRefreshOnMouseUp = false;
 
 	/** Whether to disable the refresh until the picker closes. */
-	bool bOnlyRefreshOnOk;
+	bool bOnlyRefreshOnOk = false;
 
 	/** Whether to automatically expand the Advanced section. */
-	bool bExpandAdvancedSection;
+	bool bExpandAdvancedSection = true;
 	
 	/** Whether to open the color picker as a menu window. */
-	bool bOpenAsMenu;
+	bool bOpenAsMenu = false;
 
 	/** The current display gamma used to correct colors picked from the display. */
-	TAttribute<float> DisplayGamma;
+	TAttribute<float> DisplayGamma = 2.2f;
 
 	/** If set overrides the global option for the desired setting of sRGB mode. */
-	TOptional<bool> sRGBOverride;
+	TOptional<bool> sRGBOverride = TOptional<bool>();
 
 	/** An array of FColors to target. */
-	const TArray<FColor*>* ColorArray;
+	const TArray<FColor*>* ColorArray = nullptr;
 
 	/** An array of FLinearColors to target. */
-	const TArray<FLinearColor*>* LinearColorArray;
+	const TArray<FLinearColor*>* LinearColorArray = nullptr;
 
 	/** An array of FColorChannels to target. (deprecated now that wx is gone?) */
-	const TArray<FColorChannels>* ColorChannelsArray;
+	const TArray<FColorChannels>* ColorChannelsArray = nullptr;
 
 	/** A delegate to be called when the color changes. */
 	FOnLinearColorValueChanged OnColorCommitted;
@@ -556,34 +556,18 @@ struct FColorPickerArgs
 	FSimpleDelegate OnInteractivePickEnd;
 
 	/** Overrides the initial color set on the color picker. */
-	FLinearColor InitialColorOverride;
+	FLinearColor InitialColorOverride = FLinearColor::White;
 
 	/** Allows a details view to own the color picker so refreshing another details view doesn't close it */
 	TSharedPtr<SWidget> OptionalOwningDetailsView;
 
 	/** Default constructor. */
-	FColorPickerArgs()
-		: bIsModal(false)
-		, ParentWidget(nullptr)
-		, bUseAlpha(false)
-		, bOnlyRefreshOnMouseUp(false)
-		, bOnlyRefreshOnOk(false)
-		, bExpandAdvancedSection(true)
-		, bOpenAsMenu(false)
-		, DisplayGamma(2.2f)
-		, sRGBOverride()
-		, ColorArray(nullptr)
-		, LinearColorArray(nullptr)
-		, ColorChannelsArray(nullptr)
-		, OnColorCommitted()
-		, PreColorCommitted()
-		, OnColorPickerWindowClosed()
-		, OnColorPickerCancelled()
-		, OnInteractivePickBegin()
-		, OnInteractivePickEnd()
-		, InitialColorOverride()
-		, OptionalOwningDetailsView(nullptr)
-	{ }
+	FColorPickerArgs() = default;
+
+	FColorPickerArgs(FLinearColor InInitialColor, FOnLinearColorValueChanged InOnColorCommitted)
+		: OnColorCommitted(MoveTemp(InOnColorCommitted))
+		, InitialColorOverride(InInitialColor)
+	{}
 };
 
 /** Get a pointer to the static color picker, or nullptr if it does not exist. */
