Ignore empty strings when renaming a material editor node.

This fies a bug causing crash whenmaterial parameter node is renamed to empty string.

diff --git a/Engine/Source/Editor/MaterialEditor/Private/MaterialEditor.cpp b/Engine/Source/Editor/MaterialEditor/Private/MaterialEditor.cpp
--- a/Engine/Source/Editor/MaterialEditor/Private/MaterialEditor.cpp
+++ b/Engine/Source/Editor/MaterialEditor/Private/MaterialEditor.cpp
@@ -7409,8 +7409,13 @@ void FMaterialEditor::OnNodeTitleCommitted(const FText& NewText, ETextCommit::Ty
 {
 	if (NodeBeingChanged)
 	{
-		FString NewName = NewText.ToString();
-		if (NewName.Len() >= NAME_SIZE) {
+		FString NewName = NewText.ToString().TrimStartAndEnd();
+		if (NewName.IsEmpty())
+		{
+			// Ignore empty names
+			return;
+		}
+		else if (NewName.Len() >= NAME_SIZE) {
 			UE_LOG(LogMaterialEditor, Warning, TEXT("New material graph node name '%s...' exceeds maximum length of %d and thus was truncated."), *NewName.Left(8), NAME_SIZE - 1);
 			NewName = NewName.Left(NAME_SIZE - 1);
 		}
