PR #10158: Generate new Guid for material expressions after they are pasted (Contributed by [REDACTED_PERSON])

diff --git a/Engine/Source/Editor/MaterialEditor/Private/MaterialEditor.cpp b/Engine/Source/Editor/MaterialEditor/Private/MaterialEditor.cpp
--- a/Engine/Source/Editor/MaterialEditor/Private/MaterialEditor.cpp
+++ b/Engine/Source/Editor/MaterialEditor/Private/MaterialEditor.cpp
@@ -5815,6 +5815,11 @@ void FMaterialEditor::PostPasteMaterialExpression(UMaterialExpression* NewExpres
 		// We just updated all our child expressions, reconstruct node.
 		Composite->GraphNode->ReconstructNode();
 	}
+	else
+	{
+		// Give new expression a different Guid from the old one after pasting
+		NewExpression->UpdateMaterialExpressionGuid(true, true);
+	}
 
 	// There can be only one default mesh paint texture.
 	UMaterialExpressionTextureBase* TextureSample = Cast<UMaterialExpressionTextureBase>(NewExpression);
