Retain details of assigned texture and other params when converting between Texture Object and Texture Object Parameter material expression.

diff --git a/Engine/Source/Editor/MaterialEditor/Private/MaterialEditor.cpp b/Engine/Source/Editor/MaterialEditor/Private/MaterialEditor.cpp
--- a/Engine/Source/Editor/MaterialEditor/Private/MaterialEditor.cpp
+++ b/Engine/Source/Editor/MaterialEditor/Private/MaterialEditor.cpp
@@ -4022,6 +4022,22 @@ void FMaterialEditor::OnConvertObjects()
 							NewTextureExpr->MipValueMode = TextureSampleExpression->MipValueMode;
 							NewGraphNode->ReconstructNode();
 						}
+						else if (TextureObjectExpression && !TextureObjectParameterExpression)
+						{
+							bNeedsRefresh = true;
+							UMaterialExpressionTextureObjectParameter* NewTextureObjectParameterExpression = CastChecked<UMaterialExpressionTextureObjectParameter>(NewExpression);
+							NewTextureObjectParameterExpression->Texture = TextureObjectExpression->Texture;
+							NewTextureObjectParameterExpression->AutoSetSampleType();
+							NewTextureObjectParameterExpression->IsDefaultMeshpaintTexture = TextureObjectExpression->IsDefaultMeshpaintTexture;
+						}	
+						else if (TextureObjectParameterExpression && (!TextureObjectExpression || !TextureSampleExpression))
+						{
+							bNeedsRefresh = true;
+							UMaterialExpressionTextureObject* NewTextureObjectExpression = CastChecked<UMaterialExpressionTextureObject>(NewExpression);
+							NewTextureObjectExpression->Texture = TextureObjectParameterExpression->Texture;
+							NewTextureObjectExpression->AutoSetSampleType();
+							NewTextureObjectExpression->IsDefaultMeshpaintTexture = TextureObjectParameterExpression->IsDefaultMeshpaintTexture;
+						}
 						else if (RuntimeVirtualTextureSampleExpression)
 						{
 							bNeedsRefresh = true;
