Remove r.MaterialEnableControlFlow and disable the feature as how control flow should be implemented in the material editor is still under discussion

diff --git a/Engine/Source/Editor/MaterialEditor/Private/MaterialEditorDetailCustomization.cpp b/Engine/Source/Editor/MaterialEditor/Private/MaterialEditorDetailCustomization.cpp
--- a/Engine/Source/Editor/MaterialEditor/Private/MaterialEditorDetailCustomization.cpp
+++ b/Engine/Source/Editor/MaterialEditor/Private/MaterialEditorDetailCustomization.cpp
@@ -43,6 +43,7 @@
 #include "Materials/MaterialExpressionCurveAtlasRowParameter.h"
 #include "Curves/CurveLinearColorAtlas.h"
 #include "RenderUtils.h"
+#include "MaterialShared.h"
 
 #define LOCTEXT_NAMESPACE "MaterialEditor"
 
@@ -777,7 +778,6 @@ TSharedRef<class IDetailCustomization> FMaterialDetailCustomization::MakeInstanc
 
 void FMaterialDetailCustomization::CustomizeDetails( IDetailLayoutBuilder& DetailLayout )
 {
-	static const auto CVarMaterialEnableControlFlow = IConsoleManager::Get().FindTConsoleVariableDataInt(TEXT("r.MaterialEnableControlFlow"));
 	static const auto CVarMaterialEnableNewHLSLGenerator = IConsoleManager::Get().FindTConsoleVariableDataInt(TEXT("r.MaterialEnableNewHLSLGenerator"));
 
 	TArray<TWeakObjectPtr<UObject> > Objects;
@@ -840,7 +840,7 @@ void FMaterialDetailCustomization::CustomizeDetails( IDetailLayoutBuilder& Detai
 				BlendModeProperty->Enum = GetBlendModeEnum();
 			}
 
-			if (PropertyName == GET_MEMBER_NAME_CHECKED(UMaterial, bEnableExecWire) && !CVarMaterialEnableControlFlow->GetValueOnAnyThread())
+			if (PropertyName == GET_MEMBER_NAME_CHECKED(UMaterial, bEnableExecWire) && !AllowMaterialControlFlow())
 			{
 				DetailLayout.HideProperty(PropertyHandle);
 			}
@@ -900,7 +900,6 @@ TSharedRef<class IDetailCustomization> FMaterialFunctionDetailCustomization::Mak
 
 void FMaterialFunctionDetailCustomization::CustomizeDetails(IDetailLayoutBuilder& DetailLayout)
 {
-	static const auto CVarMaterialEnableControlFlow = IConsoleManager::Get().FindTConsoleVariableDataInt(TEXT("r.MaterialEnableControlFlow"));
 	static const auto CVarMaterialEnableNewHLSLGenerator = IConsoleManager::Get().FindTConsoleVariableDataInt(TEXT("r.MaterialEnableNewHLSLGenerator"));
 
 	// MaterialFunction category
@@ -915,7 +914,7 @@ void FMaterialFunctionDetailCustomization::CustomizeDetails(IDetailLayoutBuilder
 			FProperty* Property = PropertyHandle->GetProperty();
 			FName PropertyName = Property->GetFName();
 
-			if (PropertyName == GET_MEMBER_NAME_CHECKED(UMaterialFunction, bEnableExecWire) && !CVarMaterialEnableControlFlow->GetValueOnAnyThread())
+			if (PropertyName == GET_MEMBER_NAME_CHECKED(UMaterialFunction, bEnableExecWire) && !AllowMaterialControlFlow())
 			{
 				DetailLayout.HideProperty(PropertyHandle);
 			}
diff --git a/Engine/Source/Editor/MaterialEditor/Private/MaterialExpressionClasses.cpp b/Engine/Source/Editor/MaterialEditor/Private/MaterialExpressionClasses.cpp
--- a/Engine/Source/Editor/MaterialEditor/Private/MaterialExpressionClasses.cpp
+++ b/Engine/Source/Editor/MaterialEditor/Private/MaterialExpressionClasses.cpp
@@ -12,6 +12,7 @@
 #include "Modules/ModuleManager.h"
 #include "Preferences/MaterialEditorOptions.h"
 #include "MaterialEditorModule.h"
+#include "MaterialShared.h"
 
 #include "Materials/MaterialExpressionComment.h"
 #include "Materials/MaterialExpressionParameter.h"
@@ -85,10 +86,9 @@ void MaterialExpressionClasses::InitMaterialExpressionClasses()
 {
 	if(!bInitialized)
 	{
-		static const auto CVarMaterialEnableControlFlow = IConsoleManager::Get().FindTConsoleVariableDataInt(TEXT("r.MaterialEnableControlFlow"));
 		static const auto CVarMaterialEnableNewHLSLGenerator = IConsoleManager::Get().FindTConsoleVariableDataInt(TEXT("r.MaterialEnableNewHLSLGenerator"));
 		
-		const bool bEnableControlFlow = CVarMaterialEnableControlFlow->GetValueOnAnyThread() != 0;
+		const bool bEnableControlFlow = AllowMaterialControlFlow();
 		const bool bEnableNewHLSLGenerator = CVarMaterialEnableNewHLSLGenerator->GetValueOnAnyThread() != 0;
 
 		UMaterialEditorOptions* TempEditorOptions = NewObject<UMaterialEditorOptions>();
diff --git a/Engine/Source/Runtime/Engine/Private/Materials/MaterialExpressions.cpp b/Engine/Source/Runtime/Engine/Private/Materials/MaterialExpressions.cpp
--- a/Engine/Source/Runtime/Engine/Private/Materials/MaterialExpressions.cpp
+++ b/Engine/Source/Runtime/Engine/Private/Materials/MaterialExpressions.cpp
@@ -20,6 +20,7 @@
 #include "Engine/Font.h"
 #include "DataDrivenShaderPlatformInfo.h"
 #include "MaterialDomain.h"
+#include "MaterialShared.h"
 #include "Materials/HLSLMaterialTranslator.h"
 #include "Materials/MaterialAttributeDefinitionMap.h"
 #include "Materials/MaterialExpressionMaterialFunctionCall.h"
@@ -14991,8 +14992,7 @@ bool UMaterialFunction::IsUsingControlFlow() const
 {
 	if (bEnableExecWire)
 	{
-		static const auto CVar = IConsoleManager::Get().FindTConsoleVariableDataInt(TEXT("r.MaterialEnableControlFlow"));
-		return CVar->GetValueOnAnyThread() != 0;
+		return AllowMaterialControlFlow();
 	}
 	return false;
 }
diff --git a/Engine/Source/Runtime/Engine/Private/Materials/MaterialInterface.cpp b/Engine/Source/Runtime/Engine/Private/Materials/MaterialInterface.cpp
--- a/Engine/Source/Runtime/Engine/Private/Materials/MaterialInterface.cpp
+++ b/Engine/Source/Runtime/Engine/Private/Materials/MaterialInterface.cpp
@@ -66,12 +66,6 @@ static TAutoConsoleVariable<int32> CVarMaterialEnableNewHLSLGenerator(
 	TEXT("2 - Force all materials to use new generator\n"),
 	ECVF_RenderThreadSafe | ECVF_ReadOnly);
 
-static TAutoConsoleVariable<int32> CVarMaterialEnableControlFlow(
-	TEXT("r.MaterialEnableControlFlow"),
-	0,
-	TEXT("Allows experemental control flow to be used in the material editor.\n"),
-	ECVF_RenderThreadSafe | ECVF_ReadOnly);
-
 //////////////////////////////////////////////////////////////////////////
 
 bool IsHairStrandsGeometrySupported(const EShaderPlatform Platform)
diff --git a/Engine/Source/Runtime/Engine/Public/MaterialShared.h b/Engine/Source/Runtime/Engine/Public/MaterialShared.h
--- a/Engine/Source/Runtime/Engine/Public/MaterialShared.h
+++ b/Engine/Source/Runtime/Engine/Public/MaterialShared.h
@@ -3040,4 +3040,11 @@ inline bool ShouldIncludeMaterialInDefaultOpaquePass(const FMaterial& Material)
 {
 	return !Material.IsSky()
 		&& !Material.GetShadingModels().HasShadingModel(MSM_SingleLayerWater);
-}
\ No newline at end of file
+}
+
+// r.MaterialEnableControlFlow is removed and the feature is forced disabled as how control flow should be
+// implemented in the material editor is still under discussion.
+inline bool AllowMaterialControlFlow()
+{
+	return false;
+}
