Fix for duplicated material applying multiple times  data conversion after a duplicate when it is in fact not needed since it has been done on load already.

diff --git a/Engine/Source/Editor/MaterialEditor/Private/MaterialEditor.cpp b/Engine/Source/Editor/MaterialEditor/Private/MaterialEditor.cpp
--- a/Engine/Source/Editor/MaterialEditor/Private/MaterialEditor.cpp
+++ b/Engine/Source/Editor/MaterialEditor/Private/MaterialEditor.cpp
@@ -477,7 +477,7 @@ void FMaterialEditor::UnregisterTabSpawners(const TSharedRef<class FTabManager>&
 
 	OnUnregisterTabSpawners().Broadcast(InTabManager);
 }
-
+UE_DISABLE_OPTIMIZATION
 void FMaterialEditor::InitEditorForMaterial(UMaterial* InMaterial)
 {
 	check(InMaterial);
@@ -487,11 +487,15 @@ void FMaterialEditor::InitEditorForMaterial(UMaterial* InMaterial)
 	OriginalMaterialObject = InMaterial;
 
 	ExpressionPreviewMaterial = NULL;
+
+	// Now, the material has been loaded and serialized. PostLoad has been called and the asset has been converted and now it up to date.
+	// Let's thus reset the linker (to the last version) to make sure PostLoad is not doing any data conversion squashing new properties with not properly initialized _DEPRECATED members.
+	ResetLoaders(OriginalMaterial->GetPackage());
 	
 	// Create a copy of the material for preview usage (duplicating to a different class than original!)
 	// Propagate all object flags except for RF_Standalone, otherwise the preview material won't GC once
 	// the material editor releases the reference.
-	Material = (UMaterial*)StaticDuplicateObject(OriginalMaterial, GetTransientPackage(), NAME_None, ~RF_Standalone, UPreviewMaterial::StaticClass()); 
+	Material = (UMaterial*)StaticDuplicateObject(OriginalMaterial, GetTransientPackage(), NAME_None, ~RF_Standalone, UPreviewMaterial::StaticClass());  
 	
 	Material->CancelOutstandingCompilation();	//The material is compiled later on anyway so no need to do it in Duplication/PostLoad. 
 												//I'm hackily canceling the jobs here but we should really not add the jobs in the first place. <<--- TODO
