Removed bad disable optimisation that stayed behind after some more debugging...

#prelfight none

diff --git a/Engine/Source/Editor/MaterialEditor/Private/MaterialEditor.cpp b/Engine/Source/Editor/MaterialEditor/Private/MaterialEditor.cpp
--- a/Engine/Source/Editor/MaterialEditor/Private/MaterialEditor.cpp
+++ b/Engine/Source/Editor/MaterialEditor/Private/MaterialEditor.cpp
@@ -477,7 +477,7 @@ void FMaterialEditor::UnregisterTabSpawners(const TSharedRef<class FTabManager>&
 
 	OnUnregisterTabSpawners().Broadcast(InTabManager);
 }
-UE_DISABLE_OPTIMIZATION
+
 void FMaterialEditor::InitEditorForMaterial(UMaterial* InMaterial)
 {
 	check(InMaterial);
