Substrate - Editor HUD upgdate.

diff --git a/Engine/Source/Editor/EditorStyle/Private/SlateEditorStyle.cpp b/Engine/Source/Editor/EditorStyle/Private/SlateEditorStyle.cpp
--- a/Engine/Source/Editor/EditorStyle/Private/SlateEditorStyle.cpp
+++ b/Engine/Source/Editor/EditorStyle/Private/SlateEditorStyle.cpp
@@ -4410,6 +4410,12 @@ void FSlateEditorStyle::FStyle::SetupGraphEditorStyles()
 				.SetShadowOffset( FVector2D::ZeroVector )
 			);
 
+			Set("Graph.WarningText", FTextBlockStyle(NormalText)
+				.SetFont(DEFAULT_FONT("BoldCondensed", 20))
+				.SetColorAndOpacity(FLinearColor(0.8, 0.5f, 0.07f, 1.0f))
+				.SetShadowOffset(FVector2D(1, 1))
+			);
+
 			Set( "Graph.SimulatingText", FTextBlockStyle(NormalText)
 				.SetFont( DEFAULT_FONT( "BoldCondensed", 48 ) )
 				.SetColorAndOpacity( FLinearColor(0.8, 0.8f, 0.0f, 0.2f) )
diff --git a/Engine/Source/Editor/EditorStyle/Private/StarshipStyle.cpp b/Engine/Source/Editor/EditorStyle/Private/StarshipStyle.cpp
--- a/Engine/Source/Editor/EditorStyle/Private/StarshipStyle.cpp
+++ b/Engine/Source/Editor/EditorStyle/Private/StarshipStyle.cpp
@@ -3796,6 +3796,12 @@ void FStarshipEditorStyle::FStyle::SetupGraphEditorStyles()
 				.SetShadowOffset( FVector2D::ZeroVector )
 			);
 
+			Set("Graph.WarningText", FTextBlockStyle(NormalText)
+				.SetFont(DEFAULT_FONT("BoldCondensed", 20))
+				.SetColorAndOpacity(FLinearColor(0.8, 0.5f, 0.07f, 1.0f))
+				.SetShadowOffset(FVector2D(1, 1))
+			);
+
 			Set( "Graph.SimulatingText", FTextBlockStyle(NormalText)
 				.SetFont( DEFAULT_FONT( "BoldCondensed", 48 ) )
 				.SetColorAndOpacity( FLinearColor(0.8, 0.8f, 0.0f, 0.2f) )
diff --git a/Engine/Source/Editor/GraphEditor/Private/SGraphEditorImpl.cpp b/Engine/Source/Editor/GraphEditor/Private/SGraphEditorImpl.cpp
--- a/Engine/Source/Editor/GraphEditor/Private/SGraphEditorImpl.cpp
+++ b/Engine/Source/Editor/GraphEditor/Private/SGraphEditorImpl.cpp
@@ -718,6 +718,12 @@ void SGraphEditorImpl::Construct( const FArguments& InArgs )
 			FText OverrideText = Appearance.Get().ReadOnlyText;
 			return !OverrideText.IsEmpty() ? OverrideText : DefaultText;
 		}
+
+		static FText GetWarningText(TAttribute<FGraphAppearanceInfo> Appearance, FText DefaultText)
+		{
+			FText OverrideText = Appearance.Get().WarningText;
+			return !OverrideText.IsEmpty() ? OverrideText : DefaultText;
+		}
 	};
 	
 	FText DefaultPIENotify(LOCTEXT("GraphSimulatingText", "SIMULATING"));
@@ -730,6 +736,11 @@ void SGraphEditorImpl::Construct( const FArguments& InArgs )
 		TAttribute<FText>::Create(TAttribute<FText>::FGetter::CreateStatic(&Local::GetReadOnlyText, Appearance, DefaultReadOnlyText)) :
 		TAttribute<FText>(DefaultReadOnlyText);
 
+	FText DefaultWarningText(LOCTEXT("GraphWarningText", ""));
+	TAttribute<FText> WarningText = Appearance.IsBound() ?
+		TAttribute<FText>::Create(TAttribute<FText>::FGetter::CreateStatic(&Local::GetWarningText, Appearance, DefaultWarningText)) :
+		TAttribute<FText>(DefaultWarningText);
+
 	TSharedPtr<SOverlay> OverlayWidget;
 
 	this->ChildSlot
@@ -802,6 +813,18 @@ void SGraphEditorImpl::Construct( const FArguments& InArgs )
 			]			
 		]
 
+		// Bottom-left corner text for Substrate
+		+SOverlay::Slot()
+		.Padding(10)
+		.VAlign(VAlign_Bottom)
+		.HAlign(HAlign_Left)
+		[
+			SNew(STextBlock)
+			.Visibility(EVisibility::Visible)
+			.TextStyle(FAppStyle::Get(), "Graph.WarningText")
+			.Text(WarningText)
+		]
+
 		// Bottom-right corner text indicating the type of tool
 		+SOverlay::Slot()
 		.Padding(10)
diff --git a/Engine/Source/Editor/MaterialEditor/Private/MaterialEditor.cpp b/Engine/Source/Editor/MaterialEditor/Private/MaterialEditor.cpp
--- a/Engine/Source/Editor/MaterialEditor/Private/MaterialEditor.cpp
+++ b/Engine/Source/Editor/MaterialEditor/Private/MaterialEditor.cpp
@@ -3025,7 +3025,7 @@ void FMaterialEditor::UpdateMaterialinfoList_Old()
 
 				if (FMaterialShaderMap* ShaderMap = MaterialResource->GetGameThreadShaderMap())
 				{
-					const FString StrataMaterialDescription = ShaderMap->GetStrataMaterialDescription();
+					const FString StrataMaterialDescription = ShaderMap->GetStrataMaterialCompilationOutput().StrataMaterialDescription;
 
 					if (StrataMaterialDescription.Len())
 					{
@@ -7031,7 +7031,25 @@ FGraphAppearanceInfo FMaterialEditor::GetGraphAppearance() const
 	}
 	else
 	{
-		AppearanceInfo.CornerText = LOCTEXT("AppearanceCornerText_Material", "MATERIAL");
+		AppearanceInfo.CornerText = LOCTEXT("AppearanceCornerText_Material", "MATERIAL"); 
+	}
+
+	UMaterial* MaterialForStats = this->bStatsFromPreviewMaterial ? this->Material : this->OriginalMaterial;
+	const FMaterialResource* MaterialResource = MaterialForStats->GetMaterialResource(GMaxRHIFeatureLevel);
+	if (MaterialResource)
+	{
+		FString MaterialDescription;
+	
+		FMaterialShaderMap* ShaderMap = MaterialResource->GetGameThreadShaderMap();
+		if (ShaderMap)
+		{
+			const FStrataMaterialCompilationOutput& CompilationOutput = ShaderMap->GetStrataMaterialCompilationOutput();
+			if (CompilationOutput.bMaterialOutOfBudgetHasBeenSimplified)
+			{
+				AppearanceInfo.WarningText = LOCTEXT("AppearanceCornerText_Material", "Substrate material was out of budget and has been simplified.");
+
+			}
+		}
 	}
 
 	return AppearanceInfo;
@@ -7504,6 +7522,9 @@ void FMaterialEditor::UpdateStatsMaterials()
 		EmptyMaterial->PreEditChange(NULL);
 		EmptyMaterial->PostEditChange();
 	}
+
+	// Also request to update the Substrate slab.
+	StrataWidget->UpdateFromMaterial();
 }
 
 void FMaterialEditor::NotifyExternalMaterialChange()
diff --git a/Engine/Source/Editor/MaterialEditor/Private/MaterialEditorModes.cpp b/Engine/Source/Editor/MaterialEditor/Private/MaterialEditorModes.cpp
--- a/Engine/Source/Editor/MaterialEditor/Private/MaterialEditorModes.cpp
+++ b/Engine/Source/Editor/MaterialEditor/Private/MaterialEditorModes.cpp
@@ -19,7 +19,7 @@ const FName FMaterialEditorApplicationModes::StandardMaterialEditorMode(TEXT("Gr
 
 TSharedPtr<FTabManager::FLayout> FMaterialEditorApplicationModes::GetDefaultEditorLayout(TSharedPtr<FMaterialEditor> InMaterialEditor)
 {
-	const TSharedRef<FTabManager::FLayout> StandaloneDefaultLayout = FTabManager::NewLayout("Standalone_MaterialEditor_Layout_v13")
+	const TSharedRef<FTabManager::FLayout> StandaloneDefaultLayout = FTabManager::NewLayout("Standalone_MaterialEditor_Layout_v14")
 	->AddArea
 	(
 		FTabManager::NewPrimaryArea() ->SetOrientation(Orient_Vertical)
@@ -64,6 +64,7 @@ TSharedPtr<FTabManager::FLayout> FMaterialEditorApplicationModes::GetDefaultEdit
 					->AddTab( FMaterialStats::GetGridStatsTabName(), ETabState::ClosedTab )
 					->AddTab( FMaterialStats::GetGridOldStatsTabName(), ETabState::ClosedTab )
 					->AddTab( FMaterialEditorTabs::FindTabId, ETabState::ClosedTab )
+					->AddTab( FMaterialEditorTabs::StrataTabId, ETabState::OpenedTab )
 				)
 			)
 			->Split
@@ -73,7 +74,6 @@ TSharedPtr<FTabManager::FLayout> FMaterialEditorApplicationModes::GetDefaultEdit
 				(
 					FTabManager::NewStack()
 					->AddTab( FMaterialEditorTabs::PaletteTabId, ETabState::SidebarTab, ESidebarLocation::Right, 0.14f )
-					->AddTab( FMaterialEditorTabs::StrataTabId,  ETabState::SidebarTab, ESidebarLocation::Right, 0.14f)
 				)
 			)
 		)
diff --git a/Engine/Source/Editor/MaterialEditor/Private/SMaterialEditorStrataWidget.cpp b/Engine/Source/Editor/MaterialEditor/Private/SMaterialEditorStrataWidget.cpp
--- a/Engine/Source/Editor/MaterialEditor/Private/SMaterialEditorStrataWidget.cpp
+++ b/Engine/Source/Editor/MaterialEditor/Private/SMaterialEditorStrataWidget.cpp
@@ -5,11 +5,15 @@
 #include "EditorWidgetsModule.h"
 #include "AssetRegistry/AssetRegistryModule.h"
 #include "MaterialEditor.h"
+#include "StrataDefinitions.h"
+#include <functional>
 
 #include "Widgets/Input/SButton.h"
 #include "Widgets/Input/SCheckBox.h"
 #include "Widgets/Layout/SWrapBox.h"
-#include "Widgets/Text/STextBlock.h"
+#include "Widgets/Layout/SScrollBox.h"
+#include "Widgets/SNullWidget.h"
+#include "Widgets/Notifications/SErrorText.h"
 
 #define LOCTEXT_NAMESPACE "SMaterialEditorStrataWidget"
 
@@ -28,55 +32,102 @@ void SMaterialEditorStrataWidget::Construct(const FArguments& InArgs, TWeakPtr<F
 		.Padding(5.0f)
 		.ToolTipText(LOCTEXT("CheckBoxForceFullSimplificationToolTip", "This will force full simplification of the material."));	// Just a test, needs to be more explicit
 
+	DescriptionTextBlock = SNew(STextBlock)
+		.TextStyle(FAppStyle::Get(), "Log.Normal")
+		.ColorAndOpacity(FLinearColor::White)
+		.ShadowColorAndOpacity(FLinearColor::Black)
+		.ShadowOffset(FVector2D::UnitVector)
+		.Text(LOCTEXT("DescriptionTextBlock_Default", "Shader is compiling"));
+
 	if (Strata::IsStrataEnabled())
 	{
 		this->ChildSlot
 		[
 			SNew(SVerticalBox)
 			+SVerticalBox::Slot()
-			.AutoHeight()
+			//.AutoHeight()			// Cannot use that otherwise scrollbars disapear.
 			.Padding(0.0f, 5.0f, 0.0f, 0.0f)
 			[
-				SNew(SHorizontalBox)
-				+SHorizontalBox::Slot()
-				.AutoWidth()
-				.VAlign(VAlign_Center)
+				SNew(SScrollBox)
+				.Orientation(Orient_Vertical)
+				.ScrollBarAlwaysVisible(false)
+				+ SScrollBox::Slot()
 				[
-					SNew(SWrapBox)
-					.UseAllottedSize(true)
-					+SWrapBox::Slot()
-					.Padding(5.0f)
-					.HAlign(HAlign_Center)
-					.VAlign(VAlign_Center)
+					SNew(SVerticalBox)
+					+SVerticalBox::Slot()
+					.AutoHeight()
+					.Padding(0.0f, 5.0f, 0.0f, 0.0f)
 					[
-						CheckBoxForceFullSimplification->AsShared()
+						SNew(SHorizontalBox)
+						+SHorizontalBox::Slot()
+						.AutoWidth()
+						.VAlign(VAlign_Center)
+						[
+							SNew(SWrapBox)
+							.UseAllottedSize(true)
+							+SWrapBox::Slot()
+							.Padding(5.0f)
+							.HAlign(HAlign_Center)
+							.VAlign(VAlign_Center)
+							[
+								CheckBoxForceFullSimplification->AsShared()
+							]
+						]
+						+SHorizontalBox::Slot()
+						.FillWidth(1.0f)
+						.Padding(16.0f, 0.0f)
+						.HAlign(HAlign_Left)
+						.VAlign(VAlign_Center)
+						[
+							SNew(STextBlock)
+							.ColorAndOpacity(FLinearColor::White)
+							.ShadowColorAndOpacity(FLinearColor::Black)
+							.ShadowOffset(FVector2D::UnitVector)
+							.Text(LOCTEXT("FullsimplificationLabel", "Full simplification"))
+						]
+					]
+					+SVerticalBox::Slot()
+					.AutoHeight()
+					.Padding(0.0f, 5.0f, 0.0f, 0.0f)
+					[
+						SNew(SWrapBox)
+						.UseAllottedSize(true)
+						+SWrapBox::Slot()
+						.Padding(5.0f)
+						.HAlign(HAlign_Center)
+						.VAlign(VAlign_Center)
+						[
+							ButtonApplyToPreview->AsShared()
+						]
+					]
+					+SVerticalBox::Slot()
+					.AutoHeight()
+					.Padding(0.0f, 5.0f, 0.0f, 0.0f)
+					[
+						SNew(SWrapBox)
+						.UseAllottedSize(true)
+						+ SWrapBox::Slot()
+						.Padding(5.0f)
+						.HAlign(HAlign_Center)
+						.VAlign(VAlign_Center)
+						[
+							SAssignNew(MaterialBox, SBox)
+						]
+					]
+					+SVerticalBox::Slot()
+					.AutoHeight()
+					.Padding(0.0f, 5.0f, 0.0f, 0.0f)
+					[
+						SNew(SWrapBox)
+						.UseAllottedSize(true)
+						+SWrapBox::Slot()
+						.Padding(5.0f)
+						.HAlign(HAlign_Center)
+						.VAlign(VAlign_Center)
+						[
+							DescriptionTextBlock->AsShared()
+						]
 					]
-				]
-				+SHorizontalBox::Slot()
-				.FillWidth(1.0f)
-				.Padding(16.0f, 0.0f)
-				.HAlign(HAlign_Left)
-				.VAlign(VAlign_Center)
-				[
-					SNew(STextBlock)
-					.ColorAndOpacity(FLinearColor::White)
-					.ShadowColorAndOpacity(FLinearColor::Black)
-					.ShadowOffset(FVector2D::UnitVector)
-					.Text(LOCTEXT("FullsimplificationLabel", "Full simplification"))
-				]
-			]
-			+SVerticalBox::Slot()
-			.AutoHeight()
-			.Padding(0.0f, 5.0f, 0.0f, 0.0f)
-			[
-				SNew(SWrapBox)
-				.UseAllottedSize(true)
-				+SWrapBox::Slot()
-				.Padding(5.0f)
-				.HAlign(HAlign_Center)
-				.VAlign(VAlign_Center)
-				[
-					ButtonApplyToPreview->AsShared()
 				]
 			]
 		];
@@ -119,6 +170,192 @@ SMaterialEditorStrataWidget::~SMaterialEditorStrataWidget()
 
 void SMaterialEditorStrataWidget::Tick(const FGeometry& AllottedGeometry, const double InCurrentTime, const float InDeltaTime)
 {
+	if (!bUpdateRequested || !Strata::IsStrataEnabled())
+	{
+		return;
+	}
+	bUpdateRequested = false;
+
+	FText StrataMaterialDescription;
+	if (MaterialEditorPtr.IsValid())
+	{
+		auto MaterialEditor = MaterialEditorPtr.Pin();
+
+		UMaterial* MaterialForStats = MaterialEditor->bStatsFromPreviewMaterial ? MaterialEditor->Material : MaterialEditor->OriginalMaterial;
+
+		StrataMaterialDescription = FText::FromString(FString(TEXT("StrataMaterialDescription")));
+
+		const FMaterialResource* MaterialResource = MaterialForStats->GetMaterialResource(GMaxRHIFeatureLevel);
+		if (MaterialResource)
+		{
+			FString MaterialDescription;
+
+			FMaterialShaderMap* ShaderMap = MaterialResource->GetGameThreadShaderMap();
+			if (ShaderMap)
+			{
+				const FStrataMaterialCompilationOutput& CompilationOutput = ShaderMap->GetStrataMaterialCompilationOutput();
+				const uint32 FinalPixelByteCount = CompilationOutput.StrataUintPerPixel * sizeof(uint32);
+
+				if (CompilationOutput.bMaterialOutOfBudgetHasBeenSimplified)
+				{
+					MaterialDescription += FString::Printf(TEXT("The material was OUT-OF-BUDGET so it has been fully simplified: request bytes = %i / budget = %i\r\n"),
+						CompilationOutput.RequestedBytePixePixel, CompilationOutput.PlatformBytePixePixel);
+					MaterialDescription += FString::Printf(TEXT("Final per pixel byte count   = %i\r\n"),
+						FinalPixelByteCount);
+				}
+				else
+				{
+					MaterialDescription += FString::Printf(TEXT("Material per pixel byte count= %i / budget = %i\r\n"),
+						FinalPixelByteCount, CompilationOutput.PlatformBytePixePixel);
+				}
+				MaterialDescription += FString::Printf(TEXT("BSDF Count	                  = %i\r\n"), CompilationOutput.StrataBSDFCount);
+				MaterialDescription += FString::Printf(TEXT("Local bases Count            = %i\r\n"), CompilationOutput.SharedLocalBasesCount);
+
+				switch (CompilationOutput.StrataMaterialType)
+				{
+				case 0:
+					MaterialDescription += FString::Printf(TEXT("Material complexity          = SIMPLE (diffuse, albedo, roughness)\r\n"));
+					break;
+				case 1:
+					MaterialDescription += FString::Printf(TEXT("Material complexity          = SINGLE (BSDF all features except anisotropy)\r\n"));
+					break;
+				case 2:
+					MaterialDescription += FString::Printf(TEXT("Material complexity          = COMPLEX\r\n"));
+					break;
+				default:
+					MaterialDescription += FString::Printf(TEXT("Material complexity          = UNKOWN => ERROR!\r\n"));
+				}
+
+				MaterialDescription += FString::Printf(TEXT(" \r\n"));
+				MaterialDescription += FString::Printf(TEXT(" \r\n"));
+				MaterialDescription += FString::Printf(TEXT("================================================================================\r\n"));
+				MaterialDescription += FString::Printf(TEXT("================================Detailed Output=================================\r\n"));
+				MaterialDescription += FString::Printf(TEXT("================================================================================\r\n"));
+				MaterialDescription += CompilationOutput.StrataMaterialDescription;
+
+				// Now generate a visual representation of the material from the topology tree of operators.
+				{
+					if (CompilationOutput.RootOperatorIndex >= 0)
+					{
+						std::function<const TSharedRef<SWidget>(const FStrataOperator&)> ProcessOperator = [&](const FStrataOperator& Op) -> const TSharedRef<SWidget>
+						{
+							switch (Op.OperatorType)
+							{
+							case STRATA_OPERATOR_WEIGHT:
+								return ProcessOperator(CompilationOutput.Operators[Op.LeftIndex]);
+								break;
+							case STRATA_OPERATOR_VERTICAL:
+							{
+								auto VerticalOperator = SNew(SVerticalBox)
+									+SVerticalBox::Slot()
+									.AutoHeight()
+									.VAlign(VAlign_Fill)
+									.HAlign(HAlign_Fill)
+									.Padding(0.0f, 0.0f, 0.0f, 0.0f)
+									[
+										ProcessOperator(CompilationOutput.Operators[Op.LeftIndex])
+									]
+									+ SVerticalBox::Slot()
+									.AutoHeight()
+									.VAlign(VAlign_Fill)
+									.HAlign(HAlign_Fill)
+									.Padding(0.0f, 0.0f, 0.0f, 0.0f)
+									[
+										ProcessOperator(CompilationOutput.Operators[Op.RightIndex])
+									];
+								return VerticalOperator->AsShared();
+							}
+							break;
+							case STRATA_OPERATOR_HORIZONTAL:
+							{
+								auto HorizontalOperator = SNew(SHorizontalBox)
+									+SHorizontalBox::Slot()
+									.AutoWidth()
+									.VAlign(VAlign_Fill)
+									.HAlign(HAlign_Fill)
+									[
+										ProcessOperator(CompilationOutput.Operators[Op.LeftIndex])
+									]
+									+SHorizontalBox::Slot()
+									.AutoWidth()
+									.VAlign(VAlign_Fill)
+									.HAlign(HAlign_Fill)
+									[
+										ProcessOperator(CompilationOutput.Operators[Op.RightIndex])
+									];
+								return HorizontalOperator->AsShared();
+							}
+							break;
+							case STRATA_OPERATOR_ADD:
+							{
+								auto HorizontalOperator = SNew(SHorizontalBox)
+									+SHorizontalBox::Slot()
+									.AutoWidth()
+									.VAlign(VAlign_Fill)
+									.HAlign(HAlign_Fill)
+									[
+										ProcessOperator(CompilationOutput.Operators[Op.LeftIndex])
+									]
+									+SHorizontalBox::Slot()
+									.AutoWidth()
+									.VAlign(VAlign_Fill)
+									.HAlign(HAlign_Fill)
+									[
+										ProcessOperator(CompilationOutput.Operators[Op.RightIndex])
+									];
+								return HorizontalOperator->AsShared();
+							}
+							break;
+							case STRATA_OPERATOR_BSDF:
+							{
+								auto BSDF = SNew(SErrorText)
+									.ErrorText(LOCTEXT("BSDF", "BSDF"))
+									.BackgroundColor(FSlateColor(EStyleColor::AccentGreen))
+									.ToolTipText(LOCTEXT("BSDFtooltip", "BSDF tool tip"));
+								return BSDF->AsShared();
+							}
+							break;
+							case STRATA_OPERATOR_BSDF_LEGACY:
+							{
+								auto BSDF = SNew(SErrorText)
+									.ErrorText(LOCTEXT("LegacyBSDF", "LegacyBSDF"))
+									.BackgroundColor(FSlateColor(EStyleColor::AccentBlue))
+									.ToolTipText(LOCTEXT("LegacyBSDFtooltip", "LegacyBSDF tool tip"));
+								return BSDF->AsShared();
+							}
+							break;
+							}
+
+							auto TreeOperatorError = SNew(SErrorText)
+								.ErrorText(LOCTEXT("TreeOperatorError", "Tree Operator Error"))
+								.BackgroundColor(FSlateColor(EStyleColor::AccentRed));
+							return TreeOperatorError->AsShared();
+						};
+
+						const FStrataOperator& RootOperator = CompilationOutput.Operators[CompilationOutput.RootOperatorIndex];
+						MaterialBox->SetContent(ProcessOperator(RootOperator));
+					}
+					else
+					{
+						// The tree does not looks sane so generate a visual error without crashing.
+						auto TreeError = SNew(SErrorText)
+							.ErrorText(LOCTEXT("TreeError", "Tree Error"))
+							.BackgroundColor(FSlateColor(EStyleColor::AccentRed));
+
+						const TSharedRef<SWidget>& TreeErrorAsShared = TreeError->AsShared();
+						MaterialBox->SetContent(TreeErrorAsShared);
+					}
+				}
+			}
+			else
+			{
+				MaterialDescription = TEXT("Shader map not found.");
+				MaterialBox->SetContent(SNullWidget::NullWidget);
+			}
+
+			DescriptionTextBlock->SetText(FText::FromString(MaterialDescription));
+		}
+	}
 }
 
 FReply SMaterialEditorStrataWidget::OnButtonApplyToPreview()
diff --git a/Engine/Source/Editor/MaterialEditor/Private/SMaterialEditorStrataWidget.h b/Engine/Source/Editor/MaterialEditor/Private/SMaterialEditorStrataWidget.h
--- a/Engine/Source/Editor/MaterialEditor/Private/SMaterialEditorStrataWidget.h
+++ b/Engine/Source/Editor/MaterialEditor/Private/SMaterialEditorStrataWidget.h
@@ -7,8 +7,11 @@
 
 #include "Widgets/DeclarativeSyntaxSupport.h"
 #include "Widgets/SCompoundWidget.h"
+#include "Widgets/SBoxPanel.h"
 #include "Widgets/Input/SButton.h"
 #include "Widgets/Input/SCheckBox.h"
+#include "Widgets/Text/STextBlock.h"
+#include "Widgets/Layout/SBox.h"
 
 class FMaterialEditor;
 
@@ -22,6 +25,8 @@ public:
 
 	void Construct(const FArguments& InArgs, TWeakPtr<FMaterialEditor> InMaterialEditorPtr);
 
+	void UpdateFromMaterial() { bUpdateRequested = true; }
+
 	/** Gets the widget contents of the app */
 	virtual TSharedRef<SWidget> GetContent();
 
@@ -36,9 +41,15 @@ private:
 
 	TSharedPtr<class SButton> ButtonApplyToPreview;
 
+	TSharedPtr<class STextBlock> DescriptionTextBlock;
+
+	TSharedPtr<class SBox> MaterialBox;
+
 	FReply OnButtonApplyToPreview();
 
 	/** Pointer back to the material editor that owns this */
 	TWeakPtr<FMaterialEditor> MaterialEditorPtr;
+
+	bool bUpdateRequested = true;
 };
 
diff --git a/Engine/Source/Editor/UnrealEd/Public/GraphEditor.h b/Engine/Source/Editor/UnrealEd/Public/GraphEditor.h
--- a/Engine/Source/Editor/UnrealEd/Public/GraphEditor.h
+++ b/Engine/Source/Editor/UnrealEd/Public/GraphEditor.h
@@ -68,6 +68,8 @@ struct FGraphAppearanceInfo
 	FText ReadOnlyText;
 	/** Text to display if the graph is empty (to guide the user on what to do) */
 	FText InstructionText;
+	/** Bottom left warning text used for instance by Substrate */
+	FText WarningText;
 	/** Allows graphs to nicely fade instruction text (or completely hide it). */
 	TAttribute<float> InstructionFade;
 };
diff --git a/Engine/Source/Runtime/Engine/Private/Materials/HLSLMaterialTranslator.cpp b/Engine/Source/Runtime/Engine/Private/Materials/HLSLMaterialTranslator.cpp
--- a/Engine/Source/Runtime/Engine/Private/Materials/HLSLMaterialTranslator.cpp
+++ b/Engine/Source/Runtime/Engine/Private/Materials/HLSLMaterialTranslator.cpp
@@ -1848,7 +1848,7 @@ bool FHLSLMaterialTranslator::Translate()
 			}
 			else
 			{
-				MaterialCompilationOutput.StrataMaterialDescription = "";
+				MaterialCompilationOutput.StrataMaterialCompilationOutput.StrataMaterialDescription = "";
 				ResourcesString += "// No Strata material provided\r\n";
 
 				// Adde default strata functions
@@ -2304,7 +2304,8 @@ void FHLSLMaterialTranslator::GetMaterialEnvironment(EShaderPlatform InPlatform,
 			uint8 RequestedSharedLocalBasesCount = 0;
 			StrataCtx.StrataEvaluateSharedLocalBases(this, RequestedSharedLocalBasesCount, &OutEnvironment);
 
-			// Now write some feedback to the user
+#if WITH_EDITOR
+			// Now write some feedback to the user, but only produce debug string if in editor
 			{
 				// Output some debug info as comment in code and in the material stat window
 				const uint32 StrataBytePerPixel = Strata::GetBytePerPixel(InPlatform);
@@ -2383,9 +2384,15 @@ void FHLSLMaterialTranslator::GetMaterialEnvironment(EShaderPlatform InPlatform,
 				ResourcesString += TEXT("*/");
 
 				StrataMaterialDescription += StrataMaterialContextDescription;
+
+				if (StrataCompilationContextIndex == EStrataCompilationContext::SCC_Default)
+				{
+					MaterialCompilationOutput.StrataMaterialCompilationOutput.SharedLocalBasesCount = StrataCtx.FinalUsedSharedLocalBasesCount;
+				}
 			}
+#endif // WITH_EDITOR
 		}
-		MaterialCompilationOutput.StrataMaterialDescription = StrataMaterialDescription;
+		MaterialCompilationOutput.StrataMaterialCompilationOutput.StrataMaterialDescription = StrataMaterialDescription;
 	}
 
 	OutEnvironment.SetDefine(TEXT("TEXTURE_SAMPLE_DEBUG"), IsDebugTextureSampleEnabled() ? TEXT("1") : TEXT("0"));
@@ -11763,9 +11770,28 @@ bool FHLSLMaterialTranslator::FStrataCompilationContext::StrataGenerateDerivedMa
 			if (CompilationContextIndex == EStrataCompilationContext::SCC_Default)
 			{
 				// Only write those data for the default material
-				Compiler->MaterialCompilationOutput.StrataMaterialType = bStrataMaterialIsSimple ? 0 : bStrataMaterialIsSingle ? 1 : 2;
-				Compiler->MaterialCompilationOutput.StrataBSDFCount = StrataMaterialBSDFCount;
-				Compiler->MaterialCompilationOutput.StrataUintPerPixel = uint8(FMath::Clamp(RequestedSizeInUint, 0u, 0xFF));
+				Compiler->MaterialCompilationOutput.StrataMaterialCompilationOutput.StrataMaterialType = bStrataMaterialIsSimple ? 0 : bStrataMaterialIsSingle ? 1 : 2;
+				Compiler->MaterialCompilationOutput.StrataMaterialCompilationOutput.StrataBSDFCount = StrataMaterialBSDFCount;
+				Compiler->MaterialCompilationOutput.StrataMaterialCompilationOutput.StrataUintPerPixel = uint8(FMath::Clamp(RequestedSizeInUint, 0u, 0xFF));
+
+#if WITH_EDITOR
+				Compiler->MaterialCompilationOutput.StrataMaterialCompilationOutput.SharedLocalBasesCount = 0; // FinalUsedSharedLocalBasesCount is not valid yet
+				Compiler->MaterialCompilationOutput.StrataMaterialCompilationOutput.RequestedBytePixePixel = StrataMaterialRequestedSizeByte;
+				Compiler->MaterialCompilationOutput.StrataMaterialCompilationOutput.PlatformBytePixePixel = StrataBytePerPixel;
+
+				Compiler->MaterialCompilationOutput.StrataMaterialCompilationOutput.bMaterialOutOfBudgetHasBeenSimplified |= !StrataSimplificationStatus.bMaterialFitsInMemoryBudget;
+
+				if (OperatorCount <= STRATA_COMPILATION_OUTPUT_MAX_OPERATOR)
+				{
+					int32 OperatorIndex = 0;
+					for (auto& It : StrataMaterialExpressionRegisteredOperators)
+					{
+						Compiler->MaterialCompilationOutput.StrataMaterialCompilationOutput.Operators[OperatorIndex] = StrataMaterialExpressionRegisteredOperators[OperatorIndex];
+						OperatorIndex++;
+					}
+				}
+				Compiler->MaterialCompilationOutput.StrataMaterialCompilationOutput.RootOperatorIndex = RootIndex;
+#endif // EDITOR_ONLY
 			}
 		}
 	}
diff --git a/Engine/Source/Runtime/Engine/Public/MaterialCompiler.h b/Engine/Source/Runtime/Engine/Public/MaterialCompiler.h
--- a/Engine/Source/Runtime/Engine/Public/MaterialCompiler.h
+++ b/Engine/Source/Runtime/Engine/Public/MaterialCompiler.h
@@ -77,107 +77,6 @@ enum EStrataMaterialExportContext : uint8
 	SMEC_Translucent = 1
 };
 
-struct FStrataRegisteredSharedLocalBasis
-{
-	int32 NormalCodeChunk;
-	int32 TangentCodeChunk;
-	uint64 NormalCodeChunkHash;
-	uint64 TangentCodeChunkHash;
-	uint8 GraphSharedLocalBasisIndex;
-
-	FStrataRegisteredSharedLocalBasis()
-	{
-		NormalCodeChunk = INDEX_NONE;
-		TangentCodeChunk = INDEX_NONE;
-		NormalCodeChunkHash = 0;
-		TangentCodeChunkHash = 0;
-		GraphSharedLocalBasisIndex = 0;
-	}
-};
-
-struct FStrataOperator
-{
-	int32 OperatorType;
-	bool bNodeRequestParameterBlending;
-
-	int32 Index;			// Index into the array of operators
-	int32 ParentIndex;		// Parent operator index
-	int32 LeftIndex;		// Left child operator index
-	int32 RightIndex;		// Right child operator index
-	int32 ThicknessIndex;	// Thickness expression index
-
-	// Data used for BSDF type nodes only
-	int32 BSDFIndex;		// Index in the array of BSDF if a BSDF operator
-	uint8 BSDFType;
-	FStrataRegisteredSharedLocalBasis BSDFRegisteredSharedLocalBasis;
-	bool  bBSDFHasSSS;
-	bool  bBSDFHasMFPPluggedIn;
-	bool  bBSDFHasEdgeColor;
-	bool  bBSDFHasFuzz;
-	bool  bBSDFHasSecondRoughnessOrSimpleClearCoat;
-	bool  bBSDFHasAnisotropy;
-
-	// Data derived after the tree has been built.
-	int32 MaxDistanceFromLeaves;
-	int32 LayerDepth;
-	bool bIsTop;
-	bool bIsBottom;
-	bool bUseParameterBlending;				// True when part of a sub tree where parameter blending is in use
-	bool bRootOfParameterBlendingSubTree;	// True when the root of a sub tree where parameter blending is in use. Only this node will register a BSDF
-
-	FStrataOperator()
-	{
-		OperatorType = INDEX_NONE;
-		bNodeRequestParameterBlending = false;
-		Index = INDEX_NONE;
-		ParentIndex = INDEX_NONE;
-		LeftIndex = INDEX_NONE;
-		RightIndex = INDEX_NONE;
-		ThicknessIndex = INDEX_NONE;
-
-		BSDFIndex = INDEX_NONE;
-		BSDFType = 0;
-		bBSDFHasSSS = false;
-		bBSDFHasMFPPluggedIn = false;
-		bBSDFHasEdgeColor = false;
-		bBSDFHasFuzz = false;
-		bBSDFHasSecondRoughnessOrSimpleClearCoat = false;
-		bBSDFHasAnisotropy = false;
-
-		MaxDistanceFromLeaves = 0;
-		LayerDepth = 0;
-		bIsTop = false;
-		bIsBottom = false;
-		bUseParameterBlending = false;
-		bRootOfParameterBlendingSubTree = false;
-	}
-
-	void CombineFlagsForParameterBlending(FStrataOperator& A, FStrataOperator& B)
-	{
-		bBSDFHasSSS = A.bBSDFHasSSS || B.bBSDFHasSSS;
-		bBSDFHasMFPPluggedIn = A.bBSDFHasMFPPluggedIn || B.bBSDFHasMFPPluggedIn;
-		bBSDFHasEdgeColor = A.bBSDFHasEdgeColor || B.bBSDFHasEdgeColor;
-		bBSDFHasFuzz = A.bBSDFHasFuzz || B.bBSDFHasFuzz;
-		bBSDFHasSecondRoughnessOrSimpleClearCoat = A.bBSDFHasSecondRoughnessOrSimpleClearCoat || B.bBSDFHasSecondRoughnessOrSimpleClearCoat;
-		bBSDFHasAnisotropy = A.bBSDFHasAnisotropy || B.bBSDFHasAnisotropy;
-	}
-
-	void CopyFlagsForParameterBlending(FStrataOperator& A)
-	{
-		bBSDFHasSSS = A.bBSDFHasSSS;
-		bBSDFHasMFPPluggedIn = A.bBSDFHasMFPPluggedIn;
-		bBSDFHasEdgeColor = A.bBSDFHasEdgeColor;
-		bBSDFHasFuzz = A.bBSDFHasFuzz;
-		bBSDFHasSecondRoughnessOrSimpleClearCoat = A.bBSDFHasSecondRoughnessOrSimpleClearCoat;
-		bBSDFHasAnisotropy = A.bBSDFHasAnisotropy;
-	}
-
-	bool IsDiscarded() const
-	{
-		return bUseParameterBlending && !bRootOfParameterBlendingSubTree;
-	}
-};
-
 /** 
  * The interface used to translate material expressions into executable code. 
  * Note: Most member functions should be pure virtual to force a FProxyMaterialCompiler override!
@@ -758,7 +657,11 @@ public:
 protected:
 	FString GetStrataSharedLocalBasisIndexMacroInner(const FStrataRegisteredSharedLocalBasis& SharedLocalBasis, uint32 Mode)
 	{
+#if WITH_EDITOR
 		return FString::Printf(TEXT("SHAREDLOCALBASIS_INDEX_%u_%u"), SharedLocalBasis.GraphSharedLocalBasisIndex, Mode);
+#else
+		return TEXT("");
+#endif
 	}
 
 
diff --git a/Engine/Source/Runtime/Engine/Public/MaterialShared.h b/Engine/Source/Runtime/Engine/Public/MaterialShared.h
--- a/Engine/Source/Runtime/Engine/Public/MaterialShared.h
+++ b/Engine/Source/Runtime/Engine/Public/MaterialShared.h
@@ -37,6 +37,7 @@
 #include "ShaderCompilerCore.h"
 #include "PSOPrecache.h"
 #include "UObject/ObjectMacros.h"
+#include "Rendering/StrataMaterialShared.h"
 
 #if UE_ENABLE_INCLUDE_ORDER_DEPRECATED_IN_5_2
 #include "RHI.h"
@@ -673,8 +674,8 @@ public:
 		EstimatedNumVirtualTextureLookups(0),
 		NumUsedUVScalars(0),
 		NumUsedCustomInterpolatorScalars(0),
-		StrataMaterialDescription(),
 #endif
+		StrataMaterialCompilationOutput(),
 		UsedDBufferTextures(0),
 		RuntimeVirtualTextureOutputAttributeMask(0),
 		bNeedsSceneTextures(false),
@@ -690,9 +691,6 @@ public:
 		bUsesVertexInterpolator(false),
 		bHasRuntimeVirtualTextureOutputNode(false),
 		bUsesAnisotropy(false),
-		StrataMaterialType(0),
-		StrataBSDFCount(0),
-		StrataUintPerPixel(0),
 		bUsesDisplacement(false)
 	{}
 
@@ -759,7 +757,7 @@ public:
 	LAYOUT_FIELD_EDITORONLY(uint8, NumUsedCustomInterpolatorScalars);
 
 	/** The Strata material layout */
-	LAYOUT_FIELD_EDITORONLY(FMemoryImageString, StrataMaterialDescription);
+	LAYOUT_FIELD(FStrataMaterialCompilationOutput, StrataMaterialCompilationOutput);
 
 	/** Bitfield of used DBuffer textures . */
 	LAYOUT_FIELD(uint8, UsedDBufferTextures);
@@ -805,15 +803,6 @@ public:
 
 	/** true if the material uses non 0 anisotropy value */
 	LAYOUT_BITFIELD(uint8, bUsesAnisotropy, 1);
-
-	/** Strata material type, at compile time (0:simple, 1:single, 2: complex) */
-	LAYOUT_BITFIELD(uint8, StrataMaterialType, 2);
-
-	/** Strata BSDF count, at compile time (0-7) */
-	LAYOUT_BITFIELD(uint8, StrataBSDFCount, 3);
-
-	/** Strata uint per pixel, at compile time (0-7) */
-	LAYOUT_BITFIELD(uint8, StrataUintPerPixel, 8);
 	
 	/** Whether the material uses scalar displacement. */
 	LAYOUT_BITFIELD(uint8, bUsesDisplacement, 1);
@@ -1377,16 +1366,17 @@ public:
 	bool UsesVelocitySceneTexture() const { return GetContent()->MaterialCompilationOutput.UsesVelocitySceneTexture(); }
 	bool UsesDistanceCullFade() const { return GetContent()->MaterialCompilationOutput.bUsesDistanceCullFade; }
 	bool UsesAnisotropy() const { return GetContent()->MaterialCompilationOutput.bUsesAnisotropy; }
-	uint8 GetStrataMaterialType() const { return GetContent()->MaterialCompilationOutput.StrataMaterialType; }
-	uint8 GetStrataBSDFCount() const { return GetContent()->MaterialCompilationOutput.StrataBSDFCount; }
-	uint8 GetStrataUintPerPixel() const { return GetContent()->MaterialCompilationOutput.StrataUintPerPixel; }
+
+	const FStrataMaterialCompilationOutput& GetStrataMaterialCompilationOutput() const { return GetContent()->MaterialCompilationOutput.StrataMaterialCompilationOutput; }
+	uint8 GetStrataMaterialType() const { return GetStrataMaterialCompilationOutput().StrataMaterialType; }
+	uint8 GetStrataBSDFCount() const { return GetStrataMaterialCompilationOutput().StrataBSDFCount; }
+	uint8 GetStrataUintPerPixel() const { return GetStrataMaterialCompilationOutput().StrataUintPerPixel; }
 	
 #if WITH_EDITOR
 	uint32 GetNumUsedUVScalars() const { return GetContent()->MaterialCompilationOutput.NumUsedUVScalars; }
 	uint32 GetNumUsedCustomInterpolatorScalars() const { return GetContent()->MaterialCompilationOutput.NumUsedCustomInterpolatorScalars; }
 	void GetEstimatedNumTextureSamples(uint32& VSSamples, uint32& PSSamples) const { VSSamples = GetContent()->MaterialCompilationOutput.EstimatedNumTextureSamplesVS; PSSamples = GetContent()->MaterialCompilationOutput.EstimatedNumTextureSamplesPS; }
 	uint32 GetEstimatedNumVirtualTextureLookups() const { return GetContent()->MaterialCompilationOutput.EstimatedNumVirtualTextureLookups; }
-	const FString GetStrataMaterialDescription() const { return GetContent()->MaterialCompilationOutput.StrataMaterialDescription; }
 #endif
 	uint32 GetNumVirtualTextureStacks() const { return GetContent()->MaterialCompilationOutput.UniformExpressionSet.VTStacks.Num(); }
 	uint8 GetRuntimeVirtualTextureOutputAttributeMask() const { return GetContent()->MaterialCompilationOutput.RuntimeVirtualTextureOutputAttributeMask; }
diff --git a/Engine/Source/Runtime/Engine/Public/Rendering/StrataMaterialShared.cpp b/Engine/Source/Runtime/Engine/Public/Rendering/StrataMaterialShared.cpp
new file mode 100644
--- /dev/null
+++ b/Engine/Source/Runtime/Engine/Public/Rendering/StrataMaterialShared.cpp
@@ -0,0 +1,99 @@
+// Copyright Epic Games, Inc. All Rights Reserved.
+
+#include "StrataMaterialShared.h"
+
+
+IMPLEMENT_TYPE_LAYOUT(FStrataMaterialCompilationOutput);
+
+FStrataMaterialCompilationOutput::FStrataMaterialCompilationOutput()
+	: StrataMaterialType(0)
+	, StrataBSDFCount(0)
+	, StrataUintPerPixel(0)
+#if WITH_EDITOR
+	, StrataMaterialDescription()
+	, SharedLocalBasesCount(0)
+	, RequestedBytePixePixel(0)
+	, PlatformBytePixePixel(0)
+	, bMaterialOutOfBudgetHasBeenSimplified(0)
+#endif
+{
+}
+
+
+IMPLEMENT_TYPE_LAYOUT(FStrataOperator);
+
+FStrataOperator::FStrataOperator()
+{
+#if WITH_EDITOR
+	OperatorType = INDEX_NONE;
+	bNodeRequestParameterBlending = false;
+	Index = INDEX_NONE;
+	ParentIndex = INDEX_NONE;
+	LeftIndex = INDEX_NONE;
+	RightIndex = INDEX_NONE;
+	ThicknessIndex = INDEX_NONE;
+
+	BSDFIndex = INDEX_NONE;
+	BSDFType = 0;
+	bBSDFHasSSS = false;
+	bBSDFHasMFPPluggedIn = false;
+	bBSDFHasEdgeColor = false;
+	bBSDFHasFuzz = false;
+	bBSDFHasSecondRoughnessOrSimpleClearCoat = false;
+	bBSDFHasAnisotropy = false;
+
+	MaxDistanceFromLeaves = 0;
+	LayerDepth = 0;
+	bIsTop = false;
+	bIsBottom = false;
+	bUseParameterBlending = false;
+	bRootOfParameterBlendingSubTree = false;
+#endif
+}
+
+void FStrataOperator::CombineFlagsForParameterBlending(FStrataOperator& A, FStrataOperator& B)
+{
+#if WITH_EDITOR
+	bBSDFHasSSS = A.bBSDFHasSSS || B.bBSDFHasSSS;
+	bBSDFHasMFPPluggedIn = A.bBSDFHasMFPPluggedIn || B.bBSDFHasMFPPluggedIn;
+	bBSDFHasEdgeColor = A.bBSDFHasEdgeColor || B.bBSDFHasEdgeColor;
+	bBSDFHasFuzz = A.bBSDFHasFuzz || B.bBSDFHasFuzz;
+	bBSDFHasSecondRoughnessOrSimpleClearCoat = A.bBSDFHasSecondRoughnessOrSimpleClearCoat || B.bBSDFHasSecondRoughnessOrSimpleClearCoat;
+	bBSDFHasAnisotropy = A.bBSDFHasAnisotropy || B.bBSDFHasAnisotropy;
+#endif
+}
+
+void FStrataOperator::CopyFlagsForParameterBlending(FStrataOperator& A)
+{
+#if WITH_EDITOR
+	bBSDFHasSSS = A.bBSDFHasSSS;
+	bBSDFHasMFPPluggedIn = A.bBSDFHasMFPPluggedIn;
+	bBSDFHasEdgeColor = A.bBSDFHasEdgeColor;
+	bBSDFHasFuzz = A.bBSDFHasFuzz;
+	bBSDFHasSecondRoughnessOrSimpleClearCoat = A.bBSDFHasSecondRoughnessOrSimpleClearCoat;
+	bBSDFHasAnisotropy = A.bBSDFHasAnisotropy;
+#endif
+}
+
+bool FStrataOperator::IsDiscarded() const
+{
+#if WITH_EDITOR
+	return bUseParameterBlending && !bRootOfParameterBlendingSubTree;
+#else
+	return true;
+#endif
+}
+
+
+IMPLEMENT_TYPE_LAYOUT(FStrataRegisteredSharedLocalBasis);
+
+FStrataRegisteredSharedLocalBasis::FStrataRegisteredSharedLocalBasis()
+{
+#if WITH_EDITOR
+	NormalCodeChunk = INDEX_NONE;
+	TangentCodeChunk = INDEX_NONE;
+	NormalCodeChunkHash = 0;
+	TangentCodeChunkHash = 0;
+	GraphSharedLocalBasisIndex = 0;
+#endif
+}
\ No newline at end of file
diff --git a/Engine/Source/Runtime/Engine/Public/Rendering/StrataMaterialShared.h b/Engine/Source/Runtime/Engine/Public/Rendering/StrataMaterialShared.h
new file mode 100644
--- /dev/null
+++ b/Engine/Source/Runtime/Engine/Public/Rendering/StrataMaterialShared.h
@@ -0,0 +1,114 @@
+// Copyright Epic Games, Inc. All Rights Reserved.
+
+#pragma once
+
+#include "CoreMinimal.h"
+#include "StrataDefinitions.h"
+#include "Serialization/MemoryImage.h"
+
+
+// Structures in this files are only used a compilation result return by the compiler.
+// They are also used to present material information in the editor UI.
+
+
+struct FStrataRegisteredSharedLocalBasis
+{
+	DECLARE_TYPE_LAYOUT(FStrataRegisteredSharedLocalBasis, NonVirtual);
+public:
+	FStrataRegisteredSharedLocalBasis();
+
+	LAYOUT_FIELD_EDITORONLY(int32, NormalCodeChunk);
+	LAYOUT_FIELD_EDITORONLY(int32, TangentCodeChunk);
+	LAYOUT_FIELD_EDITORONLY(uint64, NormalCodeChunkHash);
+	LAYOUT_FIELD_EDITORONLY(uint64, TangentCodeChunkHash);
+	LAYOUT_FIELD_EDITORONLY(uint8, GraphSharedLocalBasisIndex);
+};
+
+struct FStrataOperator
+{
+	DECLARE_TYPE_LAYOUT(FStrataOperator, NonVirtual);
+public:
+	FStrataOperator();
+
+	LAYOUT_FIELD_EDITORONLY(int32, OperatorType);
+	LAYOUT_BITFIELD_EDITORONLY(uint8, bNodeRequestParameterBlending, 1);
+
+	LAYOUT_FIELD_EDITORONLY(int32, Index);			// Index into the array of operators
+	LAYOUT_FIELD_EDITORONLY(int32, ParentIndex);	// Parent operator index
+	LAYOUT_FIELD_EDITORONLY(int32, LeftIndex);		// Left child operator index
+	LAYOUT_FIELD_EDITORONLY(int32, RightIndex);		// Right child operator index
+	LAYOUT_FIELD_EDITORONLY(int32, ThicknessIndex);	// Thickness expression index
+
+	// Data used for BSDF type nodes only
+	LAYOUT_FIELD_EDITORONLY(int32, BSDFIndex);		// Index in the array of BSDF if a BSDF operator
+	LAYOUT_FIELD_EDITORONLY(int32, BSDFType);
+	LAYOUT_FIELD_EDITORONLY(FStrataRegisteredSharedLocalBasis, BSDFRegisteredSharedLocalBasis);
+	LAYOUT_BITFIELD_EDITORONLY(uint8, bBSDFHasSSS, 1);
+	LAYOUT_BITFIELD_EDITORONLY(uint8, bBSDFHasMFPPluggedIn, 1);
+	LAYOUT_BITFIELD_EDITORONLY(uint8, bBSDFHasEdgeColor, 1);
+	LAYOUT_BITFIELD_EDITORONLY(uint8, bBSDFHasFuzz, 1);
+	LAYOUT_BITFIELD_EDITORONLY(uint8, bBSDFHasSecondRoughnessOrSimpleClearCoat, 1);
+	LAYOUT_BITFIELD_EDITORONLY(uint8, bBSDFHasAnisotropy, 1);
+
+	// Data derived after the tree has been built.
+	LAYOUT_FIELD_EDITORONLY(int32, MaxDistanceFromLeaves);
+	LAYOUT_FIELD_EDITORONLY(int32, LayerDepth);
+	LAYOUT_BITFIELD_EDITORONLY(uint8, bIsTop, 1);
+	LAYOUT_BITFIELD_EDITORONLY(uint8, bIsBottom, 1);
+
+	LAYOUT_BITFIELD_EDITORONLY(uint8, bUseParameterBlending, 1);			// True when part of a sub tree where parameter blending is in use
+	LAYOUT_BITFIELD_EDITORONLY(uint8, bRootOfParameterBlendingSubTree, 1);	// True when the root of a sub tree where parameter blending is in use. Only this node will register a BSDF
+
+	void CombineFlagsForParameterBlending(FStrataOperator& A, FStrataOperator& B);
+
+	void CopyFlagsForParameterBlending(FStrataOperator& A);
+
+	bool IsDiscarded() const;
+};
+
+#define STRATA_COMPILATION_OUTPUT_MAX_OPERATOR 32
+
+struct FStrataMaterialCompilationOutput
+{
+	DECLARE_TYPE_LAYOUT(FStrataMaterialCompilationOutput, NonVirtual);
+public:
+
+	FStrataMaterialCompilationOutput();
+
+	////
+	//// The following data is required at runtime
+	////
+
+	/** Strata material type, at compile time (0:simple, 1:single, 2: complex) */
+	LAYOUT_BITFIELD(uint8, StrataMaterialType, 2);
+
+	/** Strata BSDF count, at compile time (0-7) */
+	LAYOUT_BITFIELD(uint8, StrataBSDFCount, 3);
+
+	/** Strata uint per pixel, at compile time (0-255) */
+	LAYOUT_BITFIELD(uint8, StrataUintPerPixel, 8);
+
+	////
+	//// The following data is only needed when compiling with the editor.
+	////
+
+	// STRATA_TODO Strata Tree topology
+	// STRATA_TODO corresponding uint8 ShadingModel? IsThin?
+
+	/** The Strata verbose description */
+	LAYOUT_FIELD_EDITORONLY(FMemoryImageString, StrataMaterialDescription);
+
+	/** The number of local normal/tangent bases */
+	LAYOUT_FIELD_EDITORONLY(uint8, SharedLocalBasesCount);
+	/** Material requested byte count per pixel */
+	LAYOUT_FIELD_EDITORONLY(uint8, RequestedBytePixePixel);
+	/** The byte count per pixel supported by the platform the material has been compiled against */
+	LAYOUT_FIELD_EDITORONLY(uint8, PlatformBytePixePixel);
+
+	/** The byte per pixel count supported by the platform the material has been compiled against */
+	LAYOUT_BITFIELD_EDITORONLY(uint8, bMaterialOutOfBudgetHasBeenSimplified, 1);
+
+	LAYOUT_FIELD_EDITORONLY(uint8, RootOperatorIndex);
+	LAYOUT_ARRAY_EDITORONLY(FStrataOperator, Operators, STRATA_COMPILATION_OUTPUT_MAX_OPERATOR);
+};
+
