MaterialEditor - Fixed LOCTEXT conflict

diff --git a/Engine/Source/Editor/MaterialEditor/Private/MaterialEditor.cpp b/Engine/Source/Editor/MaterialEditor/Private/MaterialEditor.cpp
--- a/Engine/Source/Editor/MaterialEditor/Private/MaterialEditor.cpp
+++ b/Engine/Source/Editor/MaterialEditor/Private/MaterialEditor.cpp
@@ -7046,7 +7046,7 @@ FGraphAppearanceInfo FMaterialEditor::GetGraphAppearance() const
 			const FStrataMaterialCompilationOutput& CompilationOutput = ShaderMap->GetStrataMaterialCompilationOutput();
 			if (CompilationOutput.bMaterialOutOfBudgetHasBeenSimplified)
 			{
-				AppearanceInfo.WarningText = LOCTEXT("AppearanceCornerText_Material", "Substrate material was out of budget and has been simplified.");
+				AppearanceInfo.WarningText = LOCTEXT("AppearanceWarningText_Material", "Substrate material was out of budget and has been simplified.");
 
 			}
 		}
