Substrate
- deleted details output from message log
- material in Substrate HUD no display details
- Simplification button is next to the tick box.
- !! not using bitfield which seems to cause some weird bit shift sometimes. Reported !!

diff --git a/Engine/Source/Editor/MaterialEditor/Private/MaterialEditor.cpp b/Engine/Source/Editor/MaterialEditor/Private/MaterialEditor.cpp
--- a/Engine/Source/Editor/MaterialEditor/Private/MaterialEditor.cpp
+++ b/Engine/Source/Editor/MaterialEditor/Private/MaterialEditor.cpp
@@ -3025,15 +3025,6 @@ void FMaterialEditor::UpdateMaterialinfoList_Old()
 
 				if (FMaterialShaderMap* ShaderMap = MaterialResource->GetGameThreadShaderMap())
 				{
-					const FString StrataMaterialDescription = ShaderMap->GetStrataMaterialCompilationOutput().StrataMaterialDescription;
-
-					if (StrataMaterialDescription.Len())
-					{
-						TSharedRef<FTokenizedMessage> Line = FTokenizedMessage::Create(EMessageSeverity::Info);
-						Line->AddToken(FTextToken::Create(FText::FromString(StrataMaterialDescription)));
-						Messages.Add(Line);
-					}
-
 					// Add shader count
 					FString ShaderCountString = FString::Printf(TEXT("Shader Count: %u"), ShaderMap->GetShaderNum());
 					TSharedRef<FTokenizedMessage> ShaderCountLine = FTokenizedMessage::Create(EMessageSeverity::Info);
diff --git a/Engine/Source/Editor/MaterialEditor/Private/SMaterialEditorStrataWidget.cpp b/Engine/Source/Editor/MaterialEditor/Private/SMaterialEditorStrataWidget.cpp
--- a/Engine/Source/Editor/MaterialEditor/Private/SMaterialEditorStrataWidget.cpp
+++ b/Engine/Source/Editor/MaterialEditor/Private/SMaterialEditorStrataWidget.cpp
@@ -56,7 +56,6 @@ void SMaterialEditorStrataWidget::Construct(const FArguments& InArgs, TWeakPtr<F
 					SNew(SVerticalBox)
 					+SVerticalBox::Slot()
 					.AutoHeight()
-					.Padding(0.0f, 5.0f, 0.0f, 0.0f)
 					[
 						SNew(SHorizontalBox)
 						+SHorizontalBox::Slot()
@@ -85,19 +84,21 @@ void SMaterialEditorStrataWidget::Construct(const FArguments& InArgs, TWeakPtr<F
 							.ShadowOffset(FVector2D::UnitVector)
 							.Text(LOCTEXT("FullsimplificationLabel", "Full simplification"))
 						]
-					]
-					+SVerticalBox::Slot()
-					.AutoHeight()
-					.Padding(0.0f, 5.0f, 0.0f, 0.0f)
-					[
-						SNew(SWrapBox)
-						.UseAllottedSize(true)
-						+SWrapBox::Slot()
-						.Padding(5.0f)
-						.HAlign(HAlign_Center)
+						+SHorizontalBox::Slot()
+						.FillWidth(1.0f)
+						.Padding(16.0f, 0.0f)
+						.HAlign(HAlign_Left)
 						.VAlign(VAlign_Center)
 						[
-							ButtonApplyToPreview->AsShared()
+							SNew(SWrapBox)
+							.UseAllottedSize(true)
+							+ SWrapBox::Slot()
+							.Padding(5.0f)
+							.HAlign(HAlign_Center)
+							.VAlign(VAlign_Center)
+							[
+								ButtonApplyToPreview->AsShared()
+							]
 						]
 					]
 					+SVerticalBox::Slot()
@@ -251,7 +252,7 @@ void SMaterialEditorStrataWidget::Tick(const FGeometry& AllottedGeometry, const
 									.AutoHeight()
 									.VAlign(VAlign_Fill)
 									.HAlign(HAlign_Fill)
-									.Padding(0.0f, 0.0f, 0.0f, 0.0f)
+									.Padding(0.0f, 0.0f, 1.0f, 1.0f)
 									[
 										ProcessOperator(CompilationOutput.Operators[Op.LeftIndex])
 									]
@@ -259,7 +260,7 @@ void SMaterialEditorStrataWidget::Tick(const FGeometry& AllottedGeometry, const
 									.AutoHeight()
 									.VAlign(VAlign_Fill)
 									.HAlign(HAlign_Fill)
-									.Padding(0.0f, 0.0f, 0.0f, 0.0f)
+									.Padding(0.0f, 0.0f, 1.0f, 1.0f)
 									[
 										ProcessOperator(CompilationOutput.Operators[Op.RightIndex])
 									];
@@ -273,6 +274,7 @@ void SMaterialEditorStrataWidget::Tick(const FGeometry& AllottedGeometry, const
 									.AutoWidth()
 									.VAlign(VAlign_Fill)
 									.HAlign(HAlign_Fill)
+									.Padding(0.0f, 0.0f, 1.0f, 1.0f)
 									[
 										ProcessOperator(CompilationOutput.Operators[Op.LeftIndex])
 									]
@@ -280,6 +282,7 @@ void SMaterialEditorStrataWidget::Tick(const FGeometry& AllottedGeometry, const
 									.AutoWidth()
 									.VAlign(VAlign_Fill)
 									.HAlign(HAlign_Fill)
+									.Padding(0.0f, 0.0f, 1.0f, 1.0f)
 									[
 										ProcessOperator(CompilationOutput.Operators[Op.RightIndex])
 									];
@@ -293,6 +296,7 @@ void SMaterialEditorStrataWidget::Tick(const FGeometry& AllottedGeometry, const
 									.AutoWidth()
 									.VAlign(VAlign_Fill)
 									.HAlign(HAlign_Fill)
+									.Padding(0.0f, 0.0f, 1.0f, 1.0f)
 									[
 										ProcessOperator(CompilationOutput.Operators[Op.LeftIndex])
 									]
@@ -300,27 +304,40 @@ void SMaterialEditorStrataWidget::Tick(const FGeometry& AllottedGeometry, const
 									.AutoWidth()
 									.VAlign(VAlign_Fill)
 									.HAlign(HAlign_Fill)
+									.Padding(0.0f, 0.0f, 1.0f, 1.0f)
 									[
 										ProcessOperator(CompilationOutput.Operators[Op.RightIndex])
 									];
 								return HorizontalOperator->AsShared();
 							}
 							break;
+							case STRATA_OPERATOR_BSDF_LEGACY:	// legacy BSDF should have been converted to BSDF already.
 							case STRATA_OPERATOR_BSDF:
 							{
+								FString BSDFDesc = FString::Printf(TEXT("BSDF (%s%s%s%s%s%s)")
+									, Op.bBSDFHasEdgeColor ? TEXT("F90 ") : TEXT("")
+									, Op.bBSDFHasSSS ? TEXT("SSS ") : TEXT("")
+									, Op.bBSDFHasMFPPluggedIn ? TEXT("MFP ") : TEXT("")
+									, Op.bBSDFHasAnisotropy ? TEXT("Ani ") : TEXT("")
+									, Op.bBSDFHasSecondRoughnessOrSimpleClearCoat ? TEXT("2Ro ") : TEXT("")
+									, Op.bBSDFHasFuzz ? TEXT("Fuz ") : TEXT("")
+								);
+
+								static FString ToolTip;
+								if (ToolTip.IsEmpty())
+								{
+									ToolTip += TEXT("SSS means the BSDF features subsurface profile or subsurface setup using MFP.\n");
+									ToolTip += TEXT("MFP means the BSDF MFP is specified by the user.\n");
+									ToolTip += TEXT("F90 means the BSDF edge specular color representing reflectivity at grazing angle is used.\n");
+									ToolTip += TEXT("Fuz means the BSDF fuzz layer is enabled.\n");
+									ToolTip += TEXT("2Ro means the BSDF either uses a second specular lob with a second roughness, or the legacy simple clear coat.\n");
+									ToolTip += TEXT("Ani means the BSDF anisotropic specular lighting is used.");
+								}
+
 								auto BSDF = SNew(SErrorText)
-									.ErrorText(LOCTEXT("BSDF", "BSDF"))
+									.ErrorText(FText::FromString(BSDFDesc))
 									.BackgroundColor(FSlateColor(EStyleColor::AccentGreen))
-									.ToolTipText(LOCTEXT("BSDFtooltip", "BSDF tool tip"));
-								return BSDF->AsShared();
-							}
-							break;
-							case STRATA_OPERATOR_BSDF_LEGACY:
-							{
-								auto BSDF = SNew(SErrorText)
-									.ErrorText(LOCTEXT("LegacyBSDF", "LegacyBSDF"))
-									.BackgroundColor(FSlateColor(EStyleColor::AccentBlue))
-									.ToolTipText(LOCTEXT("LegacyBSDFtooltip", "LegacyBSDF tool tip"));
+									.ToolTipText(FText::FromString(ToolTip));
 								return BSDF->AsShared();
 							}
 							break;
diff --git a/Engine/Source/Runtime/Engine/Public/Rendering/StrataMaterialShared.h b/Engine/Source/Runtime/Engine/Public/Rendering/StrataMaterialShared.h
--- a/Engine/Source/Runtime/Engine/Public/Rendering/StrataMaterialShared.h
+++ b/Engine/Source/Runtime/Engine/Public/Rendering/StrataMaterialShared.h
@@ -30,8 +30,13 @@ struct FStrataOperator
 public:
 	FStrataOperator();
 
+	// !!!!!!!!!!
+	// Not using LAYOUT_BITFIELD_EDITORONLY because it seems to cause issue with bit being shifted around when copy happens.
+	// So in the meantime we find it out, LAYOUT_FIELD_EDITORONLY using uint8 is used.
+	// !!!!!!!!!!
+
 	LAYOUT_FIELD_EDITORONLY(int32, OperatorType);
-	LAYOUT_BITFIELD_EDITORONLY(uint8, bNodeRequestParameterBlending, 1);
+	LAYOUT_FIELD_EDITORONLY(uint8, bNodeRequestParameterBlending);
 
 	LAYOUT_FIELD_EDITORONLY(int32, Index);			// Index into the array of operators
 	LAYOUT_FIELD_EDITORONLY(int32, ParentIndex);	// Parent operator index
@@ -43,21 +48,21 @@ public:
 	LAYOUT_FIELD_EDITORONLY(int32, BSDFIndex);		// Index in the array of BSDF if a BSDF operator
 	LAYOUT_FIELD_EDITORONLY(int32, BSDFType);
 	LAYOUT_FIELD_EDITORONLY(FStrataRegisteredSharedLocalBasis, BSDFRegisteredSharedLocalBasis);
-	LAYOUT_BITFIELD_EDITORONLY(uint8, bBSDFHasSSS, 1);
-	LAYOUT_BITFIELD_EDITORONLY(uint8, bBSDFHasMFPPluggedIn, 1);
-	LAYOUT_BITFIELD_EDITORONLY(uint8, bBSDFHasEdgeColor, 1);
-	LAYOUT_BITFIELD_EDITORONLY(uint8, bBSDFHasFuzz, 1);
-	LAYOUT_BITFIELD_EDITORONLY(uint8, bBSDFHasSecondRoughnessOrSimpleClearCoat, 1);
-	LAYOUT_BITFIELD_EDITORONLY(uint8, bBSDFHasAnisotropy, 1);
+	LAYOUT_FIELD_EDITORONLY(uint8, bBSDFHasSSS);
+	LAYOUT_FIELD_EDITORONLY(uint8, bBSDFHasMFPPluggedIn);
+	LAYOUT_FIELD_EDITORONLY(uint8, bBSDFHasEdgeColor);
+	LAYOUT_FIELD_EDITORONLY(uint8, bBSDFHasFuzz);
+	LAYOUT_FIELD_EDITORONLY(uint8, bBSDFHasSecondRoughnessOrSimpleClearCoat);
+	LAYOUT_FIELD_EDITORONLY(uint8, bBSDFHasAnisotropy);
 
 	// Data derived after the tree has been built.
 	LAYOUT_FIELD_EDITORONLY(int32, MaxDistanceFromLeaves);
 	LAYOUT_FIELD_EDITORONLY(int32, LayerDepth);
-	LAYOUT_BITFIELD_EDITORONLY(uint8, bIsTop, 1);
-	LAYOUT_BITFIELD_EDITORONLY(uint8, bIsBottom, 1);
+	LAYOUT_FIELD_EDITORONLY(uint8, bIsTop);
+	LAYOUT_FIELD_EDITORONLY(uint8, bIsBottom);
 
-	LAYOUT_BITFIELD_EDITORONLY(uint8, bUseParameterBlending, 1);			// True when part of a sub tree where parameter blending is in use
-	LAYOUT_BITFIELD_EDITORONLY(uint8, bRootOfParameterBlendingSubTree, 1);	// True when the root of a sub tree where parameter blending is in use. Only this node will register a BSDF
+	LAYOUT_FIELD_EDITORONLY(uint8, bUseParameterBlending);			// True when part of a sub tree where parameter blending is in use
+	LAYOUT_FIELD_EDITORONLY(uint8, bRootOfParameterBlendingSubTree);	// True when the root of a sub tree where parameter blending is in use. Only this node will register a BSDF
 
 	void CombineFlagsForParameterBlending(FStrataOperator& A, FStrataOperator& B);
 
@@ -66,7 +71,7 @@ public:
 	bool IsDiscarded() const;
 };
 
-#define STRATA_COMPILATION_OUTPUT_MAX_OPERATOR 32
+#define STRATA_COMPILATION_OUTPUT_MAX_OPERATOR 24
 
 struct FStrataMaterialCompilationOutput
 {
@@ -80,13 +85,13 @@ public:
 	////
 
 	/** Strata material type, at compile time (0:simple, 1:single, 2: complex) */
-	LAYOUT_BITFIELD(uint8, StrataMaterialType, 2);
+	LAYOUT_FIELD(uint8, StrataMaterialType);
 
 	/** Strata BSDF count, at compile time (0-7) */
-	LAYOUT_BITFIELD(uint8, StrataBSDFCount, 3);
+	LAYOUT_FIELD(uint8, StrataBSDFCount);
 
 	/** Strata uint per pixel, at compile time (0-255) */
-	LAYOUT_BITFIELD(uint8, StrataUintPerPixel, 8);
+	LAYOUT_FIELD(uint8, StrataUintPerPixel);
 
 	////
 	//// The following data is only needed when compiling with the editor.
