Move ResetLoaders wwhen duplicating material to a more correct place: UpdateOriginalMaterial, called when duplicating materials on apply, save or else.
Follow up to 25171370.

diff --git a/Engine/Source/Editor/MaterialEditor/Private/MaterialEditor.cpp b/Engine/Source/Editor/MaterialEditor/Private/MaterialEditor.cpp
--- a/Engine/Source/Editor/MaterialEditor/Private/MaterialEditor.cpp
+++ b/Engine/Source/Editor/MaterialEditor/Private/MaterialEditor.cpp
@@ -487,10 +487,6 @@ void FMaterialEditor::InitEditorForMaterial(UMaterial* InMaterial)
 	OriginalMaterialObject = InMaterial;
 
 	ExpressionPreviewMaterial = NULL;
-
-	// Now, the material has been loaded and serialized. PostLoad has been called and the asset has been converted and now it up to date.
-	// Let's thus reset the linker (to the last version) to make sure PostLoad is not doing any data conversion squashing new properties with not properly initialized _DEPRECATED members.
-	ResetLoaders(OriginalMaterial->GetPackage());
 	
 	// Create a copy of the material for preview usage (duplicating to a different class than original!)
 	// Propagate all object flags except for RF_Standalone, otherwise the preview material won't GC once
@@ -2703,6 +2699,10 @@ bool FMaterialEditor::UpdateOriginalMaterial()
 		// A bit hacky, but disable material compilation in post load when we duplicate the material.
 		UMaterial::ForceNoCompilationInPostLoad(true);
 
+		// Now, the material has been loaded and serialized. PostLoad has been called and the asset has been converted and now it is up to date.
+		// Let's thus reset the linker (to the last version) to make sure PostLoad on the duplicated object is not doing any data conversion, squashing new properties with not properly initialized _DEPRECATED members (e.g. RefractionMode_DEPRECATED).
+		ResetLoaders(OriginalMaterial->GetPackage());
+
 		// overwrite the original material in place by constructing a new one with the same name
 		OriginalMaterial = (UMaterial*)StaticDuplicateObject( Material, OriginalMaterial->GetOuter(), OriginalMaterial->GetFName(), 
 			RF_AllFlags, 
