Disable ability to override parameters that cause Material Instances to create a new shader permutation when parenting a restricted (non-base game cooked) parent.

In restrictive mode, material instances could not inherit from a cooked material not in base modules. Now material instances can inherit from such restricted materials as long as they do not introduce a new static parameters permutation and thus a new shader. To prevent this, the UI disables the checkboxes that would trigger a new shader to be created (full validation is still performed at cook time, so even if a user manages to create a material instance with a static permutation, this would be rejected by the content worker).

diff --git a/Engine/Source/Editor/MaterialEditor/Private/MaterialEditorInstanceDetailCustomization.cpp b/Engine/Source/Editor/MaterialEditor/Private/MaterialEditorInstanceDetailCustomization.cpp
--- a/Engine/Source/Editor/MaterialEditor/Private/MaterialEditorInstanceDetailCustomization.cpp
+++ b/Engine/Source/Editor/MaterialEditor/Private/MaterialEditorInstanceDetailCustomization.cpp
@@ -963,12 +963,6 @@ bool FMaterialInstanceParameterDetails::OnShouldSetAsset(const FAssetData& Asset
 		}
 	}
 
-	if (!MaterialEditorInstance->SourceInstance->IsSpecificMaterialValidParent(MaterialInstance))
-	{
-		FMessageDialog::Open(EAppMsgType::Ok, NSLOCTEXT("UnrealEd", "Error_MaterialEditor_InvalidMaterialInstanceParent", "Only uncooked, user-defined, engine or game Materials/MaterialInstances can be selected as a MaterialInstance parent."));
-		return false;
-	}
-
 	return true;
 }
 
@@ -1512,8 +1506,24 @@ bool FMaterialInstanceParameterDetails::OverrideMaxWorldPositionOffsetDisplaceme
 	return MaterialEditorInstance->BasePropertyOverrides.bOverride_MaxWorldPositionOffsetDisplacement;
 }
 
+/** Helper function used by some parameters to verify that they are allowed to be overridden. This
+  * must be prevented if the source material instance disallows the creation of new static parameter
+  * permutations as that would trigger a new shader creation. */
+static bool DoesSourceMaterialInstanceDisallowStaticParameterPermutation(const UMaterialEditorInstanceConstant* mi, bool NewValue)
+{
+	if (NewValue && mi->SourceInstance->bDisallowStaticParameterPermutations)
+	{
+		return true;
+	}
+	return false;
+}
+
 void FMaterialInstanceParameterDetails::OnOverrideOpacityClipMaskValueChanged(bool NewValue)
 {
+	if (DoesSourceMaterialInstanceDisallowStaticParameterPermutation(MaterialEditorInstance, NewValue))
+	{
+		return;
+	}
 	MaterialEditorInstance->BasePropertyOverrides.bOverride_OpacityMaskClipValue = NewValue;
 	MaterialEditorInstance->PostEditChange();
 	FEditorSupportDelegates::RedrawAllViewports.Broadcast();
@@ -1521,6 +1531,10 @@ void FMaterialInstanceParameterDetails::OnOverrideOpacityClipMaskValueChanged(bo
 
 void FMaterialInstanceParameterDetails::OnOverrideBlendModeChanged(bool NewValue)
 {
+	if (DoesSourceMaterialInstanceDisallowStaticParameterPermutation(MaterialEditorInstance, NewValue))
+	{
+		return;
+	}
 	MaterialEditorInstance->BasePropertyOverrides.bOverride_BlendMode = NewValue;
 	MaterialEditorInstance->PostEditChange();
 	FEditorSupportDelegates::RedrawAllViewports.Broadcast();
@@ -1528,6 +1542,10 @@ void FMaterialInstanceParameterDetails::OnOverrideBlendModeChanged(bool NewValue
 
 void FMaterialInstanceParameterDetails::OnOverrideShadingModelChanged(bool NewValue)
 {
+	if (DoesSourceMaterialInstanceDisallowStaticParameterPermutation(MaterialEditorInstance, NewValue))
+	{
+		return;
+	}
 	MaterialEditorInstance->BasePropertyOverrides.bOverride_ShadingModel = NewValue;
 	MaterialEditorInstance->PostEditChange();
 	FEditorSupportDelegates::RedrawAllViewports.Broadcast();
@@ -1535,6 +1553,10 @@ void FMaterialInstanceParameterDetails::OnOverrideShadingModelChanged(bool NewVa
 
 void FMaterialInstanceParameterDetails::OnOverrideTwoSidedChanged(bool NewValue)
 {
+	if (DoesSourceMaterialInstanceDisallowStaticParameterPermutation(MaterialEditorInstance, NewValue))
+	{
+		return;
+	}
 	MaterialEditorInstance->BasePropertyOverrides.bOverride_TwoSided = NewValue;
 	MaterialEditorInstance->PostEditChange();
 	FEditorSupportDelegates::RedrawAllViewports.Broadcast();
@@ -1542,6 +1564,10 @@ void FMaterialInstanceParameterDetails::OnOverrideTwoSidedChanged(bool NewValue)
 
 void FMaterialInstanceParameterDetails::OnOverrideIsThinSurfaceChanged(bool NewValue)
 {
+	if (DoesSourceMaterialInstanceDisallowStaticParameterPermutation(MaterialEditorInstance, NewValue))
+	{
+		return;
+	}
 	MaterialEditorInstance->BasePropertyOverrides.bOverride_bIsThinSurface = NewValue;
 	MaterialEditorInstance->PostEditChange();
 	FEditorSupportDelegates::RedrawAllViewports.Broadcast();
@@ -1549,6 +1575,10 @@ void FMaterialInstanceParameterDetails::OnOverrideIsThinSurfaceChanged(bool NewV
 
 void FMaterialInstanceParameterDetails::OnOverrideDitheredLODTransitionChanged(bool NewValue)
 {
+	if (DoesSourceMaterialInstanceDisallowStaticParameterPermutation(MaterialEditorInstance, NewValue))
+	{
+		return;
+	}
 	MaterialEditorInstance->BasePropertyOverrides.bOverride_DitheredLODTransition = NewValue;
 	MaterialEditorInstance->PostEditChange();
 	FEditorSupportDelegates::RedrawAllViewports.Broadcast();
@@ -1556,6 +1586,10 @@ void FMaterialInstanceParameterDetails::OnOverrideDitheredLODTransitionChanged(b
 
 void FMaterialInstanceParameterDetails::OnOverrideOutputTranslucentVelocityChanged(bool NewValue)
 {
+	if (DoesSourceMaterialInstanceDisallowStaticParameterPermutation(MaterialEditorInstance, NewValue))
+	{
+		return;
+	}
 	MaterialEditorInstance->BasePropertyOverrides.bOverride_OutputTranslucentVelocity = NewValue;
 	MaterialEditorInstance->PostEditChange();
 	FEditorSupportDelegates::RedrawAllViewports.Broadcast();
@@ -1570,6 +1604,10 @@ void FMaterialInstanceParameterDetails::OnOverrideDisplacementScalingChanged(boo
 
 void FMaterialInstanceParameterDetails::OnOverrideMaxWorldPositionOffsetDisplacementChanged(bool NewValue)
 {
+	if (DoesSourceMaterialInstanceDisallowStaticParameterPermutation(MaterialEditorInstance, NewValue))
+	{
+		return;
+	}
 	MaterialEditorInstance->BasePropertyOverrides.bOverride_MaxWorldPositionOffsetDisplacement = NewValue;
 	MaterialEditorInstance->PostEditChange();
 	FEditorSupportDelegates::RedrawAllViewports.Broadcast();
diff --git a/Engine/Source/Editor/MaterialEditor/Private/MaterialPropertyHelpers.cpp b/Engine/Source/Editor/MaterialEditor/Private/MaterialPropertyHelpers.cpp
--- a/Engine/Source/Editor/MaterialEditor/Private/MaterialPropertyHelpers.cpp
+++ b/Engine/Source/Editor/MaterialEditor/Private/MaterialPropertyHelpers.cpp
@@ -544,6 +544,17 @@ ECheckBoxState FMaterialPropertyHelpers::IsOverriddenExpressionCheckbox(UDEditor
 
 void FMaterialPropertyHelpers::OnOverrideParameter(bool NewValue, class UDEditorParameterValue* Parameter, UMaterialEditorInstanceConstant* MaterialEditorInstance)
 {
+	// Whether this parameter is static (i.e. overriding and changing its value requires a new shader permutation)
+	bool bParameterIsStatic =
+		   Parameter->GetParameterType() == EMaterialParameterType::StaticSwitch
+		|| Parameter->GetParameterType() == EMaterialParameterType::StaticComponentMask;
+
+	// If the material instance disallows the creation of new shader permutations, prevent overriding the static parameter.
+	if (NewValue && bParameterIsStatic && MaterialEditorInstance->SourceInstance->bDisallowStaticParameterPermutations)
+	{
+		return;
+	}
+
 	const FScopedTransaction Transaction( LOCTEXT( "OverrideParameter", "Override Parameter" ) );
 	Parameter->Modify();
 	Parameter->bOverride = NewValue;
diff --git a/Engine/Source/Runtime/Engine/Classes/Materials/MaterialInstance.h b/Engine/Source/Runtime/Engine/Classes/Materials/MaterialInstance.h
--- a/Engine/Source/Runtime/Engine/Classes/Materials/MaterialInstance.h
+++ b/Engine/Source/Runtime/Engine/Classes/Materials/MaterialInstance.h
@@ -656,6 +656,9 @@ public:
 #if WITH_EDITOR
 	/** Flag to detect cycles in the material instance graph, this is only used at content creation time where the hierarchy can be changed. */
 	bool ReentrantFlag[2];
+
+	/** Whether static parameter permutations should be disabled (e.g. due to having a restricted parent material) */
+	bool bDisallowStaticParameterPermutations;
 #endif
 
 	/** 
@@ -855,14 +858,14 @@ public:
 	ENGINE_API void RemoveLayerParameterIndex(int32 Index);
 	
 	/**
-	 * Returns whether specified MaterialInterface is a valid parent for this MaterialInstance.
+	 * Returns this material instance is allowed to override static parameters and introduce a new shader permutation when it derives from `CandidateParent`.
 	 */
-	ENGINE_API bool IsSpecificMaterialValidParent(UMaterialInterface* CandidateParent) const;
+	ENGINE_API bool IsStaticPermutationAllowedForCandidateParent(UMaterialInterface* CandidateParent) const;
 
 	/**
 	 * Ensures that current parent is a valid for this material instance and if not, it resets the parent to null.
 	 */
-	ENGINE_API void ValidateParent();
+	ENGINE_API void ValidateStaticPermutationAllowed();
 
 #endif // WITH_EDITOR
 
diff --git a/Engine/Source/Runtime/Engine/Private/Materials/MaterialInstance.cpp b/Engine/Source/Runtime/Engine/Private/Materials/MaterialInstance.cpp
--- a/Engine/Source/Runtime/Engine/Private/Materials/MaterialInstance.cpp
+++ b/Engine/Source/Runtime/Engine/Private/Materials/MaterialInstance.cpp
@@ -2063,8 +2063,8 @@ void UMaterialInstance::InitStaticPermutation(EMaterialShaderPrecompileMode Prec
 #if WITH_EDITORONLY_DATA
 	if (!GetPackage()->HasAnyPackageFlags(PKG_FilterEditorOnly))
 	{
-		// Update bHasStaticPermutationResource in case the parent was not found
-		bHasStaticPermutationResource = (HasStaticParameters() || HasOverridenBaseProperties()) && Parent;
+		bHasStaticPermutationResource = Parent && (HasStaticParameters() || HasOverridenBaseProperties());
+		ValidateStaticPermutationAllowed();
 	}
 #endif // WITH_EDITORONLY_DATA
 
@@ -3029,7 +3029,7 @@ void UMaterialInstance::PostLoad()
 	}
 
 #if WITH_EDITOR
-	ValidateParent();
+	ValidateStaticPermutationAllowed();
 #endif
 
 	// Add references to the expression object if we do not have one already, and fix up any names that were changed.
@@ -3170,7 +3170,7 @@ void UMaterialInstance::PostLoad()
 
 #if WITH_EDITOR
 
-bool UMaterialInstance::IsSpecificMaterialValidParent(UMaterialInterface* CandidateParent) const
+bool UMaterialInstance::IsStaticPermutationAllowedForCandidateParent(UMaterialInterface* CandidateParent) const
 {
 	// Nothing to do if specified parent is null or if restrictive mode is disabled
 	if (!CandidateParent || !bEnableRestrictiveMaterialInstanceParents)
@@ -3193,7 +3193,7 @@ bool UMaterialInstance::IsSpecificMaterialValidParent(UMaterialInterface* Candid
 
 	// Cache this material package
 	UPackage* Package = GetPackage();
-	
+
 	if (Package->HasAnyPackageFlags(PKG_Cooked)								// If this material instance package is cooked
 		|| Package == GetTransientPackage()									// Or if this material instance is transient
 		|| !CandidateParent->GetPackage()->HasAnyPackageFlags(PKG_Cooked))	// Or if the candidate material is uncooked
@@ -3205,14 +3205,21 @@ bool UMaterialInstance::IsSpecificMaterialValidParent(UMaterialInterface* Candid
 	return false;
 }
 
-void UMaterialInstance::ValidateParent()
+void UMaterialInstance::ValidateStaticPermutationAllowed()
 {
 	const UMaterialInterface* PrevParent = Parent;
-	if (!IsSpecificMaterialValidParent(Parent))
+
+	// Update the flag that controls whether a static permutation is allowed for this material instance
+	bDisallowStaticParameterPermutations = !IsStaticPermutationAllowedForCandidateParent(Parent);
+	
+	// Check that that either this material instance has no permutation or that it is allowed
+	if (bHasStaticPermutationResource && bDisallowStaticParameterPermutations)
 	{
 		// We don't allow Material Instances to parent to cooked materials.
-		UE_LOG(LogMaterial, Warning, TEXT("Can't set material instance '%s' to cooked non-user non-base parent material '%s'. Setting parent to null."), *GetName(), *PrevParent->GetName());
+		UE_LOG(LogMaterial, Warning, TEXT("Material instance '%s' with cooked non-user non-base parent material '%s' is not allowed to create new shader permutations. Setting parent to null."), *GetName(), *PrevParent->GetName());
+
 		SetParentInternal(nullptr, true);
+		bDisallowStaticParameterPermutations = false;
 	}
 }
 
@@ -3940,6 +3947,8 @@ void UMaterialInstance::UpdateStaticPermutation(const FStaticParameterSet& NewPa
 			FMaterialUpdateContext LocalMaterialUpdateContext(FMaterialUpdateContext::EOptions::RecreateRenderStates);
 			LocalMaterialUpdateContext.AddMaterialInstance(this);
 		}
+
+		ValidateStaticPermutationAllowed();
 	}
 }
 
@@ -4023,7 +4032,7 @@ void UMaterialInstance::PostEditChangeProperty(FPropertyChangedEvent& PropertyCh
 
 	if (PropertyChangedEvent.MemberProperty != nullptr && PropertyChangedEvent.MemberProperty->GetFName() == GET_MEMBER_NAME_CHECKED(UMaterialInstance, Parent))
 	{
-		ValidateParent();
+		ValidateStaticPermutationAllowed();
 	}
 
 	// If BLEND_TranslucentColoredTransmittance is selected while Strata is not enabled, force BLEND_Translucent blend mode
diff --git a/Engine/Source/Runtime/Engine/Private/Materials/MaterialInstanceConstant.cpp b/Engine/Source/Runtime/Engine/Private/Materials/MaterialInstanceConstant.cpp
--- a/Engine/Source/Runtime/Engine/Private/Materials/MaterialInstanceConstant.cpp
+++ b/Engine/Source/Runtime/Engine/Private/Materials/MaterialInstanceConstant.cpp
@@ -91,7 +91,7 @@ void UMaterialInstanceConstant::SetParentEditorOnly(UMaterialInterface* NewParen
 
 	if (SetParentInternal(NewParent, RecacheShader))
 	{
-		ValidateParent();
+		ValidateStaticPermutationAllowed();
 		UpdateCachedData();
 	}
 }
