Substrate - Editor HUD - show IsThin and fixed a few UI presentation/alignement issues.

diff --git a/Engine/Source/Editor/MaterialEditor/Private/SMaterialEditorStrataWidget.cpp b/Engine/Source/Editor/MaterialEditor/Private/SMaterialEditorStrataWidget.cpp
--- a/Engine/Source/Editor/MaterialEditor/Private/SMaterialEditorStrataWidget.cpp
+++ b/Engine/Source/Editor/MaterialEditor/Private/SMaterialEditorStrataWidget.cpp
@@ -60,20 +60,21 @@ void SMaterialEditorStrataWidget::Construct(const FArguments& InArgs, TWeakPtr<F
 						SNew(SHorizontalBox)
 						+SHorizontalBox::Slot()
 						.AutoWidth()
+						.HAlign(HAlign_Left)
 						.VAlign(VAlign_Center)
 						[
 							SNew(SWrapBox)
 							.UseAllottedSize(true)
 							+SWrapBox::Slot()
 							.Padding(5.0f)
-							.HAlign(HAlign_Center)
+							.HAlign(HAlign_Left)
 							.VAlign(VAlign_Center)
 							[
 								CheckBoxForceFullSimplification->AsShared()
 							]
 						]
 						+SHorizontalBox::Slot()
-						.FillWidth(1.0f)
+						.AutoWidth()
 						.Padding(16.0f, 0.0f)
 						.HAlign(HAlign_Left)
 						.VAlign(VAlign_Center)
@@ -85,16 +86,16 @@ void SMaterialEditorStrataWidget::Construct(const FArguments& InArgs, TWeakPtr<F
 							.Text(LOCTEXT("FullsimplificationLabel", "Full simplification"))
 						]
 						+SHorizontalBox::Slot()
-						.FillWidth(1.0f)
+						.AutoWidth()
 						.Padding(16.0f, 0.0f)
 						.HAlign(HAlign_Left)
 						.VAlign(VAlign_Center)
 						[
 							SNew(SWrapBox)
-							.UseAllottedSize(true)
+							//.UseAllottedSize(true)
 							+ SWrapBox::Slot()
 							.Padding(5.0f)
-							.HAlign(HAlign_Center)
+							.HAlign(HAlign_Left)
 							.VAlign(VAlign_Center)
 							[
 								ButtonApplyToPreview->AsShared()
@@ -227,6 +228,8 @@ void SMaterialEditorStrataWidget::Tick(const FGeometry& AllottedGeometry, const
 					MaterialDescription += FString::Printf(TEXT("Material complexity          = UNKOWN => ERROR!\r\n"));
 				}
 
+				MaterialDescription += FString::Printf(TEXT("Is Thin                      = %i\r\n"), CompilationOutput.bIsThin);
+
 				MaterialDescription += FString::Printf(TEXT(" \r\n"));
 				MaterialDescription += FString::Printf(TEXT(" \r\n"));
 				MaterialDescription += FString::Printf(TEXT("================================================================================\r\n"));
@@ -332,7 +335,7 @@ void SMaterialEditorStrataWidget::Tick(const FGeometry& AllottedGeometry, const
 									ToolTip += TEXT("F90 means the BSDF edge specular color representing reflectivity at grazing angle is used.\n");
 									ToolTip += TEXT("Fuz means the BSDF fuzz layer is enabled.\n");
 									ToolTip += TEXT("2Ro means the BSDF either uses a second specular lob with a second roughness, or the legacy simple clear coat.\n");
-									ToolTip += TEXT("Ani means the BSDF anisotropic specular lighting is used.");
+									ToolTip += TEXT("Ani means the BSDF anisotropic specular lighting is used.\n");
 									ToolTip += TEXT("Gli means the BSDF features glints.");
 								}
 
diff --git a/Engine/Source/Runtime/Engine/Private/Materials/HLSLMaterialTranslator.cpp b/Engine/Source/Runtime/Engine/Private/Materials/HLSLMaterialTranslator.cpp
--- a/Engine/Source/Runtime/Engine/Private/Materials/HLSLMaterialTranslator.cpp
+++ b/Engine/Source/Runtime/Engine/Private/Materials/HLSLMaterialTranslator.cpp
@@ -11790,6 +11790,7 @@ bool FHLSLMaterialTranslator::FStrataCompilationContext::StrataGenerateDerivedMa
 				Compiler->MaterialCompilationOutput.StrataMaterialCompilationOutput.StrataMaterialType = bStrataMaterialIsSimple ? 0 : bStrataMaterialIsSingle ? 1 : 2;
 				Compiler->MaterialCompilationOutput.StrataMaterialCompilationOutput.StrataBSDFCount = StrataMaterialBSDFCount;
 				Compiler->MaterialCompilationOutput.StrataMaterialCompilationOutput.StrataUintPerPixel = uint8(FMath::Clamp(RequestedSizeInUint, 0u, 0xFF));
+				Compiler->MaterialCompilationOutput.StrataMaterialCompilationOutput.bIsThin = CompilerMaterial->IsThinSurface() ? 1 : 0;
 
 #if WITH_EDITOR
 				Compiler->MaterialCompilationOutput.StrataMaterialCompilationOutput.SharedLocalBasesCount = 0; // FinalUsedSharedLocalBasesCount is not valid yet
diff --git a/Engine/Source/Runtime/Engine/Public/Rendering/StrataMaterialShared.h b/Engine/Source/Runtime/Engine/Public/Rendering/StrataMaterialShared.h
--- a/Engine/Source/Runtime/Engine/Public/Rendering/StrataMaterialShared.h
+++ b/Engine/Source/Runtime/Engine/Public/Rendering/StrataMaterialShared.h
@@ -94,6 +94,9 @@ public:
 	/** Strata uint per pixel, at compile time (0-255) */
 	LAYOUT_FIELD(uint8, StrataUintPerPixel);
 
+	/** Indicate that the material is considered a thin surface instead of a volume filled up with matter */
+	LAYOUT_FIELD(uint8, bIsThin);
+
 	////
 	//// The following data is only needed when compiling with the editor.
 	////
