Following up on 25509187, which disables the ability to override material parameters leading to a new shader permutation if the material is derived from a restricted parent, this change enhances the tooltip of the affected parameters to provide users with informative messages.

diff --git a/Engine/Source/Editor/MaterialEditor/Private/MaterialEditorInstanceDetailCustomization.cpp b/Engine/Source/Editor/MaterialEditor/Private/MaterialEditorInstanceDetailCustomization.cpp
--- a/Engine/Source/Editor/MaterialEditor/Private/MaterialEditorInstanceDetailCustomization.cpp
+++ b/Engine/Source/Editor/MaterialEditor/Private/MaterialEditorInstanceDetailCustomization.cpp
@@ -1247,6 +1247,9 @@ void FMaterialInstanceParameterDetails::CreateBasePropertyOverrideWidgets(IDetai
 	TSharedPtr<IPropertyHandle> DisplacementScalingProperty = BasePropertyOverridePropery->GetChildHandle("DisplacementScaling");
 	TSharedPtr<IPropertyHandle> MaxWorldPositionOffsetDisplacementProperty = BasePropertyOverridePropery->GetChildHandle("MaxWorldPositionOffsetDisplacement");
 
+	const FText ParameterDisabledToolTipString = FText::FromString(TEXT("This material instance parent restricts the creation of new shader permutations. Overriding this parameter would result in the generation of additional shader permutations."));
+	const bool bStaticParametersOverrideDisabled = MaterialEditorInstance->SourceInstance->bDisallowStaticParameterPermutations;
+
 	// Update blend mode display names
 	if (FByteProperty* BlendModeByteProperty = (FByteProperty*)BlendModeProperty->GetProperty())
 	{
@@ -1267,7 +1270,7 @@ void FMaterialInstanceParameterDetails::CreateBasePropertyOverrideWidgets(IDetai
 		IDetailPropertyRow& OpacityClipMaskValuePropertyRow = BasePropertyOverrideGroup.AddPropertyRow(OpacityClipMaskValueProperty.ToSharedRef());
 		OpacityClipMaskValuePropertyRow
 			.DisplayName(OpacityClipMaskValueProperty->GetPropertyDisplayName())
-			.ToolTip(OpacityClipMaskValueProperty->GetToolTipText())
+			.ToolTip(bStaticParametersOverrideDisabled ? ParameterDisabledToolTipString : OpacityClipMaskValueProperty->GetToolTipText())
 			.EditCondition(IsOverrideOpacityClipMaskValueEnabled, FOnBooleanValueChanged::CreateSP(this, &FMaterialInstanceParameterDetails::OnOverrideOpacityClipMaskValueChanged))
 			.Visibility(TAttribute<EVisibility>::Create(TAttribute<EVisibility>::FGetter::CreateSP(this, &FMaterialInstanceParameterDetails::IsOverriddenAndVisible, IsOverrideOpacityClipMaskValueEnabled)))
 			.OverrideResetToDefault(ResetOpacityClipMaskValuePropertyOverride);
@@ -1286,7 +1289,7 @@ void FMaterialInstanceParameterDetails::CreateBasePropertyOverrideWidgets(IDetai
 		IDetailPropertyRow& BlendModePropertyRow = BasePropertyOverrideGroup.AddPropertyRow(BlendModeProperty.ToSharedRef());
 		BlendModePropertyRow
 			.DisplayName(BlendModeProperty->GetPropertyDisplayName())
-			.ToolTip(BlendModeProperty->GetToolTipText())
+			.ToolTip(bStaticParametersOverrideDisabled ? ParameterDisabledToolTipString : BlendModeProperty->GetToolTipText())
 			.EditCondition(IsOverrideBlendModeEnabled, FOnBooleanValueChanged::CreateSP(this, &FMaterialInstanceParameterDetails::OnOverrideBlendModeChanged))
 			.Visibility(TAttribute<EVisibility>::Create(TAttribute<EVisibility>::FGetter::CreateSP(this, &FMaterialInstanceParameterDetails::IsOverriddenAndVisible, IsOverrideBlendModeEnabled)))
 			.OverrideResetToDefault(ResetBlendModePropertyOverride);
@@ -1326,7 +1329,7 @@ void FMaterialInstanceParameterDetails::CreateBasePropertyOverrideWidgets(IDetai
 		IDetailPropertyRow& ShadingModelPropertyRow = BasePropertyOverrideGroup.AddPropertyRow(ShadingModelProperty.ToSharedRef());
 		ShadingModelPropertyRow
 			.DisplayName(ShadingModelProperty->GetPropertyDisplayName())
-			.ToolTip(ShadingModelProperty->GetToolTipText())
+			.ToolTip(bStaticParametersOverrideDisabled ? ParameterDisabledToolTipString : ShadingModelProperty->GetToolTipText())
 			.EditCondition(IsOverrideShadingModelEnabled, FOnBooleanValueChanged::CreateSP(this, &FMaterialInstanceParameterDetails::OnOverrideShadingModelChanged))
 			.Visibility(TAttribute<EVisibility>::Create(TAttribute<EVisibility>::FGetter::CreateSP(this, &FMaterialInstanceParameterDetails::IsOverriddenAndVisible, IsOverrideShadingModelEnabled)))
 			.OverrideResetToDefault(ResetShadingModelPropertyOverride);
@@ -1345,7 +1348,7 @@ void FMaterialInstanceParameterDetails::CreateBasePropertyOverrideWidgets(IDetai
 		IDetailPropertyRow& TwoSidedPropertyRow = BasePropertyOverrideGroup.AddPropertyRow(TwoSidedProperty.ToSharedRef());
 		TwoSidedPropertyRow
 			.DisplayName(TwoSidedProperty->GetPropertyDisplayName())
-			.ToolTip(TwoSidedProperty->GetToolTipText())
+			.ToolTip(bStaticParametersOverrideDisabled ? ParameterDisabledToolTipString : TwoSidedProperty->GetToolTipText())
 			.EditCondition(IsOverrideTwoSidedEnabled, FOnBooleanValueChanged::CreateSP(this, &FMaterialInstanceParameterDetails::OnOverrideTwoSidedChanged))
 			.Visibility(TAttribute<EVisibility>::Create(TAttribute<EVisibility>::FGetter::CreateSP(this, &FMaterialInstanceParameterDetails::IsOverriddenAndVisible, IsOverrideTwoSidedEnabled)))
 			.OverrideResetToDefault(ResetTwoSidedPropertyOverride);
@@ -1365,7 +1368,7 @@ void FMaterialInstanceParameterDetails::CreateBasePropertyOverrideWidgets(IDetai
 		IDetailPropertyRow& IsThinSurfacePropertyRow = BasePropertyOverrideGroup.AddPropertyRow(IsThinSurfaceProperty.ToSharedRef());
 		IsThinSurfacePropertyRow
 			.DisplayName(IsThinSurfaceProperty->GetPropertyDisplayName())
-			.ToolTip(IsThinSurfaceProperty->GetToolTipText())
+			.ToolTip(bStaticParametersOverrideDisabled ? ParameterDisabledToolTipString : IsThinSurfaceProperty->GetToolTipText())
 			.EditCondition(IsOverrideIsThinSurfaceEnabled, FOnBooleanValueChanged::CreateSP(this, &FMaterialInstanceParameterDetails::OnOverrideIsThinSurfaceChanged))
 			.Visibility(TAttribute<EVisibility>::Create(TAttribute<EVisibility>::FGetter::CreateSP(this, &FMaterialInstanceParameterDetails::IsOverriddenAndVisible, IsOverrideIsThinSurfaceEnabled)))
 			.OverrideResetToDefault(ResetIsThinSurfacePropertyOverride);
@@ -1384,7 +1387,7 @@ void FMaterialInstanceParameterDetails::CreateBasePropertyOverrideWidgets(IDetai
 		IDetailPropertyRow& DitheredLODTransitionPropertyRow = BasePropertyOverrideGroup.AddPropertyRow(DitheredLODTransitionProperty.ToSharedRef());
 		DitheredLODTransitionPropertyRow
 			.DisplayName(DitheredLODTransitionProperty->GetPropertyDisplayName())
-			.ToolTip(DitheredLODTransitionProperty->GetToolTipText())
+			.ToolTip(bStaticParametersOverrideDisabled ? ParameterDisabledToolTipString : DitheredLODTransitionProperty->GetToolTipText())
 			.EditCondition(IsOverrideDitheredLODTransitionEnabled, FOnBooleanValueChanged::CreateSP(this, &FMaterialInstanceParameterDetails::OnOverrideDitheredLODTransitionChanged))
 			.Visibility(TAttribute<EVisibility>::Create(TAttribute<EVisibility>::FGetter::CreateSP(this, &FMaterialInstanceParameterDetails::IsOverriddenAndVisible, IsOverrideDitheredLODTransitionEnabled)))
 			.OverrideResetToDefault(ResetDitheredLODTransitionPropertyOverride);
@@ -1403,7 +1406,7 @@ void FMaterialInstanceParameterDetails::CreateBasePropertyOverrideWidgets(IDetai
 		IDetailPropertyRow& OutputTranslucentVelocityPropertyRow = BasePropertyOverrideGroup.AddPropertyRow(OutputTranslucentVelocityProperty.ToSharedRef());
 		OutputTranslucentVelocityPropertyRow
 			.DisplayName(OutputTranslucentVelocityProperty->GetPropertyDisplayName())
-			.ToolTip(OutputTranslucentVelocityProperty->GetToolTipText())
+			.ToolTip(bStaticParametersOverrideDisabled ? ParameterDisabledToolTipString : OutputTranslucentVelocityProperty->GetToolTipText())
 			.EditCondition(IsOverrideOutputTranslucentVelocityEnabled, FOnBooleanValueChanged::CreateSP(this, &FMaterialInstanceParameterDetails::OnOverrideOutputTranslucentVelocityChanged))
 			.Visibility(TAttribute<EVisibility>::Create(TAttribute<EVisibility>::FGetter::CreateSP(this, &FMaterialInstanceParameterDetails::IsOverriddenAndVisible, IsOverrideOutputTranslucentVelocityEnabled)))
 			.OverrideResetToDefault(ResetOutputTranslucentVelocityPropertyOverride);
@@ -1444,7 +1447,7 @@ void FMaterialInstanceParameterDetails::CreateBasePropertyOverrideWidgets(IDetai
 		IDetailPropertyRow& MaxWorldPositionOffsetDisplacementPropertyRow = BasePropertyOverrideGroup.AddPropertyRow(MaxWorldPositionOffsetDisplacementProperty.ToSharedRef());
 		MaxWorldPositionOffsetDisplacementPropertyRow
 			.DisplayName(MaxWorldPositionOffsetDisplacementProperty->GetPropertyDisplayName())
-			.ToolTip(MaxWorldPositionOffsetDisplacementProperty->GetToolTipText())
+			.ToolTip(bStaticParametersOverrideDisabled ? ParameterDisabledToolTipString : MaxWorldPositionOffsetDisplacementProperty->GetToolTipText())
 			.EditCondition(IsOverrideMaxWorldPositionOffsetDisplacementEnabled, FOnBooleanValueChanged::CreateSP(this, &FMaterialInstanceParameterDetails::OnOverrideMaxWorldPositionOffsetDisplacementChanged))
 			.Visibility(TAttribute<EVisibility>::Create(TAttribute<EVisibility>::FGetter::CreateSP(this, &FMaterialInstanceParameterDetails::IsOverriddenAndVisible, IsOverrideMaxWorldPositionOffsetDisplacementEnabled)))
 			.OverrideResetToDefault(ResetMaxWorldPositionOffsetDisplacementPropertyOverride);
diff --git a/Engine/Source/Editor/MaterialEditor/Private/MaterialPropertyHelpers.cpp b/Engine/Source/Editor/MaterialEditor/Private/MaterialPropertyHelpers.cpp
--- a/Engine/Source/Editor/MaterialEditor/Private/MaterialPropertyHelpers.cpp
+++ b/Engine/Source/Editor/MaterialEditor/Private/MaterialPropertyHelpers.cpp
@@ -544,13 +544,8 @@ ECheckBoxState FMaterialPropertyHelpers::IsOverriddenExpressionCheckbox(UDEditor
 
 void FMaterialPropertyHelpers::OnOverrideParameter(bool NewValue, class UDEditorParameterValue* Parameter, UMaterialEditorInstanceConstant* MaterialEditorInstance)
 {
-	// Whether this parameter is static (i.e. overriding and changing its value requires a new shader permutation)
-	bool bParameterIsStatic =
-		   Parameter->GetParameterType() == EMaterialParameterType::StaticSwitch
-		|| Parameter->GetParameterType() == EMaterialParameterType::StaticComponentMask;
-
 	// If the material instance disallows the creation of new shader permutations, prevent overriding the static parameter.
-	if (NewValue && bParameterIsStatic && MaterialEditorInstance->SourceInstance->bDisallowStaticParameterPermutations)
+	if (NewValue && Parameter->IsStaticParameter() && MaterialEditorInstance->SourceInstance->bDisallowStaticParameterPermutations)
 	{
 		return;
 	}
@@ -575,7 +570,13 @@ FText FMaterialPropertyHelpers::GetParameterTooltip(UDEditorParameterValue* Para
 {
 	const FText AssetPath = FText::FromString(Parameter->AssetPath);
 	FText TooltipText;
-	if (!Parameter->Description.IsEmpty())
+
+	// If the material instance disallows the creation of new shader permutations, prevent overriding the static parameter.
+	if (Parameter->IsStaticParameter() && static_cast<UMaterialEditorInstanceConstant*>(MaterialEditorInstance)->SourceInstance->bDisallowStaticParameterPermutations)
+	{
+		return FText::FromString(TEXT("This material instance parent restricts the creation of new shader permutations. Overriding this parameter would result in the generation of additional shader permutations."));
+	}
+	else if (!Parameter->Description.IsEmpty())
 	{
 		TooltipText = FText::Format(LOCTEXT("ParameterInfoDescAndLocation", "{0} \nFound in: {1}"), FText::FromString(Parameter->Description), AssetPath);
 	}
diff --git a/Engine/Source/Editor/UnrealEd/Classes/MaterialEditor/DEditorParameterValue.h b/Engine/Source/Editor/UnrealEd/Classes/MaterialEditor/DEditorParameterValue.h
--- a/Engine/Source/Editor/UnrealEd/Classes/MaterialEditor/DEditorParameterValue.h
+++ b/Engine/Source/Editor/UnrealEd/Classes/MaterialEditor/DEditorParameterValue.h
@@ -55,6 +55,12 @@ class UNREALED_API UDEditorParameterValue : public UObject
 		return false;
 	}
 
+	bool IsStaticParameter() const
+	{
+		EMaterialParameterType Type = GetParameterType();
+		return Type == EMaterialParameterType::StaticSwitch || Type == EMaterialParameterType::StaticComponentMask;
+	}
+
 	EMaterialParameterType GetParameterType() const
 	{
 		FMaterialParameterMetadata Result;
