Substrate - Fixed some editor poping message due to bad bitfield macro.

diff --git a/Engine/Source/Editor/MaterialEditor/Private/MaterialEditor.cpp b/Engine/Source/Editor/MaterialEditor/Private/MaterialEditor.cpp
--- a/Engine/Source/Editor/MaterialEditor/Private/MaterialEditor.cpp
+++ b/Engine/Source/Editor/MaterialEditor/Private/MaterialEditor.cpp
@@ -7063,20 +7063,23 @@ FGraphAppearanceInfo FMaterialEditor::GetGraphAppearance() const
 		AppearanceInfo.CornerText = LOCTEXT("AppearanceCornerText_Material", "MATERIAL"); 
 	}
 
-	UMaterial* MaterialForStats = this->bStatsFromPreviewMaterial ? this->Material : this->OriginalMaterial;
-	const FMaterialResource* MaterialResource = MaterialForStats->GetMaterialResource(GMaxRHIFeatureLevel);
-	if (MaterialResource)
+	if (Strata::IsStrataEnabled())
 	{
-		FString MaterialDescription;
-	
-		FMaterialShaderMap* ShaderMap = MaterialResource->GetGameThreadShaderMap();
-		if (ShaderMap)
+		UMaterial* MaterialForStats = this->bStatsFromPreviewMaterial ? this->Material : this->OriginalMaterial;
+		const FMaterialResource* MaterialResource = MaterialForStats->GetMaterialResource(GMaxRHIFeatureLevel);
+		if (MaterialResource)
 		{
-			const FStrataMaterialCompilationOutput& CompilationOutput = ShaderMap->GetStrataMaterialCompilationOutput();
-			if (CompilationOutput.bMaterialOutOfBudgetHasBeenSimplified)
+			FString MaterialDescription;
+
+			FMaterialShaderMap* ShaderMap = MaterialResource->GetGameThreadShaderMap();
+			if (ShaderMap)
 			{
-				AppearanceInfo.WarningText = LOCTEXT("AppearanceWarningText_Material", "Substrate material was out of budget and has been simplified.");
+				const FStrataMaterialCompilationOutput& CompilationOutput = ShaderMap->GetStrataMaterialCompilationOutput();
+				if (CompilationOutput.bMaterialOutOfBudgetHasBeenSimplified)
+				{
+					AppearanceInfo.WarningText = LOCTEXT("AppearanceWarningText_Material", "Substrate material was out of budget and has been simplified.");
 
+				}
 			}
 		}
 	}
diff --git a/Engine/Source/Runtime/Engine/Public/Rendering/StrataMaterialShared.h b/Engine/Source/Runtime/Engine/Public/Rendering/StrataMaterialShared.h
--- a/Engine/Source/Runtime/Engine/Public/Rendering/StrataMaterialShared.h
+++ b/Engine/Source/Runtime/Engine/Public/Rendering/StrataMaterialShared.h
@@ -115,6 +115,9 @@ public:
 	//// The following data is only needed when compiling with the editor.
 	////
 
+	// Note we use LAYOUT_FIELD_EDITORONLY for bools because LAYOUT_BITFIELD_EDITORONLY was causing issues when serialising the structure.
+	// STRATA_TODO pack that data.
+
 	/** The Strata verbose description */
 	LAYOUT_FIELD_EDITORONLY(FMemoryImageString, StrataMaterialDescription);
 
@@ -131,7 +134,7 @@ public:
 	LAYOUT_FIELD_EDITORONLY(uint8, MaterialType);
 
 	/** The byte per pixel count supported by the platform the material has been compiled against */
-	LAYOUT_BITFIELD_EDITORONLY(uint8, bMaterialOutOfBudgetHasBeenSimplified, 1);
+	LAYOUT_FIELD_EDITORONLY(uint8, bMaterialOutOfBudgetHasBeenSimplified, 1);
 
 	LAYOUT_FIELD_EDITORONLY(uint8, RootOperatorIndex);
 	LAYOUT_ARRAY_EDITORONLY(FStrataOperator, Operators, STRATA_COMPILATION_OUTPUT_MAX_OPERATOR);
