Substrate - progressive material operator simplification to fit in memory.

diff --git a/Engine/Source/Editor/MaterialEditor/Private/SMaterialEditorStrataWidget.cpp b/Engine/Source/Editor/MaterialEditor/Private/SMaterialEditorStrataWidget.cpp
--- a/Engine/Source/Editor/MaterialEditor/Private/SMaterialEditorStrataWidget.cpp
+++ b/Engine/Source/Editor/MaterialEditor/Private/SMaterialEditorStrataWidget.cpp
@@ -31,7 +31,7 @@ void SMaterialEditorStrataWidget::Construct(const FArguments& InArgs, TWeakPtr<F
 
 	CheckBoxForceFullSimplification = SNew(SCheckBox)
 		.Padding(5.0f)
-		.ToolTipText(LOCTEXT("CheckBoxForceFullSimplificationToolTip", "This will force full simplification of the material."));	// Just a test, needs to be more explicit
+		.ToolTipText(LOCTEXT("CheckBoxForceFullSimplificationToolTip", "This will force full simplification of the material."));
 
 	DescriptionTextBlock = SNew(STextBlock)
 		.TextStyle(FAppStyle::Get(), "Log.Normal")
@@ -40,6 +40,25 @@ void SMaterialEditorStrataWidget::Construct(const FArguments& InArgs, TWeakPtr<F
 		.ShadowOffset(FVector2D::UnitVector)
 		.Text(LOCTEXT("DescriptionTextBlock_Default", "Shader is compiling"));
 
+	BytesPerPixelOverride = Strata::GetBytePerPixel(SP_PCD3D_SM5);
+	BytesPerPixelOverrideInput = SNew(SNumericEntryBox<uint32>)
+		.MinValue(12)
+		.MaxValue(BytesPerPixelOverride)
+		.MinSliderValue(12)
+		.MaxSliderValue(BytesPerPixelOverride)
+		.OnBeginSliderMovement(this, &SMaterialEditorStrataWidget::OnBeginBytesPerPixelSliderMovement)
+		.OnEndSliderMovement(this, &SMaterialEditorStrataWidget::OnEndBytesPerPixelSliderMovement)
+		.AllowSpin(true)
+//		.OnValueChanged(this, &SMaterialEditorStrataWidget::OnBytesPerPixelChanged)
+		.OnValueCommitted(this, &SMaterialEditorStrataWidget::OnBytesPerPixelCommitted)
+		.Value(this, &SMaterialEditorStrataWidget::GetBytesPerPixelValue)
+		.IsEnabled(false);
+
+	CheckBoxBytesPerPixelOverride = SNew(SCheckBox)
+		.Padding(5.0f)
+		.ToolTipText(LOCTEXT("CheckBoxBytesPerPixelOverride", "This will force the byte per pixel count for the preview material. It cannot go higher than the current project setting."))
+		.OnCheckStateChanged(this, &SMaterialEditorStrataWidget::OnCheckBoxBytesPerPixelChanged);
+
 	if (Strata::IsStrataEnabled())
 	{
 		this->ChildSlot
@@ -88,6 +107,60 @@ void SMaterialEditorStrataWidget::Construct(const FArguments& InArgs, TWeakPtr<F
 						]
 						+SHorizontalBox::Slot()
 						.AutoWidth()
+						.HAlign(HAlign_Left)
+						.VAlign(VAlign_Center)
+						[
+							SNew(SWrapBox)
+							.UseAllottedSize(true)
+							+SWrapBox::Slot()
+							.Padding(5.0f)
+							.HAlign(HAlign_Left)
+							.VAlign(VAlign_Center)
+							[
+								CheckBoxBytesPerPixelOverride->AsShared()
+							]
+						]
+						+SHorizontalBox::Slot()
+						.AutoWidth()
+						.Padding(16.0f, 0.0f)
+						.HAlign(HAlign_Left)
+						.VAlign(VAlign_Center)
+						[
+							SNew(STextBlock)
+							.ColorAndOpacity(FLinearColor::White)
+							.ShadowColorAndOpacity(FLinearColor::Black)
+							.ShadowOffset(FVector2D::UnitVector)
+							.Text(LOCTEXT("OverrideBytesPerPixel", "Override bytes per pixel"))
+						]
+						+SHorizontalBox::Slot()
+						.AutoWidth()
+						.HAlign(HAlign_Left)
+						.VAlign(VAlign_Center)
+						[
+							SNew(SWrapBox)
+							.UseAllottedSize(true)
+							+SWrapBox::Slot()
+							.Padding(5.0f)
+							.HAlign(HAlign_Left)
+							.VAlign(VAlign_Center)
+							[
+								BytesPerPixelOverrideInput->AsShared()
+							]
+						]
+						+SHorizontalBox::Slot()
+						.AutoWidth()
+						.Padding(16.0f, 0.0f)
+						.HAlign(HAlign_Left)
+						.VAlign(VAlign_Center)
+						[
+							SNew(STextBlock)
+							.ColorAndOpacity(FLinearColor::White)
+							.ShadowColorAndOpacity(FLinearColor::Black)
+							.ShadowOffset(FVector2D::UnitVector)
+							.Text(LOCTEXT("FullsimplificationLabel", "Bytes per pixel"))
+						]
+						+SHorizontalBox::Slot()
+						.AutoWidth()
 						.Padding(16.0f, 0.0f)
 						.HAlign(HAlign_Left)
 						.VAlign(VAlign_Center)
@@ -330,6 +403,46 @@ void SMaterialEditorStrataWidget::Tick(const FGeometry& AllottedGeometry, const
 	}
 }
 
+//void SMaterialEditorStrataWidget::OnBytesPerPixelChanged(uint32 NewValue)
+//{
+//}
+
+void SMaterialEditorStrataWidget::OnBytesPerPixelCommitted(uint32 NewValue, ETextCommit::Type InCommitType)
+{
+	if (InCommitType == ETextCommit::OnEnter)
+	{
+		BytesPerPixelOverride = NewValue;
+	}
+}
+
+void SMaterialEditorStrataWidget::OnBeginBytesPerPixelSliderMovement()
+{
+	if (bBytesPerPixelStartedTransaction == false)
+	{
+		bBytesPerPixelStartedTransaction = true;
+		GEditor->BeginTransaction(LOCTEXT("PastePoseTransation", "Paste Pose"));
+	}
+}
+void SMaterialEditorStrataWidget::OnEndBytesPerPixelSliderMovement(uint32 NewValue)
+{
+	if (bBytesPerPixelStartedTransaction)
+	{
+		GEditor->EndTransaction();
+		bBytesPerPixelStartedTransaction = false;
+		BytesPerPixelOverride = NewValue;
+	}
+}
+
+TOptional<uint32> SMaterialEditorStrataWidget::GetBytesPerPixelValue() const
+{
+	return BytesPerPixelOverride;
+}
+
+void SMaterialEditorStrataWidget::OnCheckBoxBytesPerPixelChanged(ECheckBoxState InCheckBoxState)
+{
+	BytesPerPixelOverrideInput->SetEnabled(InCheckBoxState == ECheckBoxState::Checked);
+}
+
 FReply SMaterialEditorStrataWidget::OnButtonApplyToPreview()
 {
 	if (MaterialEditorPtr.IsValid())
@@ -338,13 +451,12 @@ FReply SMaterialEditorStrataWidget::OnButtonApplyToPreview()
 
 		FStrataCompilationConfig StrataCompilationConfig;
 		StrataCompilationConfig.bFullSimplify = CheckBoxForceFullSimplification->IsChecked();
+		StrataCompilationConfig.BytesPerPixelOverride = CheckBoxBytesPerPixelOverride->IsChecked() ? FMath::Clamp(BytesPerPixelOverride, 12, Strata::GetBytePerPixel(SP_PCD3D_SM5)) : -1;
 		MaterialInterface->SetStrataCompilationConfig(StrataCompilationConfig);
 
 		MaterialInterface->ForceRecompileForRendering();
 	}
 
-
-
 	return FReply::Handled();
 }
 
diff --git a/Engine/Source/Editor/MaterialEditor/Private/SMaterialEditorStrataWidget.h b/Engine/Source/Editor/MaterialEditor/Private/SMaterialEditorStrataWidget.h
--- a/Engine/Source/Editor/MaterialEditor/Private/SMaterialEditorStrataWidget.h
+++ b/Engine/Source/Editor/MaterialEditor/Private/SMaterialEditorStrataWidget.h
@@ -10,6 +10,7 @@
 #include "Widgets/SBoxPanel.h"
 #include "Widgets/Input/SButton.h"
 #include "Widgets/Input/SCheckBox.h"
+#include "Widgets/Input/SNumericEntryBox.h"
 #include "Widgets/Text/STextBlock.h"
 #include "Widgets/Layout/SBox.h"
 
@@ -38,6 +39,9 @@ public:
 private:
 
 	TSharedPtr<class SCheckBox> CheckBoxForceFullSimplification;
+	TSharedPtr<class SCheckBox> CheckBoxBytesPerPixelOverride;
+
+	TSharedPtr<class SNumericEntryBox<uint32>> BytesPerPixelOverrideInput;
 
 	TSharedPtr<class SButton> ButtonApplyToPreview;
 
@@ -51,5 +55,15 @@ private:
 	TWeakPtr<FMaterialEditor> MaterialEditorPtr;
 
 	bool bUpdateRequested = true;
+
+	bool bBytesPerPixelStartedTransaction;
+	uint32 BytesPerPixelOverride;
+//	void OnBytesPerPixelChanged(uint32 NewValue);
+	void OnBytesPerPixelCommitted(uint32 NewValue, ETextCommit::Type InCommitType);
+	void OnBeginBytesPerPixelSliderMovement();
+	void OnEndBytesPerPixelSliderMovement(uint32 NewValue);
+	TOptional<uint32> GetBytesPerPixelValue() const;
+
+	void OnCheckBoxBytesPerPixelChanged(ECheckBoxState InCheckBoxState);
 };
 
diff --git a/Engine/Source/Runtime/Engine/Private/Materials/HLSLMaterialTranslator.cpp b/Engine/Source/Runtime/Engine/Private/Materials/HLSLMaterialTranslator.cpp
--- a/Engine/Source/Runtime/Engine/Private/Materials/HLSLMaterialTranslator.cpp
+++ b/Engine/Source/Runtime/Engine/Private/Materials/HLSLMaterialTranslator.cpp
@@ -11419,9 +11419,109 @@ bool FHLSLMaterialTranslator::FStrataCompilationContext::StrataGenerateDerivedMa
 
 	StrataSimplificationStatus.bRunFullSimplification = Compiler->StrataCompilationConfig.bFullSimplify || CompilationContextIndex == EStrataCompilationContext::SCC_FullySimplified;
 
-	const uint32 StrataBytePerPixel = Strata::GetBytePerPixel(Compiler->GetShaderPlatform());
+	if (!StrataSimplificationStatus.bRunFullSimplification)
+	{
+		// Generate LayerDepth value for all operators/bsdfs
+		int VOpTopBranchCountTaken = 0;
+		int VOpBottomBranchCountTaken = 0;
+		std::function<void(FStrataOperator&)> WalkOperatorsForDepth = [&](FStrataOperator& CurrentOperator) -> void
+		{
+			CurrentOperator.LayerDepth = VOpBottomBranchCountTaken;
+			switch (CurrentOperator.OperatorType)
+			{
+			case STRATA_OPERATOR_VERTICAL:
+			{
+				VOpTopBranchCountTaken++;
+				WalkOperatorsForDepth(StrataMaterialExpressionRegisteredOperators[CurrentOperator.LeftIndex]);
+				VOpTopBranchCountTaken--;
+				VOpBottomBranchCountTaken++;
+				WalkOperatorsForDepth(StrataMaterialExpressionRegisteredOperators[CurrentOperator.RightIndex]);
+				VOpBottomBranchCountTaken--;
+				break;
+			}
+			case STRATA_OPERATOR_HORIZONTAL:
+			case STRATA_OPERATOR_ADD:
+			{
+				WalkOperatorsForDepth(StrataMaterialExpressionRegisteredOperators[CurrentOperator.LeftIndex]);
+				WalkOperatorsForDepth(StrataMaterialExpressionRegisteredOperators[CurrentOperator.RightIndex]);
+				break;
+			}
+			case STRATA_OPERATOR_WEIGHT:
+			{
+				WalkOperatorsForDepth(StrataMaterialExpressionRegisteredOperators[CurrentOperator.LeftIndex]);
+				break;
+			}
+			case STRATA_OPERATOR_BSDF:
+			{
+				break;
+			}
+			}
+		};
+		WalkOperatorsForDepth(*StrataMaterialRootOperator);
+
+		//
+		// Make sure each and every path of the strata tree have valid children and path.
+		//
+		for (auto& It : StrataMaterialExpressionRegisteredOperators)
+		{
+			bool bMustHaveLeftChild = false;
+			bool bMustHaveRightChild = false;
+			switch (It.OperatorType)
+			{
+				// Operators without any child
+			case STRATA_OPERATOR_BSDF:
+				break;
+
+				// Operators with a single child
+			case STRATA_OPERATOR_WEIGHT:
+			{
+			}
+			break;
+
+			// Operators with two children
+			case STRATA_OPERATOR_HORIZONTAL:
+			case STRATA_OPERATOR_VERTICAL:
+			case STRATA_OPERATOR_ADD:
+			{
+				FStrataSimplificationStatus::FOperatorToSimplify OperatorToSimplify;
+				OperatorToSimplify.Data.Index = It.Index;
+				OperatorToSimplify.Data.Depth = It.LayerDepth;
+				StrataSimplificationStatus.OperatorSimplificationOrder.Push(OperatorToSimplify);
+			}
+			break;
+			}
+			StrataSimplificationStatus.OperatorSimplificationOrder.Sort();	// sort according to depth
+		}
+	}
+
+
+	const uint32 StrataBytePerPixel = Compiler->StrataCompilationConfig.BytesPerPixelOverride > 0 ? Compiler->StrataCompilationConfig.BytesPerPixelOverride : Strata::GetBytePerPixel(Compiler->GetShaderPlatform());
+
+	bool bFirstLoop = true;
 	do 
 	{
+		if (!bFirstLoop && !StrataSimplificationStatus.bRunFullSimplification && !StrataSimplificationStatus.bMaterialFitsInMemoryBudget && StrataSimplificationStatus.OperatorSimplificationOrder.Num() > 0)
+		{
+			// Mark the deepest operator for parameter blending
+			FStrataSimplificationStatus::FOperatorToSimplify& OperatorToSimplify = StrataSimplificationStatus.OperatorSimplificationOrder.Top();
+			StrataMaterialExpressionRegisteredOperators[OperatorToSimplify.Data.Index].bNodeRequestParameterBlending = true;
+			StrataSimplificationStatus.OperatorSimplificationOrder.Pop(/*bAllowShrinking*/false);
+		}
+		else if (!bFirstLoop && !StrataSimplificationStatus.bRunFullSimplification && !StrataSimplificationStatus.bMaterialFitsInMemoryBudget && StrataSimplificationStatus.OperatorSimplificationOrder.Num() == 0)
+		{
+			StrataSimplificationStatus.bRunFullSimplification = !StrataSimplificationStatus.bMaterialFitsInMemoryBudget;
+		}
+		else if (bFirstLoop)
+		{
+			// Fall through. This is needed for material with a single BSDF and no operator to simplify.
+			bFirstLoop = false;
+		}
+		else
+		{
+			Compiler->Errorf(TEXT("Unkown Substrate material simplification status error for %s (asset: %s).\r\n"), *CompilerMaterial->GetDebugName(), *CompilerMaterial->GetAssetPath().ToString());
+			return false;
+		}
+
 		// Reset some data
 		StrataMaterialBSDFCount = 0;
 
@@ -11917,7 +12017,6 @@ bool FHLSLMaterialTranslator::FStrataCompilationContext::StrataGenerateDerivedMa
 				// Record the original requested byte size before simplification, only for the first pass.
 				StrataSimplificationStatus.OriginalRequestedByteSize = StrataMaterialRequestedSizeByte;
 			}
-			StrataSimplificationStatus.bRunFullSimplification = !StrataSimplificationStatus.bMaterialFitsInMemoryBudget;
 
 			const uint32 RequestedSizeInUint = FMath::DivideAndRoundUp(StrataMaterialRequestedSizeByte, 4u);
 			check(RequestedSizeInUint < 256u);
diff --git a/Engine/Source/Runtime/Engine/Private/Materials/HLSLMaterialTranslator.h b/Engine/Source/Runtime/Engine/Private/Materials/HLSLMaterialTranslator.h
--- a/Engine/Source/Runtime/Engine/Private/Materials/HLSLMaterialTranslator.h
+++ b/Engine/Source/Runtime/Engine/Private/Materials/HLSLMaterialTranslator.h
@@ -436,6 +436,30 @@ protected:
 
 		uint32 OriginalRequestedByteSize = 0;
 		bool bRunFullSimplification = false;	// Simple implementation for now: if the material does not fit, we simply everything.
+
+		struct FOperatorToSimplify
+		{
+			union
+			{
+				uint32 PackedData;
+				struct
+				{
+					uint32 Index : 16; // Index of the operator
+					uint32 Depth : 16; // Depth of the operator
+				} Data;
+			};
+
+			FORCEINLINE bool operator!=(FOperatorToSimplify B) const
+			{
+				return PackedData != B.PackedData;
+			}
+
+			FORCEINLINE bool operator<(FOperatorToSimplify B) const
+			{
+				return PackedData < B.PackedData;
+			}
+		};
+		TArray<FOperatorToSimplify> OperatorSimplificationOrder;
 	};
 
 	/** Represent a shared local basis description with its associated code. */
diff --git a/Engine/Source/Runtime/Engine/Private/Materials/MaterialShader.cpp b/Engine/Source/Runtime/Engine/Private/Materials/MaterialShader.cpp
--- a/Engine/Source/Runtime/Engine/Private/Materials/MaterialShader.cpp
+++ b/Engine/Source/Runtime/Engine/Private/Materials/MaterialShader.cpp
@@ -578,7 +578,7 @@ FString FStrataCompilationConfig::GetShaderMapKeyString() const
 {
 	if (bFullSimplify)
 	{
-		return TEXT("_STRTFS");
+		return FString::Printf(TEXT("_STRTFS%i"), BytesPerPixelOverride < 0 ? 0 : BytesPerPixelOverride);
 	}
 	return TEXT("");
 }
@@ -586,11 +586,13 @@ FString FStrataCompilationConfig::GetShaderMapKeyString() const
 void FStrataCompilationConfig::UpdateHash(FSHA1& Hasher) const
 {
 	Hasher.Update((const uint8*)(&bFullSimplify), sizeof(bFullSimplify));
+	Hasher.Update((const uint8*)(&BytesPerPixelOverride), sizeof(BytesPerPixelOverride));
 }
 
 void FStrataCompilationConfig::Serialize(FArchive& Ar)
 {
 	Ar << bFullSimplify;
+	Ar << BytesPerPixelOverride;
 }
 #endif
 
diff --git a/Engine/Source/Runtime/Engine/Public/MaterialTypes.h b/Engine/Source/Runtime/Engine/Public/MaterialTypes.h
--- a/Engine/Source/Runtime/Engine/Public/MaterialTypes.h
+++ b/Engine/Source/Runtime/Engine/Public/MaterialTypes.h
@@ -465,6 +465,7 @@ class FSHA1;
 struct FStrataCompilationConfig
 {
 	bool bFullSimplify = false;
+	int16 BytesPerPixelOverride = -1;
 
 	FString GetShaderMapKeyString() const;
 
@@ -474,7 +475,7 @@ struct FStrataCompilationConfig
 
 	friend inline bool operator==(const FStrataCompilationConfig& Lhs, const FStrataCompilationConfig& Rhs)
 	{
-		return Lhs.bFullSimplify == Rhs.bFullSimplify;
+		return Lhs.bFullSimplify == Rhs.bFullSimplify && Lhs.BytesPerPixelOverride == Rhs.BytesPerPixelOverride;
 	}
 
 	friend inline bool operator!=(const FStrataCompilationConfig& Lhs, const FStrataCompilationConfig& Rhs)
