Substrate - rework substrate tab to be nicer

diff --git a/Engine/Source/Editor/MaterialEditor/Private/SMaterialEditorStrataWidget.cpp b/Engine/Source/Editor/MaterialEditor/Private/SMaterialEditorStrataWidget.cpp
--- a/Engine/Source/Editor/MaterialEditor/Private/SMaterialEditorStrataWidget.cpp
+++ b/Engine/Source/Editor/MaterialEditor/Private/SMaterialEditorStrataWidget.cpp
@@ -12,6 +12,7 @@
 #include "Widgets/Input/SCheckBox.h"
 #include "Widgets/Layout/SWrapBox.h"
 #include "Widgets/Layout/SScrollBox.h"
+#include "Widgets/Layout/SBorder.h"
 #include "Widgets/SNullWidget.h"
 #include "Widgets/Notifications/SErrorText.h"
 #include "SGraphSubstrateMaterial.h"
@@ -42,6 +43,7 @@
 
 	BytesPerPixelOverride = Strata::GetBytePerPixel(SP_PCD3D_SM5);
 	BytesPerPixelOverrideInput = SNew(SNumericEntryBox<uint32>)
+		.MinDesiredValueWidth(150.0f)
 		.MinValue(12)
 		.MaxValue(BytesPerPixelOverride)
 		.MinSliderValue(12)
@@ -49,7 +51,7 @@
 		.OnBeginSliderMovement(this, &SMaterialEditorStrataWidget::OnBeginBytesPerPixelSliderMovement)
 		.OnEndSliderMovement(this, &SMaterialEditorStrataWidget::OnEndBytesPerPixelSliderMovement)
 		.AllowSpin(true)
-//		.OnValueChanged(this, &SMaterialEditorStrataWidget::OnBytesPerPixelChanged)
+		.OnValueChanged(this, &SMaterialEditorStrataWidget::OnBytesPerPixelChanged)
 		.OnValueCommitted(this, &SMaterialEditorStrataWidget::OnBytesPerPixelCommitted)
 		.Value(this, &SMaterialEditorStrataWidget::GetBytesPerPixelValue)
 		.IsEnabled(false);
@@ -66,7 +68,7 @@
 			SNew(SVerticalBox)
 			+SVerticalBox::Slot()
 			//.AutoHeight()			// Cannot use that otherwise scrollbars disapear.
-			.Padding(0.0f, 5.0f, 0.0f, 0.0f)
+			.Padding(0.0f, 0.0f, 0.0f, 0.0f)
 			[
 				SNew(SScrollBox)
 				.Orientation(Orient_Vertical)
@@ -74,28 +76,14 @@
 				+ SScrollBox::Slot()
 				[
 					SNew(SVerticalBox)
-					+SVerticalBox::Slot()
+					+ SVerticalBox::Slot()
 					.AutoHeight()
+					.Padding(0.0f, 5.0f, 0.0f, 0.0f)
 					[
-						SNew(SHorizontalBox)
-						+SHorizontalBox::Slot()
-						.AutoWidth()
+						SNew(SWrapBox)
+						.UseAllottedSize(true)
+						+ SWrapBox::Slot()
+						.Padding(0.0f)
 						.HAlign(HAlign_Left)
 						.VAlign(VAlign_Center)
 						[
-							SNew(SWrapBox)
-							.UseAllottedSize(true)
-							+SWrapBox::Slot()
-							.Padding(5.0f)
-							.HAlign(HAlign_Left)
-							.VAlign(VAlign_Center)
-							[
-								CheckBoxForceFullSimplification->AsShared()
-							]
-						]
-						+SHorizontalBox::Slot()
-						.AutoWidth()
-						.Padding(16.0f, 0.0f)
-						.HAlign(HAlign_Left)
-						.VAlign(VAlign_Center)
-						[
@@ -103,53 +91,117 @@
 							.ColorAndOpacity(FLinearColor::White)
 							.ShadowColorAndOpacity(FLinearColor::Black)
 							.ShadowOffset(FVector2D::UnitVector)
-							.Text(LOCTEXT("FullsimplificationLabel", "Full simplification"))
+							.Text(LOCTEXT("MaterialSimplificationPreview", "Material simplification preview"))
 						]
-						+SHorizontalBox::Slot()
-						.AutoWidth()
+					]
+					+SVerticalBox::Slot()
+					.AutoHeight()
+					[
+						SNew(SBorder)
+						//.BorderImage(FAppStyle::GetBrush("DetailsView.CategoryTop"))
+						//.BorderBackgroundColor(FLinearColor(0.9f, 0.6f, 0.6f, 1.0f))
+						.Padding(FMargin(5.0f, 5.0f, 5.0f, 5.0f))
+						[
+							SNew(SHorizontalBox)
+							+SHorizontalBox::Slot()
+							.AutoWidth()
+							.HAlign(HAlign_Left)
+							.VAlign(VAlign_Center)
+							[
+								SNew(SWrapBox)
+								.UseAllottedSize(true)
+								+SWrapBox::Slot()
+								.Padding(20.0, 0.0, 0.0, 0.0) // Padding(float Left, float Top, float Right, float Bottom)
+								.HAlign(HAlign_Left)
+								.VAlign(VAlign_Center)
+								[
+									CheckBoxForceFullSimplification->AsShared()
+								]
+							]
+							+SHorizontalBox::Slot()
+							.AutoWidth()
+							.Padding(0.0f)
+							.HAlign(HAlign_Left)
+							.VAlign(VAlign_Center)
+							[
+								SNew(STextBlock)
+								.ColorAndOpacity(FLinearColor::White)
+								.ShadowColorAndOpacity(FLinearColor::Black)
+								.ShadowOffset(FVector2D::UnitVector)
+								.Text(LOCTEXT("FullsimplificationLabel", "Full simplification"))
+							]
+							+SHorizontalBox::Slot()
+							.AutoWidth()
+							.Padding(20.0, 0.0, 0.0, 0.0)
+							.HAlign(HAlign_Left)
+							.VAlign(VAlign_Center)
+							[
+								SNew(SWrapBox)
+								.UseAllottedSize(true)
+								+SWrapBox::Slot()
+								.Padding(5.0f)
+								.HAlign(HAlign_Left)
+								.VAlign(VAlign_Center)
+								[
+									CheckBoxBytesPerPixelOverride->AsShared()
+								]
+							]
+							+SHorizontalBox::Slot()
+							.AutoWidth()
+							.Padding(0.0f)
+							.HAlign(HAlign_Left)
+							.VAlign(VAlign_Center)
+							[
+								SNew(STextBlock)
+								.ColorAndOpacity(FLinearColor::White)
+								.ShadowColorAndOpacity(FLinearColor::Black)
+								.ShadowOffset(FVector2D::UnitVector)
+								.Text(LOCTEXT("OverrideBytesPerPixel", "Override bytes per pixel"))
+							]
+							+SHorizontalBox::Slot()
+							.MaxWidth(200)
+							.HAlign(HAlign_Left)
+							.VAlign(VAlign_Center)
+							[
+								SNew(SWrapBox)
+								.UseAllottedSize(true)
+								+SWrapBox::Slot()
+								.Padding(0.0, 0.0, 0.0, 0.0)
+								.HAlign(HAlign_Left)
+								.VAlign(VAlign_Center)
+								[
+									BytesPerPixelOverrideInput->AsShared()
+								]
+							]
+							+SHorizontalBox::Slot()
+							.AutoWidth()
+							.Padding(0.0f)
+							.HAlign(HAlign_Left)
+							.VAlign(VAlign_Center)
+							[
+								SNew(SWrapBox)
+								//.UseAllottedSize(true)
+								+ SWrapBox::Slot()
+								.Padding(20.0, 0.0, 0.0, 0.0)
+								.HAlign(HAlign_Left)
+								.VAlign(VAlign_Center)
+								[
+									ButtonApplyToPreview->AsShared()
+								]
+							]
+						] // Border
+					] // VerticalBox
+					
+
+
+					+SVerticalBox::Slot()
+					.AutoHeight()
+					.Padding(0.0f, 15.0f, 0.0f, 0.0f)
+					[
+						SNew(SWrapBox)
+						.UseAllottedSize(true)
+						+ SWrapBox::Slot()
+						.Padding(0.0f)
 						.HAlign(HAlign_Left)
 						.VAlign(VAlign_Center)
 						[
-							SNew(SWrapBox)
-							.UseAllottedSize(true)
-							+SWrapBox::Slot()
-							.Padding(5.0f)
-							.HAlign(HAlign_Left)
-							.VAlign(VAlign_Center)
-							[
-								CheckBoxBytesPerPixelOverride->AsShared()
-							]
-						]
-						+SHorizontalBox::Slot()
-						.AutoWidth()
-						.Padding(16.0f, 0.0f)
-						.HAlign(HAlign_Left)
-						.VAlign(VAlign_Center)
-						[
-							SNew(STextBlock)
-							.ColorAndOpacity(FLinearColor::White)
-							.ShadowColorAndOpacity(FLinearColor::Black)
-							.ShadowOffset(FVector2D::UnitVector)
-							.Text(LOCTEXT("OverrideBytesPerPixel", "Override bytes per pixel"))
-						]
-						+SHorizontalBox::Slot()
-						.AutoWidth()
-						.HAlign(HAlign_Left)
-						.VAlign(VAlign_Center)
-						[
-							SNew(SWrapBox)
-							.UseAllottedSize(true)
-							+SWrapBox::Slot()
-							.Padding(5.0f)
-							.HAlign(HAlign_Left)
-							.VAlign(VAlign_Center)
-							[
-								BytesPerPixelOverrideInput->AsShared()
-							]
-						]
-						+SHorizontalBox::Slot()
-						.AutoWidth()
-						.Padding(16.0f, 0.0f)
-						.HAlign(HAlign_Left)
-						.VAlign(VAlign_Center)
-						[
@@ -157,51 +209,68 @@
 							.ColorAndOpacity(FLinearColor::White)
 							.ShadowColorAndOpacity(FLinearColor::Black)
 							.ShadowOffset(FVector2D::UnitVector)
-							.Text(LOCTEXT("FullsimplificationLabel", "Bytes per pixel"))
+							.Text(LOCTEXT("MaterialTopologyPreview", "Material topology preview"))
 						]
-						+SHorizontalBox::Slot()
-						.AutoWidth()
-						.Padding(16.0f, 0.0f)
-						.HAlign(HAlign_Left)
-						.VAlign(VAlign_Center)
+					]
+					+SVerticalBox::Slot()
+					.AutoHeight()
+					.Padding(0.0f, 0.0f, 0.0f, 0.0f)
+					[
+						SNew(SBorder)
+						//.BorderImage(FAppStyle::GetBrush("DetailsView.CategoryTop"))
+						//.BorderBackgroundColor(FLinearColor(0.9f, 0.6f, 0.6f, 1.0f))
+						.Padding(FMargin(5.0f, 5.0f, 5.0f, 5.0f))
 						[
 							SNew(SWrapBox)
-							//.UseAllottedSize(true)
+							.UseAllottedSize(true)
 							+ SWrapBox::Slot()
-							.Padding(5.0f)
-							.HAlign(HAlign_Left)
+							.Padding(10.0f)
+							.HAlign(HAlign_Center)
 							.VAlign(VAlign_Center)
 							[
-								ButtonApplyToPreview->AsShared()
+								SAssignNew(MaterialBox, SBox)
 							]
 						]
 					]
+
+
+					
 					+SVerticalBox::Slot()
 					.AutoHeight()
-					.Padding(0.0f, 5.0f, 0.0f, 0.0f)
+					.Padding(0.0f, 15.0f, 0.0f, 0.0f)
 					[
 						SNew(SWrapBox)
 						.UseAllottedSize(true)
 						+ SWrapBox::Slot()
-						.Padding(5.0f)
-						.HAlign(HAlign_Center)
+						.Padding(0.0f)
+						.HAlign(HAlign_Left)
 						.VAlign(VAlign_Center)
 						[
-							SAssignNew(MaterialBox, SBox)
+							SNew(STextBlock)
+							.ColorAndOpacity(FLinearColor::White)
+							.ShadowColorAndOpacity(FLinearColor::Black)
+							.ShadowOffset(FVector2D::UnitVector)
+							.Text(LOCTEXT("MaterialAdvancedDetails", "Material advanced details"))
 						]
 					]
 					+SVerticalBox::Slot()
 					.AutoHeight()
-					.Padding(0.0f, 5.0f, 0.0f, 0.0f)
+					.Padding(0.0f, 0.0f, 0.0f, 0.0f)
 					[
-						SNew(SWrapBox)
-						.UseAllottedSize(true)
-						+SWrapBox::Slot()
-						.Padding(5.0f)
-						.HAlign(HAlign_Center)
-						.VAlign(VAlign_Center)
+						SNew(SBorder)
+						//.BorderImage(FAppStyle::GetBrush("DetailsView.CategoryTop"))
+						//.BorderBackgroundColor(FLinearColor(0.9f, 0.6f, 0.6f, 1.0f))
+						.Padding(FMargin(5.0f, 5.0f, 5.0f, 5.0f))
 						[
-							DescriptionTextBlock->AsShared()
+							SNew(SWrapBox)
+							.UseAllottedSize(true)
+							+SWrapBox::Slot()
+							.Padding(5.0f)
+							.HAlign(HAlign_Center)
+							.VAlign(VAlign_Center)
+							[
+								DescriptionTextBlock->AsShared()
+							]
 						]
 					]
 				]
@@ -403,9 +472,10 @@
 	}
 }
 
-//void SMaterialEditorStrataWidget::OnBytesPerPixelChanged(uint32 NewValue)
-//{
-//}
+void SMaterialEditorStrataWidget::OnBytesPerPixelChanged(uint32 NewValue)
+{
+	BytesPerPixelOverride = NewValue;
+}
 
 void SMaterialEditorStrataWidget::OnBytesPerPixelCommitted(uint32 NewValue, ETextCommit::Type InCommitType)
 {
diff --git a/Engine/Source/Editor/MaterialEditor/Private/SMaterialEditorStrataWidget.h b/Engine/Source/Editor/MaterialEditor/Private/SMaterialEditorStrataWidget.h
--- a/Engine/Source/Editor/MaterialEditor/Private/SMaterialEditorStrataWidget.h
+++ b/Engine/Source/Editor/MaterialEditor/Private/SMaterialEditorStrataWidget.h
@@ -58,7 +58,7 @@ private:
 
 	bool bBytesPerPixelStartedTransaction;
 	uint32 BytesPerPixelOverride;
-//	void OnBytesPerPixelChanged(uint32 NewValue);
+	void OnBytesPerPixelChanged(uint32 NewValue);
 	void OnBytesPerPixelCommitted(uint32 NewValue, ETextCommit::Type InCommitType);
 	void OnBeginBytesPerPixelSliderMovement();
 	void OnEndBytesPerPixelSliderMovement(uint32 NewValue);
