substrate - refined SMaterialEditorStrataWidget

diff --git a/Engine/Source/Editor/MaterialEditor/Private/SMaterialEditorStrataWidget.cpp b/Engine/Source/Editor/MaterialEditor/Private/SMaterialEditorStrataWidget.cpp
--- a/Engine/Source/Editor/MaterialEditor/Private/SMaterialEditorStrataWidget.cpp
+++ b/Engine/Source/Editor/MaterialEditor/Private/SMaterialEditorStrataWidget.cpp
@@ -16,6 +16,7 @@
 #include "Widgets/SNullWidget.h"
 #include "Widgets/Notifications/SErrorText.h"
 #include "SGraphSubstrateMaterial.h"
+#include "RHIShaderPlatform.h"
 
 #define LOCTEXT_NAMESPACE "SMaterialEditorStrataWidget"
 
@@ -521,7 +522,10 @@ FReply SMaterialEditorStrataWidget::OnButtonApplyToPreview()
 
 		FStrataCompilationConfig StrataCompilationConfig;
 		StrataCompilationConfig.bFullSimplify = CheckBoxForceFullSimplification->IsChecked();
-		StrataCompilationConfig.BytesPerPixelOverride = CheckBoxBytesPerPixelOverride->IsChecked() ? FMath::Clamp(BytesPerPixelOverride, 12, Strata::GetBytePerPixel(SP_PCD3D_SM5)) : -1;
+
+		// Have a look at MaterialStats.cpp when we want to visualise a specific platform.
+		StrataCompilationConfig.BytesPerPixelOverride = CheckBoxBytesPerPixelOverride->IsChecked() ? FMath::Clamp(BytesPerPixelOverride, 12, Strata::GetBytePerPixel(GMaxRHIShaderPlatform)) : -1;
+
 		MaterialInterface->SetStrataCompilationConfig(StrataCompilationConfig);
 
 		MaterialInterface->ForceRecompileForRendering();
