Substrate - fixed crashes when right clicking on a substrate node pin => "promote to parameter".
- Only ask for graph refreshing without reconstructing node.
- Also use the correct StrataEnable function

diff --git a/Engine/Source/Editor/MaterialEditor/Private/MaterialEditor.cpp b/Engine/Source/Editor/MaterialEditor/Private/MaterialEditor.cpp
--- a/Engine/Source/Editor/MaterialEditor/Private/MaterialEditor.cpp
+++ b/Engine/Source/Editor/MaterialEditor/Private/MaterialEditor.cpp
@@ -6064,15 +6064,19 @@ FText FMaterialEditor::GetOriginalObjectName() const
 
 void FMaterialEditor::UpdateStrataTopologyPreview()
 {
-	if (Strata::IsGlintEnabled())
+	if (Strata::IsStrataEnabled())
 	{
-		// Update all Strata node which have a preview
+		// Update all Strata node which have a preview.
 		for (int32 Index = 0; Index < Material->MaterialGraph->Nodes.Num(); ++Index)
 		{
 			UMaterialGraphNode* MaterialNode = Cast<UMaterialGraphNode>(Material->MaterialGraph->Nodes[Index]);
 			if (MaterialNode && MaterialNode->MaterialExpression && MaterialNode->MaterialExpression->IsA(UMaterialExpressionStrataBSDF::StaticClass()))
 			{
-				MaterialNode->ReconstructNode();
+				UEdGraph* Graph = MaterialNode->GetGraph();
+				if (Graph)
+				{
+					Graph->NotifyGraphChanged();
+				}
 			}
 		}
 	}
