Add GetMaterialExpressionInputNames() to MaterialEditingLibrary.
Returns an array of names for each input on an expression.

diff --git a/Engine/Source/Editor/MaterialEditor/Private/MaterialEditingLibrary.cpp b/Engine/Source/Editor/MaterialEditor/Private/MaterialEditingLibrary.cpp
--- a/Engine/Source/Editor/MaterialEditor/Private/MaterialEditingLibrary.cpp
+++ b/Engine/Source/Editor/MaterialEditor/Private/MaterialEditingLibrary.cpp
@@ -849,6 +849,30 @@ FString UMaterialEditingLibrary::GetMaterialPropertyInputNodeOutputName(UMateria
 	return FString();
 }
 
+TArray<FString> UMaterialEditingLibrary::GetMaterialExpressionInputNames(UMaterialExpression* MaterialExpression)
+{
+	TArray<FString> InputNames;
+
+	TArrayView<FExpressionInput*> Inputs = MaterialExpression->GetInputsView();
+	for (int InputIdx = 0; InputIdx < Inputs.Num(); InputIdx++)
+	{
+		FName Name;
+		if (UMaterialExpressionMaterialFunctionCall* FuncCall = Cast<UMaterialExpressionMaterialFunctionCall>(MaterialExpression))
+		{
+			// If a function call, don't want to compare string with type postfix
+			Name = FuncCall->GetInputNameWithType(InputIdx, false);
+		}
+		else
+		{
+			const FName ExpressionInputName = MaterialExpression->GetInputName(InputIdx);
+			Name = UMaterialGraphNode::GetShortenPinName(ExpressionInputName);
+		}
+
+		InputNames.Add(Name.ToString());
+	}
+	return InputNames;
+}
+
 TArray<UMaterialExpression*> UMaterialEditingLibrary::GetInputsForMaterialExpression(UMaterial* Material, UMaterialExpression* MaterialExpression)
 {
 	TArray<UMaterialExpression*> MaterialExpressions;
diff --git a/Engine/Source/Editor/MaterialEditor/Public/MaterialEditingLibrary.h b/Engine/Source/Editor/MaterialEditor/Public/MaterialEditingLibrary.h
--- a/Engine/Source/Editor/MaterialEditor/Public/MaterialEditingLibrary.h
+++ b/Engine/Source/Editor/MaterialEditor/Public/MaterialEditingLibrary.h
@@ -199,6 +199,10 @@ public:
 	UFUNCTION(BlueprintPure, Category = "MaterialEditing")
 	static FString GetMaterialPropertyInputNodeOutputName(UMaterial* Material, EMaterialProperty Property);
 
+	/** Get the set of nodes acting as inputs to a node from an active material editor */
+	UFUNCTION(BlueprintPure, Category = "MaterialEditing")
+	static TArray<FString> GetMaterialExpressionInputNames(UMaterialExpression* MaterialExpression);
+
 	/** Get the set of nodes acting as inputs to a node from an active material editor */
 	UFUNCTION(BlueprintPure, Category = "MaterialEditing")
 	static TArray<UMaterialExpression*> GetInputsForMaterialExpression(UMaterial* Material, UMaterialExpression* MaterialExpression);
