Add GetMaterialExpressionInputTypes() to MaterialEditingLibrary.
Returns an array of ints for each input on an expression that map to the EMaterialValueType enum.

diff --git a/Engine/Source/Editor/MaterialEditor/Private/MaterialEditingLibrary.cpp b/Engine/Source/Editor/MaterialEditor/Private/MaterialEditingLibrary.cpp
--- a/Engine/Source/Editor/MaterialEditor/Private/MaterialEditingLibrary.cpp
+++ b/Engine/Source/Editor/MaterialEditor/Private/MaterialEditingLibrary.cpp
@@ -854,7 +854,7 @@ TArray<FString> UMaterialEditingLibrary::GetMaterialExpressionInputNames(UMateri
 	TArray<FString> InputNames;
 
 	TArrayView<FExpressionInput*> Inputs = MaterialExpression->GetInputsView();
-	for (int InputIdx = 0; InputIdx < Inputs.Num(); InputIdx++)
+	for (int32 InputIdx = 0; InputIdx < Inputs.Num(); InputIdx++)
 	{
 		FName Name;
 		if (UMaterialExpressionMaterialFunctionCall* FuncCall = Cast<UMaterialExpressionMaterialFunctionCall>(MaterialExpression))
@@ -873,6 +873,18 @@ TArray<FString> UMaterialEditingLibrary::GetMaterialExpressionInputNames(UMateri
 	return InputNames;
 }
 
+TArray<int32> UMaterialEditingLibrary::GetMaterialExpressionInputTypes(UMaterialExpression* MaterialExpression)
+{
+	TArray<int32> InputTypes;
+
+	const int32 NumInputs = MaterialExpression->GetInputsView().Num();
+	for (int32 InputIdx = 0; InputIdx < NumInputs; InputIdx++)
+	{
+		InputTypes.Add(MaterialExpression->GetInputType(InputIdx));
+	}
+	return InputTypes;
+}
+
 TArray<UMaterialExpression*> UMaterialEditingLibrary::GetInputsForMaterialExpression(UMaterial* Material, UMaterialExpression* MaterialExpression)
 {
 	TArray<UMaterialExpression*> MaterialExpressions;
diff --git a/Engine/Source/Editor/MaterialEditor/Public/MaterialEditingLibrary.h b/Engine/Source/Editor/MaterialEditor/Public/MaterialEditingLibrary.h
--- a/Engine/Source/Editor/MaterialEditor/Public/MaterialEditingLibrary.h
+++ b/Engine/Source/Editor/MaterialEditor/Public/MaterialEditingLibrary.h
@@ -199,10 +199,14 @@ public:
 	UFUNCTION(BlueprintPure, Category = "MaterialEditing")
 	static FString GetMaterialPropertyInputNodeOutputName(UMaterial* Material, EMaterialProperty Property);
 
-	/** Get the set of nodes acting as inputs to a node from an active material editor */
+	/** Get the array of input pin names for a material expression */
 	UFUNCTION(BlueprintPure, Category = "MaterialEditing")
 	static TArray<FString> GetMaterialExpressionInputNames(UMaterialExpression* MaterialExpression);
 
+	/** Get the array of input pin types for a material expression */
+	UFUNCTION(BlueprintPure, Category = "MaterialEditing")
+	static TArray<int32> GetMaterialExpressionInputTypes(UMaterialExpression* MaterialExpression);
+
 	/** Get the set of nodes acting as inputs to a node from an active material editor */
 	UFUNCTION(BlueprintPure, Category = "MaterialEditing")
 	static TArray<UMaterialExpression*> GetInputsForMaterialExpression(UMaterial* Material, UMaterialExpression* MaterialExpression);
