Fix GetMaterialExpressionInputTypes() to get type from the feeding pin.

diff --git a/Engine/Source/Editor/MaterialEditor/Private/MaterialEditingLibrary.cpp b/Engine/Source/Editor/MaterialEditor/Private/MaterialEditingLibrary.cpp
--- a/Engine/Source/Editor/MaterialEditor/Private/MaterialEditingLibrary.cpp
+++ b/Engine/Source/Editor/MaterialEditor/Private/MaterialEditingLibrary.cpp
@@ -877,10 +877,19 @@ TArray<int32> UMaterialEditingLibrary::GetMaterialExpressionInputTypes(UMaterial
 {
 	TArray<int32> InputTypes;
 
-	const int32 NumInputs = MaterialExpression->GetInputsView().Num();
-	for (int32 InputIdx = 0; InputIdx < NumInputs; InputIdx++)
+	TArrayView<FExpressionInput*> Inputs = MaterialExpression->GetInputsView();
+	for (int32 InputIdx = 0; InputIdx < Inputs.Num(); InputIdx++)
 	{
-		InputTypes.Add(MaterialExpression->GetInputType(InputIdx));
+		FExpressionInput* Input = Inputs[InputIdx];
+		UMaterialExpression* Expression = Input != nullptr ? Input->Expression : nullptr;
+		if (Expression != nullptr)
+		{
+			InputTypes.Add(Expression->GetOutputType(Input->OutputIndex));
+		}
+		else
+		{
+			InputTypes.Add(MaterialExpression->GetInputType(InputIdx));
+		}
 	}
 	return InputTypes;
 }
