Add bluepprint getter for nanite override material.

diff --git a/Engine/Source/Editor/MaterialEditor/Private/MaterialEditingLibrary.cpp b/Engine/Source/Editor/MaterialEditor/Private/MaterialEditingLibrary.cpp
--- a/Engine/Source/Editor/MaterialEditor/Private/MaterialEditingLibrary.cpp
+++ b/Engine/Source/Editor/MaterialEditor/Private/MaterialEditingLibrary.cpp
@@ -1392,4 +1392,9 @@ FMaterialStatistics UMaterialEditingLibrary::GetStatistics(class UMaterialInterf
 	}
 
 	return Result;
-}
\ No newline at end of file
+}
+
+UMaterialInterface* UMaterialEditingLibrary::GetNaniteOverrideMaterial(UMaterialInterface* Material)
+{
+	return Material->GetNaniteOverride();
+}
diff --git a/Engine/Source/Editor/MaterialEditor/Public/MaterialEditingLibrary.h b/Engine/Source/Editor/MaterialEditor/Public/MaterialEditingLibrary.h
--- a/Engine/Source/Editor/MaterialEditor/Public/MaterialEditingLibrary.h
+++ b/Engine/Source/Editor/MaterialEditor/Public/MaterialEditingLibrary.h
@@ -390,4 +390,8 @@ public:
 	/** Returns statistics about the given material */
 	UFUNCTION(BlueprintCallable, Category = "MaterialEditing")
 	static FMaterialStatistics GetStatistics(UMaterialInterface* Material);
+
+	/** Returns any nanite override material for the given material */
+	UFUNCTION(BlueprintCallable, Category = "MaterialEditing")
+	static UMaterialInterface* GetNaniteOverrideMaterial(UMaterialInterface* Material);
 };
\ No newline at end of file
