Fixed crash occurring in FMaterialEditorUtilities::GetVisibleMaterialParametersFromExpression when a function input isn't found.

diff --git a/Engine/Source/Editor/MaterialEditor/Private/MaterialEditorUtilities.cpp b/Engine/Source/Editor/MaterialEditor/Private/MaterialEditorUtilities.cpp
--- a/Engine/Source/Editor/MaterialEditor/Private/MaterialEditorUtilities.cpp
+++ b/Engine/Source/Editor/MaterialEditor/Private/MaterialEditorUtilities.cpp
@@ -614,9 +614,10 @@ void FMaterialEditorUtilities::GetVisibleMaterialParametersFromExpression(
 		FGetVisibleMaterialParametersFunctionState* FunctionState = FunctionStack.Pop();
 
 		const FFunctionExpressionInput* MatchingInput = FindInputById(FunctionInputExpression, FunctionState->FunctionCall->FunctionInputs);
-		check(MatchingInput);
-
-		GetVisibleMaterialParametersFromExpression(FMaterialExpressionKey(MatchingInput->Input.Expression, MatchingInput->Input.OutputIndex), MaterialInstance, VisibleExpressions, FunctionStack);
+		if (ensure(MatchingInput))
+		{
+			GetVisibleMaterialParametersFromExpression(FMaterialExpressionKey(MatchingInput->Input.Expression, MatchingInput->Input.OutputIndex), MaterialInstance, VisibleExpressions, FunctionStack);
+		}
 
 		FunctionStack.Push(FunctionState);
 	}
