New candidate fix for crash occurring when linking a node to the UV input of a Texture Sample.

diff --git a/Engine/Source/Editor/MaterialEditor/Private/MaterialEditorUtilities.cpp b/Engine/Source/Editor/MaterialEditor/Private/MaterialEditorUtilities.cpp
--- a/Engine/Source/Editor/MaterialEditor/Private/MaterialEditorUtilities.cpp
+++ b/Engine/Source/Editor/MaterialEditor/Private/MaterialEditorUtilities.cpp
@@ -614,10 +614,9 @@ void FMaterialEditorUtilities::GetVisibleMaterialParametersFromExpression(
 		FGetVisibleMaterialParametersFunctionState* FunctionState = FunctionStack.Pop();
 
 		const FFunctionExpressionInput* MatchingInput = FindInputById(FunctionInputExpression, FunctionState->FunctionCall->FunctionInputs);
-		if (ensure(MatchingInput))
-		{
-			GetVisibleMaterialParametersFromExpression(FMaterialExpressionKey(MatchingInput->Input.Expression, MatchingInput->Input.OutputIndex), MaterialInstance, VisibleExpressions, FunctionStack);
-		}
+		check(MatchingInput);
+		
+		GetVisibleMaterialParametersFromExpression(FMaterialExpressionKey(MatchingInput->Input.Expression, MatchingInput->Input.OutputIndex), MaterialInstance, VisibleExpressions, FunctionStack);
 
 		FunctionStack.Push(FunctionState);
 	}
@@ -631,11 +630,10 @@ void FMaterialEditorUtilities::GetVisibleMaterialParametersFromExpression(
 		}
 		else
 		{
-			TArrayView<FExpressionInput*> ExpressionInputs = MaterialExpressionKey.Expression->GetInputsView();
-			for (int32 ExpressionInputIndex = 0; ExpressionInputIndex < ExpressionInputs.Num(); ExpressionInputIndex++)
+			// Retrieve the expression input and then start parsing its children
+			for (int32 i = 0; i < MaterialExpressionKey.Expression->GetInputsView().Num(); i++)
 			{
-				//retrieve the expression input and then start parsing its children
-				FExpressionInput* Input = ExpressionInputs[ExpressionInputIndex];
+				FExpressionInput* Input = MaterialExpressionKey.Expression->GetInputsView()[i];
 				GetVisibleMaterialParametersFromExpression(FMaterialExpressionKey(Input->Expression, Input->OutputIndex), MaterialInstance, VisibleExpressions, FunctionStack);
 			}
 
