Fix bad cast in MaterialEditor.cpp
#rnx

diff --git a/Engine/Source/Editor/MaterialEditor/Private/MaterialEditor.cpp b/Engine/Source/Editor/MaterialEditor/Private/MaterialEditor.cpp
--- a/Engine/Source/Editor/MaterialEditor/Private/MaterialEditor.cpp
+++ b/Engine/Source/Editor/MaterialEditor/Private/MaterialEditor.cpp
@@ -7263,7 +7263,7 @@ void FMaterialEditor::DeepCopyExpressions(UMaterialGraph* CopyGraph, UMaterialEx
 			PostPasteMaterialExpression(MaterialNode->MaterialExpression);
 
 		}
-		else if (UMaterialGraphNode_Comment* Comment = Cast<UMaterialGraphNode_Comment>(MaterialNode))
+		else if (UMaterialGraphNode_Comment* Comment = Cast<UMaterialGraphNode_Comment>(Node))
 		{
 			Comment->Modify();
 			Comment->Rename(/*NewName=*/ NULL, /*NewOuter=*/ CopyGraph);
