Substrate - Rename/move file(s)

diff --git a/Engine/Source/Editor/CommonMenuExtensions/Private/CommonMenuExtensionsModule.cpp b/Engine/Source/Editor/CommonMenuExtensions/Private/CommonMenuExtensionsModule.cpp
--- a/Engine/Source/Editor/CommonMenuExtensions/Private/CommonMenuExtensionsModule.cpp
+++ b/Engine/Source/Editor/CommonMenuExtensions/Private/CommonMenuExtensionsModule.cpp
@@ -4,7 +4,7 @@
 #include "BufferVisualizationMenuCommands.h"
 #include "NaniteVisualizationMenuCommands.h"
 #include "LumenVisualizationMenuCommands.h"
-#include "StrataVisualizationMenuCommands.h"
+#include "SubstrateVisualizationMenuCommands.h"
 #include "GroomVisualizationMenuCommands.h"
 #include "VirtualShadowMapVisualizationMenuCommands.h"
 #include "ShowFlagMenuCommands.h"
diff --git a/Engine/Source/Editor/CommonMenuExtensions/Private/StrataVisualizationMenuCommands.cpp b/Engine/Source/Editor/CommonMenuExtensions/Private/StrataVisualizationMenuCommands.cpp
deleted file mode 100644
diff --git a/Engine/Source/Editor/CommonMenuExtensions/Private/SubstrateVisualizationMenuCommands.cpp b/Engine/Source/Editor/CommonMenuExtensions/Private/SubstrateVisualizationMenuCommands.cpp
new file mode 100644
--- /dev/null
+++ b/Engine/Source/Editor/CommonMenuExtensions/Private/SubstrateVisualizationMenuCommands.cpp
@@ -0,0 +1,144 @@
+// Copyright Epic Games, Inc. All Rights Reserved.
+
+#include "SubstrateVisualizationMenuCommands.h"
+
+#include "Delegates/Delegate.h"
+#include "EditorViewportClient.h"
+#include "Framework/Commands/InputChord.h"
+#include "Framework/Commands/UIAction.h"
+#include "Framework/Commands/UICommandInfo.h"
+#include "Framework/Commands/UICommandList.h"
+#include "Framework/MultiBox/MultiBoxBuilder.h"
+#include "Internationalization/Internationalization.h"
+#include "Internationalization/Text.h"
+#include "SubstrateVisualizationData.h"
+#include "Misc/AssertionMacros.h"
+#include "Styling/AppStyle.h"
+#include "UObject/UnrealNames.h"
+
+#define LOCTEXT_NAMESPACE "SubstrateVisualizationMenuCommands"
+
+FStrataVisualizationMenuCommands::FStrataVisualizationMenuCommands()
+	: TCommands<FStrataVisualizationMenuCommands>
+	(
+		TEXT("SubstrateVisualizationMenu"), // Context name for fast lookup
+		NSLOCTEXT("Contexts", "SubstrateVisualizationMenu", "Substrate"), // Localized context name for displaying
+		NAME_None, // Parent context name.  
+		FAppStyle::GetAppStyleSetName() // Icon Style Set
+	),
+	CommandMap()
+{
+}
+
+void FStrataVisualizationMenuCommands::BuildCommandMap()
+{
+	const FStrataVisualizationData& VisualizationData = GetStrataVisualizationData();
+	const FStrataVisualizationData::TModeMap& ModeMap = VisualizationData.GetModeMap();
+
+	CommandMap.Empty();
+	for (FStrataVisualizationData::TModeMap::TConstIterator It = ModeMap.CreateConstIterator(); It; ++It)
+	{
+		const FStrataVisualizationData::FModeRecord& Entry = It.Value();
+		FStrataVisualizationRecord& Record = CommandMap.Add(Entry.ModeName, FStrataVisualizationRecord());
+		Record.Name = Entry.ModeName;
+		Record.Command = FUICommandInfoDecl(
+			this->AsShared(),
+			Entry.ModeName,
+			Entry.ModeText,
+			Entry.ModeDesc)
+			.UserInterfaceType(EUserInterfaceActionType::RadioButton)
+			.DefaultChord(FInputChord()
+		);
+		Record.ViewMode = Entry.ViewMode;
+	}
+}
+
+void FStrataVisualizationMenuCommands::BuildVisualisationSubMenu(FMenuBuilder& Menu)
+{
+	const FStrataVisualizationMenuCommands& Commands = FStrataVisualizationMenuCommands::Get();
+	if (Commands.IsPopulated())
+	{
+		// General
+		{
+			Menu.BeginSection("LevelViewportSubstrateVisualizationModeGeneral", LOCTEXT("SubstrateVisualizationGeneral", "Substrate General View Mode"));
+			Commands.AddCommandTypeToMenu(Menu, FStrataViewMode::MaterialProperties);
+			Commands.AddCommandTypeToMenu(Menu, FStrataViewMode::MaterialCount);
+			Commands.AddCommandTypeToMenu(Menu, FStrataViewMode::MaterialByteCount);
+			Commands.AddCommandTypeToMenu(Menu, FStrataViewMode::StrataInfo);
+			Menu.EndSection();
+		}
+
+		// Advanced
+		{
+			Menu.BeginSection("LevelViewportSubstrateVisualizationModeAdvanced", LOCTEXT("SubstrateVisualizationAdvanced", "Substrate Advanced View Mode"));
+			Commands.AddCommandTypeToMenu(Menu, FStrataViewMode::AdvancedMaterialProperties);
+			Commands.AddCommandTypeToMenu(Menu, FStrataViewMode::MaterialClassification);
+			Commands.AddCommandTypeToMenu(Menu, FStrataViewMode::RoughRefractionClassification);
+			Commands.AddCommandTypeToMenu(Menu, FStrataViewMode::DecalClassification);
+			Menu.EndSection();
+		}
+	}
+}
+
+bool FStrataVisualizationMenuCommands::AddCommandTypeToMenu(FMenuBuilder& Menu, const FStrataViewMode ViewMode) const
+{
+	bool bAddedCommands = false;
+
+	const TStrataVisualizationModeCommandMap& Commands = CommandMap;
+	for (TCommandConstIterator It = CreateCommandConstIterator(); It; ++It)
+	{
+		const FStrataVisualizationRecord& Record = It.Value();
+		if (Record.ViewMode == ViewMode)
+		{
+			Menu.AddMenuEntry(Record.Command, NAME_None, Record.Command->GetLabel());
+			bAddedCommands = true;
+		}
+	}
+
+	return bAddedCommands;
+}
+
+FStrataVisualizationMenuCommands::TCommandConstIterator FStrataVisualizationMenuCommands::CreateCommandConstIterator() const
+{
+	return CommandMap.CreateConstIterator();
+}
+
+void FStrataVisualizationMenuCommands::RegisterCommands()
+{
+	BuildCommandMap();
+}
+
+void FStrataVisualizationMenuCommands::BindCommands(FUICommandList& CommandList, const TSharedPtr<FEditorViewportClient>& Client) const
+{
+	// Map Strata visualization mode actions
+	for (FStrataVisualizationMenuCommands::TCommandConstIterator It = FStrataVisualizationMenuCommands::Get().CreateCommandConstIterator(); It; ++It)
+	{
+		const FStrataVisualizationMenuCommands::FStrataVisualizationRecord& Record = It.Value();
+		CommandList.MapAction(
+			Record.Command,
+			FExecuteAction::CreateStatic(&FStrataVisualizationMenuCommands::ChangeStrataVisualizationMode, Client.ToWeakPtr(), Record.Name),
+			FCanExecuteAction(),
+			FIsActionChecked::CreateStatic(&FStrataVisualizationMenuCommands::IsStrataVisualizationModeSelected, Client.ToWeakPtr(), Record.Name)
+		);
+	}
+}
+
+void FStrataVisualizationMenuCommands::ChangeStrataVisualizationMode(TWeakPtr<FEditorViewportClient> WeakClient, FName InName)
+{
+	if (TSharedPtr<FEditorViewportClient> Client = WeakClient.Pin())
+	{
+		Client->ChangeStrataVisualizationMode(InName);
+	}
+}
+
+bool FStrataVisualizationMenuCommands::IsStrataVisualizationModeSelected(TWeakPtr<FEditorViewportClient> WeakClient, FName InName)
+{
+	if (TSharedPtr<FEditorViewportClient> Client = WeakClient.Pin())
+	{
+		return Client->IsStrataVisualizationModeSelected(InName);
+	}
+	
+	return false;
+}
+
+#undef LOCTEXT_NAMESPACE
diff --git a/Engine/Source/Editor/CommonMenuExtensions/Public/StrataVisualizationMenuCommands.h b/Engine/Source/Editor/CommonMenuExtensions/Public/StrataVisualizationMenuCommands.h
deleted file mode 100644
diff --git a/Engine/Source/Editor/CommonMenuExtensions/Public/SubstrateVisualizationMenuCommands.h b/Engine/Source/Editor/CommonMenuExtensions/Public/SubstrateVisualizationMenuCommands.h
new file mode 100644
--- /dev/null
+++ b/Engine/Source/Editor/CommonMenuExtensions/Public/SubstrateVisualizationMenuCommands.h
@@ -0,0 +1,56 @@
+// Copyright Epic Games, Inc. All Rights Reserved.
+
+#pragma once
+
+#include "Containers/Map.h"
+#include "Framework/Commands/Commands.h"
+#include "Framework/Commands/UICommandInfo.h"
+#include "Framework/MultiBox/MultiBoxBuilder.h"
+#include "HAL/Platform.h"
+#include "Templates/SharedPointer.h"
+#include "UObject/NameTypes.h"
+#include "SubstrateVisualizationData.h"
+
+class FEditorViewportClient;
+class FMenuBuilder;
+class FUICommandInfo;
+class FUICommandList;
+
+class COMMONMENUEXTENSIONS_API FStrataVisualizationMenuCommands : public TCommands<FStrataVisualizationMenuCommands>
+{
+public:
+	struct FStrataVisualizationRecord
+	{
+		FName Name;
+		TSharedPtr<FUICommandInfo> Command;
+		FStrataViewMode ViewMode = FStrataViewMode::MaterialProperties;
+	};
+
+	typedef TMultiMap<FName, FStrataVisualizationRecord> TStrataVisualizationModeCommandMap;
+	typedef TStrataVisualizationModeCommandMap::TConstIterator TCommandConstIterator;
+
+	FStrataVisualizationMenuCommands();
+
+	TCommandConstIterator CreateCommandConstIterator() const;
+
+	static void BuildVisualisationSubMenu(FMenuBuilder& Menu);
+
+	virtual void RegisterCommands() override;
+
+	void BindCommands(FUICommandList& CommandList, const TSharedPtr<FEditorViewportClient>& Client) const;
+	
+	inline bool IsPopulated() const
+	{
+		return CommandMap.Num() > 0;
+	}
+
+private:
+	void BuildCommandMap();
+	bool AddCommandTypeToMenu(FMenuBuilder& Menu, const FStrataViewMode ViewMode) const;
+
+	static void ChangeStrataVisualizationMode(TWeakPtr<FEditorViewportClient> WeakClient, FName InName);
+	static bool IsStrataVisualizationModeSelected(TWeakPtr<FEditorViewportClient> WeakClient, FName InName);
+
+private:
+	TStrataVisualizationModeCommandMap CommandMap;
+};
\ No newline at end of file
diff --git a/Engine/Source/Editor/GraphEditor/Private/MaterialNodes/SGraphNodeMaterialBase.cpp b/Engine/Source/Editor/GraphEditor/Private/MaterialNodes/SGraphNodeMaterialBase.cpp
--- a/Engine/Source/Editor/GraphEditor/Private/MaterialNodes/SGraphNodeMaterialBase.cpp
+++ b/Engine/Source/Editor/GraphEditor/Private/MaterialNodes/SGraphNodeMaterialBase.cpp
@@ -61,7 +61,7 @@
 #include "Widgets/SOverlay.h"
 #include "Widgets/SViewport.h"
 #include "Widgets/Text/STextBlock.h"
-#include "Rendering/StrataMaterialShared.h"
+#include "Rendering/SubstrateMaterialShared.h"
 #include "SGraphSubstrateMaterial.h"
 #include "MaterialShared.h"
 
diff --git a/Engine/Source/Editor/GraphEditor/Private/MaterialNodes/SGraphSubstrateMaterial.cpp b/Engine/Source/Editor/GraphEditor/Private/MaterialNodes/SGraphSubstrateMaterial.cpp
--- a/Engine/Source/Editor/GraphEditor/Private/MaterialNodes/SGraphSubstrateMaterial.cpp
+++ b/Engine/Source/Editor/GraphEditor/Private/MaterialNodes/SGraphSubstrateMaterial.cpp
@@ -3,7 +3,7 @@
 #include "SGraphSubstrateMaterial.h"
 #include "Internationalization/Text.h"
 #include "Rendering/RenderingCommon.h"
-#include "Rendering/StrataMaterialShared.h"
+#include "Rendering/SubstrateMaterialShared.h"
 #include "MaterialGraph/MaterialGraphNode.h"
 #include "MaterialGraph/MaterialGraphNode_Root.h"
 #include "Materials/MaterialExpression.h"
diff --git a/Engine/Source/Editor/LevelEditor/Private/SLevelViewport.cpp b/Engine/Source/Editor/LevelEditor/Private/SLevelViewport.cpp
--- a/Engine/Source/Editor/LevelEditor/Private/SLevelViewport.cpp
+++ b/Engine/Source/Editor/LevelEditor/Private/SLevelViewport.cpp
@@ -80,7 +80,7 @@
 #include "BufferVisualizationMenuCommands.h"
 #include "NaniteVisualizationMenuCommands.h"
 #include "LumenVisualizationMenuCommands.h"
-#include "StrataVisualizationMenuCommands.h"
+#include "SubstrateVisualizationMenuCommands.h"
 #include "VirtualShadowMapVisualizationMenuCommands.h"
 #include "EditorLevelUtils.h"
 #include "Engine/LevelStreaming.h"
diff --git a/Engine/Source/Editor/LevelEditor/Public/LevelViewportActions.h b/Engine/Source/Editor/LevelEditor/Public/LevelViewportActions.h
--- a/Engine/Source/Editor/LevelEditor/Public/LevelViewportActions.h
+++ b/Engine/Source/Editor/LevelEditor/Public/LevelViewportActions.h
@@ -10,7 +10,7 @@
 #include "BufferVisualizationMenuCommands.h"
 #include "NaniteVisualizationMenuCommands.h"
 #include "LumenVisualizationMenuCommands.h"
-#include "StrataVisualizationMenuCommands.h"
+#include "SubstrateVisualizationMenuCommands.h"
 #include "GroomVisualizationMenuCommands.h"
 #include "VirtualShadowMapVisualizationMenuCommands.h"
 
diff --git a/Engine/Source/Editor/MaterialEditor/Private/MaterialEditor.cpp b/Engine/Source/Editor/MaterialEditor/Private/MaterialEditor.cpp
--- a/Engine/Source/Editor/MaterialEditor/Private/MaterialEditor.cpp
+++ b/Engine/Source/Editor/MaterialEditor/Private/MaterialEditor.cpp
@@ -155,7 +155,7 @@
 #include "MaterialEditorModes.h"
 #include "Materials/MaterialExpression.h"
 #include "MaterialCachedHLSLTree.h"
-#include "SMaterialEditorStrataWidget.h"
+#include "SMaterialEditorSubstrateWidget.h"
 #include "SGraphSubstrateMaterial.h"
 
 #include "SMaterialParametersOverviewWidget.h"
diff --git a/Engine/Source/Editor/MaterialEditor/Private/SMaterialEditorStrataWidget.cpp b/Engine/Source/Editor/MaterialEditor/Private/SMaterialEditorStrataWidget.cpp
deleted file mode 100644
diff --git a/Engine/Source/Editor/MaterialEditor/Private/SMaterialEditorStrataWidget.h b/Engine/Source/Editor/MaterialEditor/Private/SMaterialEditorStrataWidget.h
deleted file mode 100644
diff --git a/Engine/Source/Editor/MaterialEditor/Private/SMaterialEditorSubstrateWidget.cpp b/Engine/Source/Editor/MaterialEditor/Private/SMaterialEditorSubstrateWidget.cpp
new file mode 100644
--- /dev/null
+++ b/Engine/Source/Editor/MaterialEditor/Private/SMaterialEditorSubstrateWidget.cpp
@@ -0,0 +1,537 @@
+// Copyright Epic Games, Inc. All Rights Reserved.
+
+#include "SMaterialEditorSubstrateWidget.h"
+
+#include "EditorWidgetsModule.h"
+#include "AssetRegistry/AssetRegistryModule.h"
+#include "MaterialEditor.h"
+#include "SubstrateDefinitions.h"
+#include <functional>
+
+#include "Widgets/Input/SButton.h"
+#include "Widgets/Input/SCheckBox.h"
+#include "Widgets/Layout/SWrapBox.h"
+#include "Widgets/Layout/SScrollBox.h"
+#include "Widgets/Layout/SBorder.h"
+#include "Widgets/SNullWidget.h"
+#include "Widgets/Notifications/SErrorText.h"
+#include "SGraphSubstrateMaterial.h"
+#include "RHIShaderPlatform.h"
+
+#define LOCTEXT_NAMESPACE "SMaterialEditorStrataWidget"
+
+DEFINE_LOG_CATEGORY_STATIC(LogMaterialEditorStrataWidget, Log, All);
+
+void SMaterialEditorStrataWidget::Construct(const FArguments& InArgs, TWeakPtr<FMaterialEditor> InMaterialEditorPtr)
+{
+	MaterialEditorPtr = InMaterialEditorPtr;
+
+	ButtonApplyToPreview = SNew(SButton)
+		.HAlign(HAlign_Center)
+		.OnClicked(this, &SMaterialEditorStrataWidget::OnButtonApplyToPreview)
+		.Text(LOCTEXT("ButtonApplyToPreview", "Apply to preview"));
+
+	CheckBoxForceFullSimplification = SNew(SCheckBox)
+		.Padding(5.0f)
+		.ToolTipText(LOCTEXT("CheckBoxForceFullSimplificationToolTip", "This will force full simplification of the material."));
+
+	DescriptionTextBlock = SNew(STextBlock)
+		.TextStyle(FAppStyle::Get(), "Log.Normal")
+		.ColorAndOpacity(FLinearColor::White)
+		.ShadowColorAndOpacity(FLinearColor::Black)
+		.ShadowOffset(FVector2D::UnitVector)
+		.Text(LOCTEXT("DescriptionTextBlock_Default", "Shader is compiling"));
+
+	BytesPerPixelOverride = Strata::GetBytePerPixel(SP_PCD3D_SM5);
+	BytesPerPixelOverrideInput = SNew(SNumericEntryBox<uint32>)
+		.MinDesiredValueWidth(150.0f)
+		.MinValue(12)
+		.MaxValue(BytesPerPixelOverride)
+		.MinSliderValue(12)
+		.MaxSliderValue(BytesPerPixelOverride)
+		.OnBeginSliderMovement(this, &SMaterialEditorStrataWidget::OnBeginBytesPerPixelSliderMovement)
+		.OnEndSliderMovement(this, &SMaterialEditorStrataWidget::OnEndBytesPerPixelSliderMovement)
+		.AllowSpin(true)
+		.OnValueChanged(this, &SMaterialEditorStrataWidget::OnBytesPerPixelChanged)
+		.OnValueCommitted(this, &SMaterialEditorStrataWidget::OnBytesPerPixelCommitted)
+		.Value(this, &SMaterialEditorStrataWidget::GetBytesPerPixelValue)
+		.IsEnabled(false);
+
+	CheckBoxBytesPerPixelOverride = SNew(SCheckBox)
+		.Padding(5.0f)
+		.ToolTipText(LOCTEXT("CheckBoxBytesPerPixelOverride", "This will force the byte per pixel count for the preview material. It cannot go higher than the current project setting."))
+		.OnCheckStateChanged(this, &SMaterialEditorStrataWidget::OnCheckBoxBytesPerPixelChanged);
+
+	if (Strata::IsStrataEnabled())
+	{
+		this->ChildSlot
+		[
+			SNew(SVerticalBox)
+			+SVerticalBox::Slot()
+			//.AutoHeight()			// Cannot use that otherwise scrollbars disapear.
+			.Padding(0.0f, 0.0f, 0.0f, 0.0f)
+			[
+				SNew(SScrollBox)
+				.Orientation(Orient_Vertical)
+				.ScrollBarAlwaysVisible(false)
+				+ SScrollBox::Slot()
+				[
+					SNew(SVerticalBox)
+					+ SVerticalBox::Slot()
+					.AutoHeight()
+					.Padding(0.0f, 5.0f, 0.0f, 0.0f)
+					[
+						SNew(SWrapBox)
+						.UseAllottedSize(true)
+						+ SWrapBox::Slot()
+						.Padding(0.0f)
+						.HAlign(HAlign_Left)
+						.VAlign(VAlign_Center)
+						[
+							SNew(STextBlock)
+							.ColorAndOpacity(FLinearColor::White)
+							.ShadowColorAndOpacity(FLinearColor::Black)
+							.ShadowOffset(FVector2D::UnitVector)
+							.Text(LOCTEXT("MaterialSimplificationPreview", "Material simplification preview"))
+						]
+					]
+					+SVerticalBox::Slot()
+					.AutoHeight()
+					[
+						SNew(SBorder)
+						//.BorderImage(FAppStyle::GetBrush("DetailsView.CategoryTop"))
+						//.BorderBackgroundColor(FLinearColor(0.9f, 0.6f, 0.6f, 1.0f))
+						.Padding(FMargin(5.0f, 5.0f, 5.0f, 5.0f))
+						[
+							SNew(SHorizontalBox)
+							+SHorizontalBox::Slot()
+							.AutoWidth()
+							.HAlign(HAlign_Left)
+							.VAlign(VAlign_Center)
+							[
+								SNew(SWrapBox)
+								.UseAllottedSize(true)
+								+SWrapBox::Slot()
+								.Padding(20.0, 0.0, 0.0, 0.0) // Padding(float Left, float Top, float Right, float Bottom)
+								.HAlign(HAlign_Left)
+								.VAlign(VAlign_Center)
+								[
+									CheckBoxForceFullSimplification->AsShared()
+								]
+							]
+							+SHorizontalBox::Slot()
+							.AutoWidth()
+							.Padding(0.0f)
+							.HAlign(HAlign_Left)
+							.VAlign(VAlign_Center)
+							[
+								SNew(STextBlock)
+								.ColorAndOpacity(FLinearColor::White)
+								.ShadowColorAndOpacity(FLinearColor::Black)
+								.ShadowOffset(FVector2D::UnitVector)
+								.Text(LOCTEXT("FullsimplificationLabel", "Full simplification"))
+							]
+							+SHorizontalBox::Slot()
+							.AutoWidth()
+							.Padding(20.0, 0.0, 0.0, 0.0)
+							.HAlign(HAlign_Left)
+							.VAlign(VAlign_Center)
+							[
+								SNew(SWrapBox)
+								.UseAllottedSize(true)
+								+SWrapBox::Slot()
+								.Padding(5.0f)
+								.HAlign(HAlign_Left)
+								.VAlign(VAlign_Center)
+								[
+									CheckBoxBytesPerPixelOverride->AsShared()
+								]
+							]
+							+SHorizontalBox::Slot()
+							.AutoWidth()
+							.Padding(0.0f)
+							.HAlign(HAlign_Left)
+							.VAlign(VAlign_Center)
+							[
+								SNew(STextBlock)
+								.ColorAndOpacity(FLinearColor::White)
+								.ShadowColorAndOpacity(FLinearColor::Black)
+								.ShadowOffset(FVector2D::UnitVector)
+								.Text(LOCTEXT("OverrideBytesPerPixel", "Override bytes per pixel"))
+							]
+							+SHorizontalBox::Slot()
+							.MaxWidth(200)
+							.HAlign(HAlign_Left)
+							.VAlign(VAlign_Center)
+							[
+								SNew(SWrapBox)
+								.UseAllottedSize(true)
+								+SWrapBox::Slot()
+								.Padding(0.0, 0.0, 0.0, 0.0)
+								.HAlign(HAlign_Left)
+								.VAlign(VAlign_Center)
+								[
+									BytesPerPixelOverrideInput->AsShared()
+								]
+							]
+							+SHorizontalBox::Slot()
+							.AutoWidth()
+							.Padding(0.0f)
+							.HAlign(HAlign_Left)
+							.VAlign(VAlign_Center)
+							[
+								SNew(SWrapBox)
+								//.UseAllottedSize(true)
+								+ SWrapBox::Slot()
+								.Padding(20.0, 0.0, 0.0, 0.0)
+								.HAlign(HAlign_Left)
+								.VAlign(VAlign_Center)
+								[
+									ButtonApplyToPreview->AsShared()
+								]
+							]
+						] // Border
+					] // VerticalBox
+					
+
+
+					+SVerticalBox::Slot()
+					.AutoHeight()
+					.Padding(0.0f, 15.0f, 0.0f, 0.0f)
+					[
+						SNew(SWrapBox)
+						.UseAllottedSize(true)
+						+ SWrapBox::Slot()
+						.Padding(0.0f)
+						.HAlign(HAlign_Left)
+						.VAlign(VAlign_Center)
+						[
+							SNew(STextBlock)
+							.ColorAndOpacity(FLinearColor::White)
+							.ShadowColorAndOpacity(FLinearColor::Black)
+							.ShadowOffset(FVector2D::UnitVector)
+							.Text(LOCTEXT("MaterialTopologyPreview", "Material topology preview"))
+						]
+					]
+					+SVerticalBox::Slot()
+					.AutoHeight()
+					.Padding(0.0f, 0.0f, 0.0f, 0.0f)
+					[
+						SNew(SBorder)
+						//.BorderImage(FAppStyle::GetBrush("DetailsView.CategoryTop"))
+						//.BorderBackgroundColor(FLinearColor(0.9f, 0.6f, 0.6f, 1.0f))
+						.Padding(FMargin(5.0f, 5.0f, 5.0f, 5.0f))
+						[
+							SNew(SWrapBox)
+							.UseAllottedSize(true)
+							+ SWrapBox::Slot()
+							.Padding(10.0f)
+							.HAlign(HAlign_Center)
+							.VAlign(VAlign_Center)
+							[
+								SAssignNew(MaterialBox, SBox)
+							]
+						]
+					]
+
+
+					
+					+SVerticalBox::Slot()
+					.AutoHeight()
+					.Padding(0.0f, 15.0f, 0.0f, 0.0f)
+					[
+						SNew(SWrapBox)
+						.UseAllottedSize(true)
+						+ SWrapBox::Slot()
+						.Padding(0.0f)
+						.HAlign(HAlign_Left)
+						.VAlign(VAlign_Center)
+						[
+							SNew(STextBlock)
+							.ColorAndOpacity(FLinearColor::White)
+							.ShadowColorAndOpacity(FLinearColor::Black)
+							.ShadowOffset(FVector2D::UnitVector)
+							.Text(LOCTEXT("MaterialAdvancedDetails", "Material advanced details"))
+						]
+					]
+					+SVerticalBox::Slot()
+					.AutoHeight()
+					.Padding(0.0f, 0.0f, 0.0f, 0.0f)
+					[
+						SNew(SBorder)
+						//.BorderImage(FAppStyle::GetBrush("DetailsView.CategoryTop"))
+						//.BorderBackgroundColor(FLinearColor(0.9f, 0.6f, 0.6f, 1.0f))
+						.Padding(FMargin(5.0f, 5.0f, 5.0f, 5.0f))
+						[
+							SNew(SWrapBox)
+							.UseAllottedSize(true)
+							+SWrapBox::Slot()
+							.Padding(5.0f)
+							.HAlign(HAlign_Center)
+							.VAlign(VAlign_Center)
+							[
+								DescriptionTextBlock->AsShared()
+							]
+						]
+					]
+				]
+			]
+		];
+	}
+	else
+	{
+		this->ChildSlot
+		[
+			SNew(SVerticalBox)
+			+SVerticalBox::Slot()
+			.AutoHeight()
+			.Padding(0.0f, 5.0f, 0.0f, 0.0f)
+			[
+				SNew(SWrapBox)
+				.UseAllottedSize(true)
+				+SWrapBox::Slot()
+				.Padding(5.0f)
+				.HAlign(HAlign_Center)
+				.VAlign(VAlign_Center)
+				[
+					SNew(STextBlock)
+					.ColorAndOpacity(FLinearColor::Yellow)
+					.ShadowColorAndOpacity(FLinearColor::Black)
+					.ShadowOffset(FVector2D::UnitVector)
+					.Text(LOCTEXT("SubstrateWidgetNotEnable", "Substrate is not enabled."))
+				]
+			]
+		];
+	}
+}
+
+TSharedRef<SWidget> SMaterialEditorStrataWidget::GetContent()
+{
+	return SharedThis(this);
+}
+
+SMaterialEditorStrataWidget::~SMaterialEditorStrataWidget()
+{
+}
+
+void SMaterialEditorStrataWidget::Tick(const FGeometry& AllottedGeometry, const double InCurrentTime, const float InDeltaTime)
+{
+	if (!bUpdateRequested || !Strata::IsStrataEnabled())
+	{
+		return;
+	}
+	bUpdateRequested = false;
+
+	FText StrataMaterialDescription;
+	if (MaterialEditorPtr.IsValid())
+	{
+		auto MaterialEditor = MaterialEditorPtr.Pin();
+
+		UMaterial* MaterialForStats = MaterialEditor->bStatsFromPreviewMaterial ? MaterialEditor->Material : MaterialEditor->OriginalMaterial;
+
+		StrataMaterialDescription = FText::FromString(FString(TEXT("StrataMaterialDescription")));
+
+		const FMaterialResource* MaterialResource = MaterialForStats->GetMaterialResource(GMaxRHIFeatureLevel);
+		if (MaterialResource)
+		{
+			FString MaterialDescription;
+
+			FMaterialShaderMap* ShaderMap = MaterialResource->GetGameThreadShaderMap();
+			if (ShaderMap)
+			{
+				const FStrataMaterialCompilationOutput& CompilationOutput = ShaderMap->GetStrataMaterialCompilationOutput();
+				const uint32 FinalPixelByteCount = CompilationOutput.StrataUintPerPixel * sizeof(uint32);
+
+				if (CompilationOutput.bMaterialOutOfBudgetHasBeenSimplified)
+				{
+					MaterialDescription += FString::Printf(TEXT("The material was OUT-OF-BUDGET so it has been fully simplified: request bytes = %i / budget = %i\r\n"),
+						CompilationOutput.RequestedBytePixePixel, CompilationOutput.PlatformBytePixePixel);
+					MaterialDescription += FString::Printf(TEXT("Final per pixel byte count   = %i\r\n"),
+						FinalPixelByteCount);
+				}
+				else
+				{
+					MaterialDescription += FString::Printf(TEXT("Material per pixel byte count= %i / budget = %i\r\n"),
+						FinalPixelByteCount, CompilationOutput.PlatformBytePixePixel);
+				}
+				MaterialDescription += FString::Printf(TEXT("BSDF Count	                  = %i\r\n"), CompilationOutput.StrataBSDFCount);
+				MaterialDescription += FString::Printf(TEXT("Local bases Count            = %i\r\n"), CompilationOutput.SharedLocalBasesCount);
+
+				switch (CompilationOutput.StrataMaterialType)
+				{
+				case 0:
+					MaterialDescription += FString::Printf(TEXT("Material complexity          = SIMPLE (diffuse, albedo, roughness)\r\n"));
+					break;
+				case 1:
+					MaterialDescription += FString::Printf(TEXT("Material complexity          = SINGLE (BSDF all features except anisotropy)\r\n"));
+					break;
+				case 2:
+					MaterialDescription += FString::Printf(TEXT("Material complexity          = COMPLEX\r\n"));
+					break;
+				case 3:
+					MaterialDescription += FString::Printf(TEXT("Material complexity          = COMPLEX SPECIAL\r\n"));
+					break;
+				default:
+					MaterialDescription += FString::Printf(TEXT("Material complexity          = UNKOWN => ERROR!\r\n"));
+				}
+
+				MaterialDescription += FString::Printf(TEXT("Root Node Is Thin            = %i\r\n"), CompilationOutput.bIsThin);
+
+				//if (CompilationOutput.StrataBSDFCount == 1)
+				//{
+				//	auto GetSingleOperatorBSDF = [&](const FStrataOperator& Op)
+				//	{
+				//		switch (Op.OperatorType)
+				//		{
+				//		case STRATA_OPERATOR_WEIGHT:
+				//		{
+				//			return ProcessOperator(CompilationOutput.Operators[Op.LeftIndex]);	// Continue walking the tree
+				//			break;
+				//		}
+				//		case STRATA_OPERATOR_VERTICAL:
+				//		case STRATA_OPERATOR_HORIZONTAL:
+				//		case STRATA_OPERATOR_ADD:
+				//		{
+				//			return nullptr;	// A operator with multiple entries must use slabs and so should be a Strata material made of slabs
+				//			break;
+				//		}
+				//		case STRATA_OPERATOR_BSDF_LEGACY:
+				//		case STRATA_OPERATOR_BSDF:
+				//		{
+				//			return Op;
+				//			break;
+				//		}
+				//		}
+				//		return nullptr;
+				//	};
+				//	const FStrataOperator& RootOperator = CompilationOutput.Operators[CompilationOutput.RootOperatorIndex];
+				//	const FStrataOperator* OpBSDF = GetSingleOperatorBSDF(RootOperator);
+				//
+				//	if (OpBSDF)
+				//	{
+				//		// STRATA_TODO gather information about SSSPRofile, subsurface, two sided, etc?
+				//	}
+				//}
+
+				FString MaterialType;
+				switch (CompilationOutput.MaterialType)
+				{
+				case SUBSTRATE_MATERIAL_TYPE_SINGLESLAB:			MaterialType = TEXT("Single Slab"); break;
+				case SUBSTRATE_MATERIAL_TYPE_MULTIPLESLABS:			MaterialType = TEXT("Multiple Slabs"); break;
+				case SUBSTRATE_MATERIAL_TYPE_VOLUMETRICFOGCLOUD:	MaterialType = TEXT("Volumetric (Fog or Cloud)"); break;
+				case SUBSTRATE_MATERIAL_TYPE_UNLIT:					MaterialType = TEXT("Unlit"); break;
+				case SUBSTRATE_MATERIAL_TYPE_HAIR:					MaterialType = TEXT("Hair"); break;
+				case SUBSTRATE_MATERIAL_TYPE_SINGLELAYERWATER:		MaterialType = TEXT("SingleLayerWater"); break;
+				case SUBSTRATE_MATERIAL_TYPE_EYE:					MaterialType = TEXT("Eye"); break;
+				case SUBSTRATE_MATERIAL_TYPE_LIGHTFUNCTION:			MaterialType = TEXT("LightFunction"); break;
+				case SUBSTRATE_MATERIAL_TYPE_POSTPROCESS:			MaterialType = TEXT("PostProcess"); break;
+				case SUBSTRATE_MATERIAL_TYPE_UI:					MaterialType = TEXT("UI"); break;
+				case SUBSTRATE_MATERIAL_TYPE_DECAL:					MaterialType = TEXT("Decal"); break;
+				}
+
+				FString MaterialTypeDetail;
+				if ((CompilationOutput.MaterialType == SUBSTRATE_MATERIAL_TYPE_SINGLESLAB || CompilationOutput.MaterialType == SUBSTRATE_MATERIAL_TYPE_MULTIPLESLABS) && CompilationOutput.bIsThin)
+				{
+					MaterialTypeDetail = TEXT("(Thin Surface / Two-Sided Lighting)");
+				}
+
+				MaterialDescription += FString::Printf(TEXT("Material Type                = %s %s\r\n"), *MaterialType, *MaterialTypeDetail);
+
+				MaterialDescription += FString::Printf(TEXT(" \r\n"));
+				MaterialDescription += FString::Printf(TEXT(" \r\n"));
+				MaterialDescription += FString::Printf(TEXT("================================================================================\r\n"));
+				MaterialDescription += FString::Printf(TEXT("================================Detailed Output=================================\r\n"));
+				MaterialDescription += FString::Printf(TEXT("================================================================================\r\n"));
+				MaterialDescription += CompilationOutput.StrataMaterialDescription;
+
+				// Now generate a visual representation of the material from the topology tree of operators.
+				{
+					if (CompilationOutput.RootOperatorIndex >= 0)
+					{										
+						MaterialBox->SetContent(FSubstrateWidget::ProcessOperator(CompilationOutput));
+					}
+					else
+					{
+						// The tree does not looks sane so generate a visual error without crashing.
+						auto TreeError = SNew(SErrorText)
+							.ErrorText(LOCTEXT("TreeError", "Tree Error"))
+							.BackgroundColor(FSlateColor(EStyleColor::AccentRed));
+
+						const TSharedRef<SWidget>& TreeErrorAsShared = TreeError->AsShared();
+						MaterialBox->SetContent(TreeErrorAsShared);
+					}
+				}
+			}
+			else
+			{
+				MaterialDescription = TEXT("Shader map not found.");
+				MaterialBox->SetContent(SNullWidget::NullWidget);
+			}
+
+			DescriptionTextBlock->SetText(FText::FromString(MaterialDescription));
+		}
+	}
+}
+
+void SMaterialEditorStrataWidget::OnBytesPerPixelChanged(uint32 NewValue)
+{
+	BytesPerPixelOverride = NewValue;
+}
+
+void SMaterialEditorStrataWidget::OnBytesPerPixelCommitted(uint32 NewValue, ETextCommit::Type InCommitType)
+{
+	if (InCommitType == ETextCommit::OnEnter)
+	{
+		BytesPerPixelOverride = NewValue;
+	}
+}
+
+void SMaterialEditorStrataWidget::OnBeginBytesPerPixelSliderMovement()
+{
+	if (bBytesPerPixelStartedTransaction == false)
+	{
+		bBytesPerPixelStartedTransaction = true;
+		GEditor->BeginTransaction(LOCTEXT("PastePoseTransation", "Paste Pose"));
+	}
+}
+void SMaterialEditorStrataWidget::OnEndBytesPerPixelSliderMovement(uint32 NewValue)
+{
+	if (bBytesPerPixelStartedTransaction)
+	{
+		GEditor->EndTransaction();
+		bBytesPerPixelStartedTransaction = false;
+		BytesPerPixelOverride = NewValue;
+	}
+}
+
+TOptional<uint32> SMaterialEditorStrataWidget::GetBytesPerPixelValue() const
+{
+	return BytesPerPixelOverride;
+}
+
+void SMaterialEditorStrataWidget::OnCheckBoxBytesPerPixelChanged(ECheckBoxState InCheckBoxState)
+{
+	BytesPerPixelOverrideInput->SetEnabled(InCheckBoxState == ECheckBoxState::Checked);
+}
+
+FReply SMaterialEditorStrataWidget::OnButtonApplyToPreview()
+{
+	if (MaterialEditorPtr.IsValid())
+	{
+		UMaterialInterface* MaterialInterface = MaterialEditorPtr.Pin()->GetMaterialInterface();
+
+		FStrataCompilationConfig StrataCompilationConfig;
+		StrataCompilationConfig.bFullSimplify = CheckBoxForceFullSimplification->IsChecked();
+
+		// Have a look at MaterialStats.cpp when we want to visualise a specific platform.
+		StrataCompilationConfig.BytesPerPixelOverride = CheckBoxBytesPerPixelOverride->IsChecked() ? FMath::Clamp(BytesPerPixelOverride, 12, Strata::GetBytePerPixel(GMaxRHIShaderPlatform)) : -1;
+
+		MaterialInterface->SetStrataCompilationConfig(StrataCompilationConfig);
+
+		MaterialInterface->ForceRecompileForRendering();
+	}
+
+	return FReply::Handled();
+}
+
+#undef LOCTEXT_NAMESPACE
diff --git a/Engine/Source/Editor/MaterialEditor/Private/SMaterialEditorSubstrateWidget.h b/Engine/Source/Editor/MaterialEditor/Private/SMaterialEditorSubstrateWidget.h
new file mode 100644
--- /dev/null
+++ b/Engine/Source/Editor/MaterialEditor/Private/SMaterialEditorSubstrateWidget.h
@@ -0,0 +1,69 @@
+// Copyright Epic Games, Inc. All Rights Reserved.
+
+#pragma once
+
+#include "CoreMinimal.h"
+#include "Widgets/SWidget.h"
+
+#include "Widgets/DeclarativeSyntaxSupport.h"
+#include "Widgets/SCompoundWidget.h"
+#include "Widgets/SBoxPanel.h"
+#include "Widgets/Input/SButton.h"
+#include "Widgets/Input/SCheckBox.h"
+#include "Widgets/Input/SNumericEntryBox.h"
+#include "Widgets/Text/STextBlock.h"
+#include "Widgets/Layout/SBox.h"
+
+class FMaterialEditor;
+
+class SMaterialEditorStrataWidget : public SCompoundWidget
+{
+public:
+	SLATE_BEGIN_ARGS(SMaterialEditorStrataWidget)
+	{
+	}
+	SLATE_END_ARGS()
+
+	void Construct(const FArguments& InArgs, TWeakPtr<FMaterialEditor> InMaterialEditorPtr);
+
+	void UpdateFromMaterial() { bUpdateRequested = true; }
+
+	/** Gets the widget contents of the app */
+	virtual TSharedRef<SWidget> GetContent();
+
+	virtual ~SMaterialEditorStrataWidget();
+
+	/** SWidget interface */
+	virtual void Tick(const FGeometry& AllottedGeometry, const double InCurrentTime, const float InDeltaTime) override;
+
+private:
+
+	TSharedPtr<class SCheckBox> CheckBoxForceFullSimplification;
+	TSharedPtr<class SCheckBox> CheckBoxBytesPerPixelOverride;
+
+	TSharedPtr<class SNumericEntryBox<uint32>> BytesPerPixelOverrideInput;
+
+	TSharedPtr<class SButton> ButtonApplyToPreview;
+
+	TSharedPtr<class STextBlock> DescriptionTextBlock;
+
+	TSharedPtr<class SBox> MaterialBox;
+
+	FReply OnButtonApplyToPreview();
+
+	/** Pointer back to the material editor that owns this */
+	TWeakPtr<FMaterialEditor> MaterialEditorPtr;
+
+	bool bUpdateRequested = true;
+
+	bool bBytesPerPixelStartedTransaction;
+	uint32 BytesPerPixelOverride;
+	void OnBytesPerPixelChanged(uint32 NewValue);
+	void OnBytesPerPixelCommitted(uint32 NewValue, ETextCommit::Type InCommitType);
+	void OnBeginBytesPerPixelSliderMovement();
+	void OnEndBytesPerPixelSliderMovement(uint32 NewValue);
+	TOptional<uint32> GetBytesPerPixelValue() const;
+
+	void OnCheckBoxBytesPerPixelChanged(ECheckBoxState InCheckBoxState);
+};
+
diff --git a/Engine/Source/Runtime/Engine/Private/StrataVisualizationData.cpp b/Engine/Source/Runtime/Engine/Private/StrataVisualizationData.cpp
deleted file mode 100644
diff --git a/Engine/Source/Runtime/Engine/Private/SubstrateVisualizationData.cpp b/Engine/Source/Runtime/Engine/Private/SubstrateVisualizationData.cpp
new file mode 100644
--- /dev/null
+++ b/Engine/Source/Runtime/Engine/Private/SubstrateVisualizationData.cpp
@@ -0,0 +1,212 @@
+// Copyright Epic Games, Inc. All Rights Reserved.
+
+#include "SubstrateVisualizationData.h"
+#include "HAL/IConsoleManager.h"
+#include "Materials/Material.h"
+#include "Misc/ConfigCacheIni.h"
+#include "RenderUtils.h"
+
+#define LOCTEXT_NAMESPACE "FStrataVisualizationData"
+
+static FStrataVisualizationData GStrataVisualizationData;
+
+static FString ConfigureConsoleCommand(FStrataVisualizationData::TModeMap& ModeMap)
+{
+	FString AvailableVisualizationModes;
+	for (FStrataVisualizationData::TModeMap::TConstIterator It = ModeMap.CreateConstIterator(); It; ++It)
+	{
+		const FStrataVisualizationData::FModeRecord& Record = It.Value();
+		AvailableVisualizationModes += FString(TEXT("\n  Value="));
+		AvailableVisualizationModes += FString::Printf(TEXT("%d: "), uint8(Record.ViewMode));
+		AvailableVisualizationModes += Record.ModeString;
+		if (!Record.bAvailableCommand)
+		{
+			AvailableVisualizationModes += FString::Printf(TEXT(" --- Unavailable, reason: %s"), *Record.UnavailableReason.ToString());
+		}
+	}
+
+	FString Out;
+	Out = TEXT("When the viewport view-mode is set to 'Substrate Visualization', this command specifies which of the various channels to display. Values entered other than the allowed values shown below will be ignored.");
+	Out += AvailableVisualizationModes;
+
+	IConsoleManager::Get().RegisterConsoleVariable(
+		FStrataVisualizationData::GetVisualizeConsoleCommandName(),
+		0,
+		*Out,
+		ECVF_Cheat);
+
+	return Out;
+}
+
+static void AddVisualizationMode(
+	FStrataVisualizationData::TModeMap& ModeMap,
+	const TCHAR* ModeString,
+	const FText& ModeText,
+	const FText& ModeDesc,
+	const FStrataViewMode ViewMode,
+	bool bDefaultComposited,
+	bool bAvailableCommand,
+	const FText& UnavailableReason
+)
+{
+	const FName ModeName = FName(ModeString);
+
+	FStrataVisualizationData::FModeRecord& Record = ModeMap.Emplace(ModeName);
+	Record.ModeString			= FString(ModeString);
+	Record.ModeName				= ModeName;
+	Record.ModeText				= ModeText;
+	Record.ModeDesc				= ModeDesc;
+	Record.ViewMode				= ViewMode;
+	Record.bDefaultComposited	= bDefaultComposited;
+	Record.bAvailableCommand	= bAvailableCommand;
+	Record.UnavailableReason	= UnavailableReason;
+}
+
+void FStrataVisualizationData::Initialize()
+{
+	if (!bIsInitialized && Strata::IsStrataEnabled())
+	{
+		TModeMap AllModeMap;
+
+		AddVisualizationMode(
+			AllModeMap,
+			TEXT("MaterialProperties"),
+			LOCTEXT("MaterialProperties", "Material Properties"),
+			LOCTEXT("MaterialPropertiesDesc", "Visualizes Substrate material properties under mouse cursor"),
+			FStrataViewMode::MaterialProperties,
+			true,
+			true,
+			LOCTEXT("None", "None"));
+
+		AddVisualizationMode(
+			AllModeMap,
+			TEXT("MaterialCount"),
+			LOCTEXT("MaterialCount", "Material Count"),
+			LOCTEXT("MaterialCountDesc", "Visualizes Substrate material count per pixel"),
+			FStrataViewMode::MaterialCount,
+			true,
+			true,
+			LOCTEXT("None", "None"));
+
+		AddVisualizationMode(
+			AllModeMap,
+			TEXT("AdvancedMaterialProperties"),
+			LOCTEXT("AdvancedMaterialProperties", "Advanced Material Properties"),
+			LOCTEXT("AdvancedMaterialPropertiesDesc", "Visualizes Substrate advanced material properties"),
+			FStrataViewMode::AdvancedMaterialProperties,
+			true,
+			Strata::IsAdvancedVisualizationEnabled(),
+			LOCTEXT("IsSubstrateAdvancedDebugShaderEnabled", "Substrate advanced debugging r.Substrate.Debug.AdvancedVisualizationShaders is disabled"));
+
+		AddVisualizationMode(
+			AllModeMap,
+			TEXT("MaterialClassification"),
+			LOCTEXT("MaterialClassification", "Material Classification"),
+			LOCTEXT("MaterialClassificationDesc", "Visualizes Substrate material classification"),
+			FStrataViewMode::MaterialClassification,
+			true,
+			true,
+			LOCTEXT("None", "None"));
+
+		AddVisualizationMode(
+			AllModeMap,
+			TEXT("DecalClassification"),
+			LOCTEXT("DecalClassification", "Decal classification"),
+			LOCTEXT("DecalClassificationDesc", "Visualizes Substrate decal classification"),
+			FStrataViewMode::DecalClassification,
+			true,
+			false, // Disable for now, as it is not important, and is mainly used for debugging
+			LOCTEXT("IsSubstrateBufferPassEnabled", "Substrate tiled DBuffer pass (r.Substrate.DBufferPass and r.Substrate.DBufferPass.DedicatedTiles) is disabled"));
+
+		AddVisualizationMode(
+			AllModeMap,
+			TEXT("OpaqueRoughRefractionClassification"),
+			LOCTEXT("OpaqueRoughRefractionClassification", "Opaque Rough Refraction Classification"),
+			LOCTEXT("OpaqueRoughRefractionClassificationDesc", "Visualizes Substrate Opaque Rough Refraction Classification"),
+			FStrataViewMode::RoughRefractionClassification,
+			true,
+			Strata::IsOpaqueRoughRefractionEnabled(),
+			LOCTEXT("IsSubstrateRoughRefractionEnabled", "Substrate rough refraction r.Substrate.OpaqueMaterialRoughRefraction is disabled"));
+
+		AddVisualizationMode(
+			AllModeMap,
+			TEXT("SubstrateInfo"),
+			LOCTEXT("SubstrateInfo", "Substrate Info"),
+			LOCTEXT("SubstrateInfoDesc", "Visualizes Substrate info"),
+			FStrataViewMode::StrataInfo,
+			true,
+			true,
+			LOCTEXT("None", "None"));
+
+		AddVisualizationMode(
+			AllModeMap,
+			TEXT("MaterialByteCount"),
+			LOCTEXT("MaterialByteCount", "Material Bytes Count"),
+			LOCTEXT("MaterialByteCountDesc", "Visualizes Substrate material footprint per pixel"),
+			FStrataViewMode::MaterialByteCount,
+			true,
+			true,
+			LOCTEXT("None", "None"));
+		
+
+		ConsoleDocumentationVisualizationMode = ConfigureConsoleCommand(AllModeMap);
+
+		// Now only copy the available modes for the menu to not overload it with useless entries.
+		for (auto& Mode : AllModeMap)
+		{
+			if(Mode.Value.bAvailableCommand)
+			{
+				ModeMap.Emplace(Mode.Key) = Mode.Value;
+			}
+		}
+	}
+	bIsInitialized = true;
+}
+
+FText FStrataVisualizationData::GetModeDisplayName(const FName& InModeName) const
+{
+	if (const FModeRecord* Record = ModeMap.Find(InModeName))
+	{
+		return Record->ModeText;
+	}
+	else
+	{
+		return FText::GetEmpty();
+	}
+}
+
+FStrataViewMode FStrataVisualizationData::GetViewMode(const FName& InModeName) const
+{
+	if (const FModeRecord* Record = ModeMap.Find(InModeName))
+	{
+		return Record->ViewMode;
+	}
+	else
+	{
+		return FStrataViewMode::None;
+	}
+}
+
+bool FStrataVisualizationData::GetModeDefaultComposited(const FName& InModeName) const
+{
+	if (const FModeRecord* Record = ModeMap.Find(InModeName))
+	{
+		return Record->bDefaultComposited;
+	}
+	else
+	{
+		return false;
+	}
+}
+
+FStrataVisualizationData& GetStrataVisualizationData()
+{
+	if (!GStrataVisualizationData.IsInitialized())
+	{
+		GStrataVisualizationData.Initialize();
+	}
+
+	return GStrataVisualizationData;
+}
+
+#undef LOCTEXT_NAMESPACE
diff --git a/Engine/Source/Runtime/Engine/Public/MaterialShared.h b/Engine/Source/Runtime/Engine/Public/MaterialShared.h
--- a/Engine/Source/Runtime/Engine/Public/MaterialShared.h
+++ b/Engine/Source/Runtime/Engine/Public/MaterialShared.h
@@ -38,7 +38,7 @@
 #include "ShaderCompilerCore.h"
 #include "PSOPrecache.h"
 #include "UObject/ObjectMacros.h"
-#include "Rendering/StrataMaterialShared.h"
+#include "Rendering/SubstrateMaterialShared.h"
 
 #if UE_ENABLE_INCLUDE_ORDER_DEPRECATED_IN_5_2
 #include "RHI.h"
diff --git a/Engine/Source/Runtime/Engine/Public/Rendering/StrataMaterialShared.cpp b/Engine/Source/Runtime/Engine/Public/Rendering/StrataMaterialShared.cpp
deleted file mode 100644
diff --git a/Engine/Source/Runtime/Engine/Public/Rendering/StrataMaterialShared.h b/Engine/Source/Runtime/Engine/Public/Rendering/StrataMaterialShared.h
deleted file mode 100644
diff --git a/Engine/Source/Runtime/Engine/Public/Rendering/SubstrateMaterialShared.cpp b/Engine/Source/Runtime/Engine/Public/Rendering/SubstrateMaterialShared.cpp
new file mode 100644
--- /dev/null
+++ b/Engine/Source/Runtime/Engine/Public/Rendering/SubstrateMaterialShared.cpp
@@ -0,0 +1,117 @@
+// Copyright Epic Games, Inc. All Rights Reserved.
+
+#include "SubstrateMaterialShared.h"
+
+
+IMPLEMENT_TYPE_LAYOUT(FStrataMaterialCompilationOutput);
+
+FStrataMaterialCompilationOutput::FStrataMaterialCompilationOutput()
+	: StrataMaterialType(0)
+	, StrataBSDFCount(0)
+	, StrataUintPerPixel(0)
+	, bUsesComplexSpecialRenderPath(0)
+#if WITH_EDITOR
+	, StrataMaterialDescription()
+	, SharedLocalBasesCount(0)
+	, RequestedBytePixePixel(0)
+	, PlatformBytePixePixel(0)
+	, bIsThin(0)
+	, MaterialType(0)
+	, bMaterialOutOfBudgetHasBeenSimplified(0)
+	, RootOperatorIndex(0)
+#endif
+{
+#if WITH_EDITOR
+	for (uint32 i = 0; i < STRATA_COMPILATION_OUTPUT_MAX_OPERATOR; ++i)
+	{
+		Operators[i] = FStrataOperator();
+	}
+#endif
+}
+
+
+IMPLEMENT_TYPE_LAYOUT(FStrataOperator);
+
+FStrataOperator::FStrataOperator()
+{
+#if WITH_EDITOR
+	OperatorType = INDEX_NONE;
+	bNodeRequestParameterBlending = false;
+	Index = INDEX_NONE;
+	ParentIndex = INDEX_NONE;
+	LeftIndex = INDEX_NONE;
+	RightIndex = INDEX_NONE;
+	ThicknessIndex = INDEX_NONE;
+
+	BSDFIndex = INDEX_NONE;
+	BSDFType = 0;
+	BSDFRegisteredSharedLocalBasis = FStrataRegisteredSharedLocalBasis();
+	bBSDFHasSSS = false;
+	bBSDFHasMFPPluggedIn = false;
+	bBSDFHasEdgeColor = false;
+	bBSDFHasFuzz = false;
+	bBSDFHasSecondRoughnessOrSimpleClearCoat = false;
+	bBSDFHasAnisotropy = false;
+	bBSDFHasGlint = false;
+	bBSDFHasSpecularProfile = false;
+
+	MaxDistanceFromLeaves = 0;
+	LayerDepth = 0;
+	bIsTop = false;
+	bIsBottom = false;
+	bUseParameterBlending = false;
+	bRootOfParameterBlendingSubTree = false;
+	MaterialExpressionGuid = FGuid();
+#endif
+}
+
+void FStrataOperator::CombineFlagsForParameterBlending(FStrataOperator& A, FStrataOperator& B)
+{
+#if WITH_EDITOR
+	bBSDFHasSSS = A.bBSDFHasSSS || B.bBSDFHasSSS;
+	bBSDFHasMFPPluggedIn = A.bBSDFHasMFPPluggedIn || B.bBSDFHasMFPPluggedIn;
+	bBSDFHasEdgeColor = A.bBSDFHasEdgeColor || B.bBSDFHasEdgeColor;
+	bBSDFHasFuzz = A.bBSDFHasFuzz || B.bBSDFHasFuzz;
+	bBSDFHasSecondRoughnessOrSimpleClearCoat = A.bBSDFHasSecondRoughnessOrSimpleClearCoat || B.bBSDFHasSecondRoughnessOrSimpleClearCoat;
+	bBSDFHasAnisotropy = A.bBSDFHasAnisotropy || B.bBSDFHasAnisotropy;
+	bBSDFHasGlint = A.bBSDFHasGlint || B.bBSDFHasGlint;
+	bBSDFHasSpecularProfile = A.bBSDFHasSpecularProfile || B.bBSDFHasSpecularProfile;
+#endif
+}
+
+void FStrataOperator::CopyFlagsForParameterBlending(FStrataOperator& A)
+{
+#if WITH_EDITOR
+	bBSDFHasSSS = A.bBSDFHasSSS;
+	bBSDFHasMFPPluggedIn = A.bBSDFHasMFPPluggedIn;
+	bBSDFHasEdgeColor = A.bBSDFHasEdgeColor;
+	bBSDFHasFuzz = A.bBSDFHasFuzz;
+	bBSDFHasSecondRoughnessOrSimpleClearCoat = A.bBSDFHasSecondRoughnessOrSimpleClearCoat;
+	bBSDFHasAnisotropy = A.bBSDFHasAnisotropy;
+	bBSDFHasGlint = A.bBSDFHasGlint;
+	bBSDFHasSpecularProfile = A.bBSDFHasSpecularProfile;
+#endif
+}
+
+bool FStrataOperator::IsDiscarded() const
+{
+#if WITH_EDITOR
+	return bUseParameterBlending && !bRootOfParameterBlendingSubTree;
+#else
+	return true;
+#endif
+}
+
+
+IMPLEMENT_TYPE_LAYOUT(FStrataRegisteredSharedLocalBasis);
+
+FStrataRegisteredSharedLocalBasis::FStrataRegisteredSharedLocalBasis()
+{
+#if WITH_EDITOR
+	NormalCodeChunk = INDEX_NONE;
+	TangentCodeChunk = INDEX_NONE;
+	NormalCodeChunkHash = 0;
+	TangentCodeChunkHash = 0;
+	GraphSharedLocalBasisIndex = 0;
+#endif
+}
\ No newline at end of file
diff --git a/Engine/Source/Runtime/Engine/Public/Rendering/SubstrateMaterialShared.h b/Engine/Source/Runtime/Engine/Public/Rendering/SubstrateMaterialShared.h
new file mode 100644
--- /dev/null
+++ b/Engine/Source/Runtime/Engine/Public/Rendering/SubstrateMaterialShared.h
@@ -0,0 +1,143 @@
+// Copyright Epic Games, Inc. All Rights Reserved.
+
+#pragma once
+
+#include "CoreMinimal.h"
+#include "SubstrateDefinitions.h"
+#include "Serialization/MemoryImage.h"
+
+
+// Structures in this files are only used a compilation result return by the compiler.
+// They are also used to present material information in the editor UI.
+
+
+struct FStrataRegisteredSharedLocalBasis
+{
+	DECLARE_TYPE_LAYOUT(FStrataRegisteredSharedLocalBasis, NonVirtual);
+public:
+	FStrataRegisteredSharedLocalBasis();
+
+	LAYOUT_FIELD_EDITORONLY(int32, NormalCodeChunk);
+	LAYOUT_FIELD_EDITORONLY(int32, TangentCodeChunk);
+	LAYOUT_FIELD_EDITORONLY(uint64, NormalCodeChunkHash);
+	LAYOUT_FIELD_EDITORONLY(uint64, TangentCodeChunkHash);
+	LAYOUT_FIELD_EDITORONLY(uint8, GraphSharedLocalBasisIndex);
+};
+
+struct FStrataOperator
+{
+	DECLARE_TYPE_LAYOUT(FStrataOperator, NonVirtual);
+public:
+	FStrataOperator();
+
+	// !!!!!!!!!!
+	// Not using LAYOUT_BITFIELD_EDITORONLY because it seems to cause issue with bit being shifted around when copy happens.
+	// So in the meantime we find it out, LAYOUT_FIELD_EDITORONLY using uint8 is used.
+	// !!!!!!!!!!
+
+	LAYOUT_FIELD_EDITORONLY(int32, OperatorType);
+	LAYOUT_FIELD_EDITORONLY(uint8, bNodeRequestParameterBlending);
+
+	LAYOUT_FIELD_EDITORONLY(int32, Index);			// Index into the array of operators
+	LAYOUT_FIELD_EDITORONLY(int32, ParentIndex);	// Parent operator index
+	LAYOUT_FIELD_EDITORONLY(int32, LeftIndex);		// Left child operator index
+	LAYOUT_FIELD_EDITORONLY(int32, RightIndex);		// Right child operator index
+	LAYOUT_FIELD_EDITORONLY(int32, ThicknessIndex);	// Thickness expression index
+
+	// Data used for BSDF type nodes only
+	LAYOUT_FIELD_EDITORONLY(int32, BSDFIndex);		// Index in the array of BSDF if a BSDF operator
+	LAYOUT_FIELD_EDITORONLY(int32, BSDFType);
+	LAYOUT_FIELD_EDITORONLY(FStrataRegisteredSharedLocalBasis, BSDFRegisteredSharedLocalBasis);
+	LAYOUT_FIELD_EDITORONLY(uint8, bBSDFHasSSS);
+	LAYOUT_FIELD_EDITORONLY(uint8, bBSDFHasMFPPluggedIn);
+	LAYOUT_FIELD_EDITORONLY(uint8, bBSDFHasEdgeColor);
+	LAYOUT_FIELD_EDITORONLY(uint8, bBSDFHasFuzz);
+	LAYOUT_FIELD_EDITORONLY(uint8, bBSDFHasSecondRoughnessOrSimpleClearCoat);
+	LAYOUT_FIELD_EDITORONLY(uint8, bBSDFHasAnisotropy);
+	LAYOUT_FIELD_EDITORONLY(uint8, bBSDFHasGlint);
+	LAYOUT_FIELD_EDITORONLY(uint8, bBSDFHasSpecularProfile);
+
+	// Data derived after the tree has been built.
+	LAYOUT_FIELD_EDITORONLY(int32, MaxDistanceFromLeaves);
+	LAYOUT_FIELD_EDITORONLY(int32, LayerDepth);
+	LAYOUT_FIELD_EDITORONLY(uint8, bIsTop);
+	LAYOUT_FIELD_EDITORONLY(uint8, bIsBottom);
+
+	LAYOUT_FIELD_EDITORONLY(uint8, bUseParameterBlending);			// True when part of a sub tree where parameter blending is in use
+	LAYOUT_FIELD_EDITORONLY(uint8, bRootOfParameterBlendingSubTree);// True when the root of a sub tree where parameter blending is in use. Only this node will register a BSDF
+	LAYOUT_FIELD_EDITORONLY(FGuid, MaterialExpressionGuid);			// Material expression Guid for mapping between UMaterialExpression and FStrataOperator
+
+	void CombineFlagsForParameterBlending(FStrataOperator& A, FStrataOperator& B);
+
+	void CopyFlagsForParameterBlending(FStrataOperator& A);
+
+	bool IsDiscarded() const;
+};
+
+#define STRATA_COMPILATION_OUTPUT_MAX_OPERATOR 24
+
+
+#define SUBSTRATE_MATERIAL_TYPE_SINGLESLAB			0
+#define SUBSTRATE_MATERIAL_TYPE_MULTIPLESLABS		1
+#define SUBSTRATE_MATERIAL_TYPE_VOLUMETRICFOGCLOUD	2
+#define SUBSTRATE_MATERIAL_TYPE_UNLIT				3
+#define SUBSTRATE_MATERIAL_TYPE_HAIR				4
+#define SUBSTRATE_MATERIAL_TYPE_SINGLELAYERWATER	5
+#define SUBSTRATE_MATERIAL_TYPE_EYE					6
+#define SUBSTRATE_MATERIAL_TYPE_LIGHTFUNCTION		7
+#define SUBSTRATE_MATERIAL_TYPE_POSTPROCESS			8
+#define SUBSTRATE_MATERIAL_TYPE_UI					9
+#define SUBSTRATE_MATERIAL_TYPE_DECAL				10
+
+struct FStrataMaterialCompilationOutput
+{
+	DECLARE_TYPE_LAYOUT(FStrataMaterialCompilationOutput, NonVirtual);
+public:
+
+	FStrataMaterialCompilationOutput();
+
+	////
+	//// The following data is required at runtime
+	////
+
+	/** Strata material type, at compile time (0:simple, 1:single, 2: complex) */
+	LAYOUT_FIELD(uint8, StrataMaterialType);
+
+	/** Strata BSDF count, at compile time (0-7) */
+	LAYOUT_FIELD(uint8, StrataBSDFCount);
+
+	/** Strata uint per pixel, at compile time (0-255) */
+	LAYOUT_FIELD(uint8, StrataUintPerPixel);
+
+	/** If Glints or SpecularLUT are used, we need to go the special more expenssive render path */
+	LAYOUT_BITFIELD(uint8, bUsesComplexSpecialRenderPath, 1);
+
+	////
+	//// The following data is only needed when compiling with the editor.
+	////
+
+	// Note we use LAYOUT_FIELD_EDITORONLY for bools because LAYOUT_BITFIELD_EDITORONLY was causing issues when serialising the structure.
+	// STRATA_TODO pack that data.
+
+	/** The Strata verbose description */
+	LAYOUT_FIELD_EDITORONLY(FMemoryImageString, StrataMaterialDescription);
+
+	/** The number of local normal/tangent bases */
+	LAYOUT_FIELD_EDITORONLY(uint8, SharedLocalBasesCount);
+	/** Material requested byte count per pixel */
+	LAYOUT_FIELD_EDITORONLY(uint8, RequestedBytePixePixel);
+	/** The byte count per pixel supported by the platform the material has been compiled against */
+	LAYOUT_FIELD_EDITORONLY(uint8, PlatformBytePixePixel);
+
+	/** Indicate that the material is considered a thin surface instead of a volume filled up with matter */
+	LAYOUT_FIELD_EDITORONLY(uint8, bIsThin);
+	/** Indicate the final material type */
+	LAYOUT_FIELD_EDITORONLY(uint8, MaterialType);
+
+	/** The byte per pixel count supported by the platform the material has been compiled against */
+	LAYOUT_FIELD_EDITORONLY(uint8, bMaterialOutOfBudgetHasBeenSimplified, 1);
+
+	LAYOUT_FIELD_EDITORONLY(uint8, RootOperatorIndex);
+	LAYOUT_ARRAY_EDITORONLY(FStrataOperator, Operators, STRATA_COMPILATION_OUTPUT_MAX_OPERATOR);
+};
+
