Enable mobile platform stats by default

diff --git a/Engine/Source/Editor/MaterialEditor/Private/MaterialStats.cpp b/Engine/Source/Editor/MaterialEditor/Private/MaterialStats.cpp
--- a/Engine/Source/Editor/MaterialEditor/Private/MaterialStats.cpp
+++ b/Engine/Source/Editor/MaterialEditor/Private/MaterialStats.cpp
@@ -477,6 +477,7 @@ void FMaterialStats::BuildShaderPlatformDB()
 	// Vulkan
 	AddShaderPlatform(EPlatformCategoryType::Desktop, SP_VULKAN_SM5, TEXT("Vulkan SM5"), true, true, TEXT("Desktop, Vulkan, Shader Model 5"));
 	AddShaderPlatform(EPlatformCategoryType::Desktop, SP_VULKAN_SM6, TEXT("Vulkan SM6"), true, true, TEXT("Desktop, Vulkan, Shader Model 6"));
+	AddShaderPlatform(EPlatformCategoryType::Desktop, SP_VULKAN_PCES3_1, TEXT("Vulkan Mobile"), true, true, TEXT("Desktop, Vulkan, Mobile"));
 
 	// Android
 	AddShaderPlatform(EPlatformCategoryType::Android, SP_OPENGL_ES3_1_ANDROID, TEXT("Android GLES Mobile"), true, true, TEXT("Android, OpenGLES Mobile"));
diff --git a/Engine/Source/Editor/UnrealEd/Private/PreferenceStubs.cpp b/Engine/Source/Editor/UnrealEd/Private/PreferenceStubs.cpp
--- a/Engine/Source/Editor/UnrealEd/Private/PreferenceStubs.cpp
+++ b/Engine/Source/Editor/UnrealEd/Private/PreferenceStubs.cpp
@@ -75,6 +75,15 @@ UMaterialStatsOptions::UMaterialStatsOptions(const FObjectInitializer& ObjectIni
 		bPlatformUsed[CurrentlyUsedSP] = 1;
 	}
 
+	// enable a mobile platform by default so we can check if shaders are compiling for mobile
+#if PLATFORM_WINDOWS
+	bPlatformUsed[SP_PCD3D_ES3_1] = 1;
+#elif PLATFORM_MAC
+	bPlatformUsed[SP_METAL] = 1;
+#elif PLATFORM_LINUX
+	bPlatformUsed[SP_VULKAN_PCES3_1] = 1;
+#endif
+
 	bMaterialQualityUsed[EMaterialQualityLevel::High] = 1;
 }
 
