Fixed naming of member variable of FSubstrateMaterialCompilationOutput.
Bump shader version required for DDC, similarly to 27789205.

diff --git a/Engine/Source/Editor/MaterialEditor/Private/SMaterialEditorSubstrateWidget.cpp b/Engine/Source/Editor/MaterialEditor/Private/SMaterialEditorSubstrateWidget.cpp
--- a/Engine/Source/Editor/MaterialEditor/Private/SMaterialEditorSubstrateWidget.cpp
+++ b/Engine/Source/Editor/MaterialEditor/Private/SMaterialEditorSubstrateWidget.cpp
@@ -408,7 +408,7 @@ void SMaterialEditorSubstrateWidget::Tick(const FGeometry& AllottedGeometry, con
 				if (CompilationOutput.bMaterialOutOfBudgetHasBeenSimplified)
 				{
 					MaterialDescription += FString::Printf(TEXT("The material was OUT-OF-BUDGET so it has been fully simplified: Request bytes = %i / budget = %i  -  Request Closures = %i / budget = %i\r\n"),
-						CompilationOutput.RequestedBytePixePixel, CompilationOutput.PlatformBytePixePixel,
+						CompilationOutput.RequestedBytePerPixel, CompilationOutput.PlatformBytePerPixel,
 						CompilationOutput.RequestedClosurePerPixel, CompilationOutput.PlatformClosurePixel);
 					MaterialDescription += FString::Printf(TEXT("Final per pixel byte count   = %i\r\n"),
 						FinalPixelByteCount);
@@ -416,7 +416,7 @@ void SMaterialEditorSubstrateWidget::Tick(const FGeometry& AllottedGeometry, con
 				else
 				{
 					MaterialDescription += FString::Printf(TEXT("Material per pixel byte count= %i / budget = %i\r\n"),
-						FinalPixelByteCount, CompilationOutput.PlatformBytePixePixel);
+						FinalPixelByteCount, CompilationOutput.PlatformBytePerPixel);
 				}
 				MaterialDescription += FString::Printf(TEXT("BSDF Count	                  = %i\r\n"), CompilationOutput.SubstrateBSDFCount);
 				MaterialDescription += FString::Printf(TEXT("Local bases Count            = %i\r\n"), CompilationOutput.SharedLocalBasesCount);
diff --git a/Engine/Source/Runtime/Engine/Private/Materials/HLSLMaterialTranslator.cpp b/Engine/Source/Runtime/Engine/Private/Materials/HLSLMaterialTranslator.cpp
--- a/Engine/Source/Runtime/Engine/Private/Materials/HLSLMaterialTranslator.cpp
+++ b/Engine/Source/Runtime/Engine/Private/Materials/HLSLMaterialTranslator.cpp
@@ -12209,8 +12209,8 @@ bool FHLSLMaterialTranslator::FSubstrateCompilationContext::SubstrateGenerateDer
 
 #if WITH_EDITOR
 				Compiler->MaterialCompilationOutput.SubstrateMaterialCompilationOutput.SharedLocalBasesCount = 0; // FinalUsedSharedLocalBasesCount is not valid yet
-				Compiler->MaterialCompilationOutput.SubstrateMaterialCompilationOutput.RequestedBytePixePixel = SubstrateMaterialRequestedSizeByte;
-				Compiler->MaterialCompilationOutput.SubstrateMaterialCompilationOutput.PlatformBytePixePixel = SubstrateBytePerPixel;
+				Compiler->MaterialCompilationOutput.SubstrateMaterialCompilationOutput.RequestedBytePerPixel = SubstrateMaterialRequestedSizeByte;
+				Compiler->MaterialCompilationOutput.SubstrateMaterialCompilationOutput.PlatformBytePerPixel = SubstrateBytePerPixel;
 
 				Compiler->MaterialCompilationOutput.SubstrateMaterialCompilationOutput.RequestedClosurePerPixel = SubstrateSimplificationStatus.OriginalRequestedClosureCount;
 				Compiler->MaterialCompilationOutput.SubstrateMaterialCompilationOutput.PlatformClosurePixel = SubstrateMaterialClosureCount;
diff --git a/Engine/Source/Runtime/Engine/Public/Rendering/SubstrateMaterialShared.cpp b/Engine/Source/Runtime/Engine/Public/Rendering/SubstrateMaterialShared.cpp
--- a/Engine/Source/Runtime/Engine/Public/Rendering/SubstrateMaterialShared.cpp
+++ b/Engine/Source/Runtime/Engine/Public/Rendering/SubstrateMaterialShared.cpp
@@ -13,8 +13,8 @@ FSubstrateMaterialCompilationOutput::FSubstrateMaterialCompilationOutput()
 #if WITH_EDITOR
 	, SubstrateMaterialDescription()
 	, SharedLocalBasesCount(0)
-	, RequestedBytePixePixel(0)
-	, PlatformBytePixePixel(0)
+	, RequestedBytePerPixel(0)
+	, PlatformBytePerPixel(0)
 	, RequestedClosurePerPixel(0)
 	, PlatformClosurePixel(0)
 	, bIsThin(0)
diff --git a/Engine/Source/Runtime/Engine/Public/Rendering/SubstrateMaterialShared.h b/Engine/Source/Runtime/Engine/Public/Rendering/SubstrateMaterialShared.h
--- a/Engine/Source/Runtime/Engine/Public/Rendering/SubstrateMaterialShared.h
+++ b/Engine/Source/Runtime/Engine/Public/Rendering/SubstrateMaterialShared.h
@@ -125,9 +125,9 @@ public:
 	/** The number of local normal/tangent bases */
 	LAYOUT_FIELD_EDITORONLY(uint8, SharedLocalBasesCount);
 	/** Material requested byte count per pixel */
-	LAYOUT_FIELD_EDITORONLY(uint8, RequestedBytePixePixel);
+	LAYOUT_FIELD_EDITORONLY(uint8, RequestedBytePerPixel);
 	/** The byte count per pixel supported by the platform the material has been compiled against */
-	LAYOUT_FIELD_EDITORONLY(uint8, PlatformBytePixePixel);
+	LAYOUT_FIELD_EDITORONLY(uint8, PlatformBytePerPixel);
 
 	/** Material requested BSDF/closure count per pixel */
 	LAYOUT_FIELD_EDITORONLY(uint8, RequestedClosurePerPixel);
