Show material hierarchy in platform stats and warn if derived material instances fail to compile

diff --git a/Engine/Source/Editor/MaterialEditor/MaterialEditor.Build.cs b/Engine/Source/Editor/MaterialEditor/MaterialEditor.Build.cs
--- a/Engine/Source/Editor/MaterialEditor/MaterialEditor.Build.cs
+++ b/Engine/Source/Editor/MaterialEditor/MaterialEditor.Build.cs
@@ -44,6 +44,7 @@ public class MaterialEditor : ModuleRules
                 "AssetRegistry",
 				"ToolMenus",
 				"MainFrame",
+				"Landscape"
 			}
 		);
 
diff --git a/Engine/Source/Editor/MaterialEditor/Private/MaterialEditor.cpp b/Engine/Source/Editor/MaterialEditor/Private/MaterialEditor.cpp
--- a/Engine/Source/Editor/MaterialEditor/Private/MaterialEditor.cpp
+++ b/Engine/Source/Editor/MaterialEditor/Private/MaterialEditor.cpp
@@ -161,6 +161,7 @@
 #include "SMaterialParametersOverviewWidget.h"
 #include "SMaterialEditorCustomPrimitiveDataWidget.h"
 #include "IPropertyRowGenerator.h"
+#include "LandscapeMaterialInstanceConstant.h"
 #include "Widgets/Layout/SScrollBox.h"
 #include "UObject/TextProperty.h"
 #include "Subsystems/AssetEditorSubsystem.h"
@@ -182,6 +183,11 @@ static TAutoConsoleVariable<int32> CVarMaterialEdUseDevShaders(
 	TEXT("Toggles whether the material editor will use shaders that include extra overhead incurred by the editor. Material editor must be re-opened if changed at runtime."),
 	ECVF_RenderThreadSafe);
 
+static TAutoConsoleVariable<int32> CVarMaterialEdMaxDerivedMaterialInstances(
+	TEXT("r.MaterialEditor.MaxDerivedMaterialInstances"),
+	-1,
+	TEXT("Limits amount of derived material instance shown in platform stats. Use negative number to disable the limit. Material editor must be re-opened if changed at runtime."));
+
 ///////////////////////////
 // FMatExpressionPreview //
 ///////////////////////////
@@ -550,8 +556,8 @@ void FMaterialEditor::InitMaterialEditor( const EToolkitMode::Type Mode, const T
 
 	GEditor->RegisterForUndo(this);
 
-	MaterialStatsManager = FMaterialStatsUtils::CreateMaterialStats(this);
-	MaterialStatsManager->SetMaterialDisplayName(OriginalMaterial->GetName());
+	MaterialStatsManager = FMaterialStatsUtils::CreateMaterialStats(this, true);
+	MaterialStatsManager->SetMaterialsDisplayNames({OriginalMaterial->GetName()});
 	MaterialStatsManager->GetOldStatsListing()->OnMessageTokenClicked().AddSP(this, &FMaterialEditor::OnMessageLogLinkActivated);
 
 	if (!Material->MaterialGraph)
@@ -1949,6 +1955,15 @@ void FMaterialEditor::GenerateInheritanceMenu(UToolMenu* Menu)
 	}
 }
 
+void FMaterialEditor::RefreshStatsMaterials()
+{
+	// unconditionally recreate as settings might have changed
+	CreateDerivedMaterialInstancesPreviews();
+
+	MaterialStatsManager->SetMaterial(bStatsFromPreviewMaterial ? Material : OriginalMaterial, bStatsFromPreviewMaterial ? DerivedMaterialInstances : OriginalDerivedMaterialInstances);
+	MaterialStatsManager->SignalMaterialChanged();
+}
+
 void FMaterialEditor::GeneratePreviewMenuContent(UToolMenu* Menu)
 {
 	Menu->bShouldCloseWindowAfterMenuSelection = true;
@@ -2569,17 +2584,26 @@ void FMaterialEditor::UpdatePreviewMaterial( bool bForce )
 		// Null out the expression preview material so they can be GC'ed
 		ExpressionPreviewMaterial = NULL;
 	}
-	MaterialStatsManager->SetMaterial(bStatsFromPreviewMaterial ? Material : OriginalMaterial);
+
+	if (DerivedMaterialInstances.IsEmpty())
+	{
+		CreateDerivedMaterialInstancesPreviews();
+	}
+
+	MaterialStatsManager->SetMaterial(bStatsFromPreviewMaterial ? Material : OriginalMaterial, bStatsFromPreviewMaterial ? DerivedMaterialInstances : OriginalDerivedMaterialInstances);
 	MaterialStatsManager->SignalMaterialChanged();
 
 	// Reregister all components that use the preview material, since UMaterial::PEC does not reregister components using a bIsPreviewMaterial=true material
 	RefreshPreviewViewport();
 }
 
-
-
 bool FMaterialEditor::UpdateOriginalMaterial()
 {
+	if (MaterialStatsManager->GetProvideDerivedMIFlag())
+	{
+		MaterialStatsManager->CacheAndCompilePendingShaders();
+	}
+
 	// If the Material has compilation errors, warn the user
 	for (int32 i = ERHIFeatureLevel::Num - 1; i >= 0; --i)
 	{
@@ -2596,6 +2620,7 @@ bool FMaterialEditor::UpdateOriginalMaterial()
 					FText::Format(NSLOCTEXT("UnrealEd", "Error_CompileErrorsInDefaultMaterial", "The current material has compilation errors for feature level {0}.\nThis material is a Default Material which must be available as a code fallback at all times, compilation errors are not allowed."), FText::FromString(*FeatureLevelName)),
 					NSLOCTEXT("UnrealEd", "Warning_CompileErrorsInDefaultMaterial_Title", "Error: Compilation errors in Default Material"), "Error_CompileErrorsInDefaultMaterial");
 				Info.ConfirmText = NSLOCTEXT("ModalDialogs", "CompileErrorsInDefaultMaterialOk", "Ok");
+				Info.bDontPersistSuppressionAcrossSessions = true;
 
 				FSuppressableWarningDialog CompileErrors(Info);
 				CompileErrors.ShowModal();
@@ -2603,21 +2628,60 @@ bool FMaterialEditor::UpdateOriginalMaterial()
 			}
 			else
 			{
-			FSuppressableWarningDialog::FSetupInfo Info(
-				FText::Format(NSLOCTEXT("UnrealEd", "Warning_CompileErrorsInMaterial", "The current material has compilation errors, so it will not render correctly in feature level {0}.\nAre you sure you wish to continue?"),FText::FromString(*FeatureLevelName)),
-				NSLOCTEXT("UnrealEd", "Warning_CompileErrorsInMaterial_Title", "Warning: Compilation errors in this Material" ), "Warning_CompileErrorsInMaterial");
-			Info.ConfirmText = NSLOCTEXT("ModalDialogs", "CompileErrorsInMaterialConfirm", "Continue");
-			Info.CancelText = NSLOCTEXT("ModalDialogs", "CompileErrorsInMaterialCancel", "Abort");
-
-			FSuppressableWarningDialog CompileErrorsWarning( Info );
-			if( CompileErrorsWarning.ShowModal() == FSuppressableWarningDialog::Cancel )
-			{
+				FSuppressableWarningDialog::FSetupInfo Info(
+					FText::Format(NSLOCTEXT("UnrealEd", "Warning_CompileErrorsInMaterial", "The current material has compilation errors, so it will not render correctly in feature level {0}.\nAre you sure you wish to continue?"),FText::FromString(*FeatureLevelName)),
+					NSLOCTEXT("UnrealEd", "Warning_CompileErrorsInMaterial_Title", "Warning: Compilation errors in this Material" ), "Warning_CompileErrorsInMaterial");
+				Info.ConfirmText = NSLOCTEXT("ModalDialogs", "CompileErrorsInMaterialConfirm", "Continue");
+				Info.CancelText = NSLOCTEXT("ModalDialogs", "CompileErrorsInMaterialCancel", "Abort");
+				Info.bDontPersistSuppressionAcrossSessions = true;
+
+				FSuppressableWarningDialog CompileErrorsWarning( Info );
+				if( CompileErrorsWarning.ShowModal() == FSuppressableWarningDialog::Cancel )
+				{
 					return false;
 				}
 			}
 		}
 	}
 
+	// If derived material instances have compilation errors, warn the user
+	const auto& PlatformList = MaterialStatsManager->GetPlatformsDB();
+	for (const auto& Pair : PlatformList)
+	{
+		const auto& PlatformPtr = Pair.Value;
+		if (PlatformPtr->IsPresentInGrid())
+		{
+			for (int32 QualityLevel = 0; QualityLevel < EMaterialQualityLevel::Num; ++QualityLevel)
+			{
+				const auto& PlatformData = PlatformPtr->GetPlatformData((EMaterialQualityLevel::Type)QualityLevel);
+				// base material is covered by previous check
+				for (int32 InstanceIndex = 1; InstanceIndex < PlatformData.Instances.Num(); ++InstanceIndex)
+				{
+					const FMaterialResource* CurrentResource = PlatformData.Instances[InstanceIndex].MaterialResourcesStats;
+					if (CurrentResource && CurrentResource->GetCompileErrors().Num() > 0)
+					{
+						const auto& PlatformName = MaterialStatsManager->GetPlatformName(Pair.Key);
+						const auto& AssetName = MaterialStatsManager->GetMaterialName(InstanceIndex);
+						const FString QualityName = FMaterialStatsUtils::MaterialQualityToShortString((EMaterialQualityLevel::Type)QualityLevel);
+
+						FSuppressableWarningDialog::FSetupInfo Info(
+						FText::Format(NSLOCTEXT("UnrealEd", "Warning_CompileErrorsInMaterial", "The current material has compilation errors in derived material instance {0} for platform {1} at quality level {2}, so it will not render correctly.\nAre you sure you wish to continue?"),FText::FromString(*AssetName),FText::FromName(PlatformName),FText::FromString(QualityName)),
+						NSLOCTEXT("UnrealEd", "Warning_CompileErrorsInMaterial_Title", "Warning: Compilation errors in this Material" ), "Warning_CompileErrorsInMaterial");
+						Info.ConfirmText = NSLOCTEXT("ModalDialogs", "CompileErrorsInMaterialConfirm", "Continue");
+						Info.CancelText = NSLOCTEXT("ModalDialogs", "CompileErrorsInMaterialCancel", "Abort");
+						Info.bDontPersistSuppressionAcrossSessions = true;
+
+						FSuppressableWarningDialog CompileErrorsWarning( Info );
+						if( CompileErrorsWarning.ShowModal() == FSuppressableWarningDialog::Cancel )
+						{
+							return false;
+						}
+					}
+				}
+			}
+		}
+	}
+
 	// Make sure any graph position changes that might not have been copied are taken into account
 	Material->MaterialGraph->LinkMaterialExpressionsFromGraph();
 
@@ -3202,9 +3266,24 @@ void FMaterialEditor::UpdateMaterialInfoList()
 		}
 	}
 
+	if (DerivedMaterialInstances.IsEmpty())
+	{
+		CreateDerivedMaterialInstancesPreviews();
+	}
+
 	// extract material stats
-	MaterialStatsManager->SetMaterial(MaterialForStats);
+	MaterialStatsManager->SetMaterial(MaterialForStats, bStatsFromPreviewMaterial ? DerivedMaterialInstances : OriginalDerivedMaterialInstances);
 	MaterialStatsManager->Update();
+
+	// check if any derived instances fail, if so show the window
+	if (MaterialStatsManager->AnyNewCompilationErrors(1))
+	{
+		// but then check if base material is compiling, if not, keep the old stats popup not to break UX flow until we merge old and new stats together
+		if (!MaterialStatsManager->AnyNewCompilationErrors(0))
+		{
+			TabManager->TryInvokeTab(MaterialStatsManager->GetGridStatsTabName());
+		}
+	}
 }
 
 void FMaterialEditor::UpdateGraphNodeStates()
@@ -3311,6 +3390,10 @@ void FMaterialEditor::AddReferencedObjects( FReferenceCollector& Collector )
 {
 	Collector.AddReferencedObject(EditorOptions);
 	Collector.AddReferencedObject(Material);
+	for (auto& DerivedMaterialInstance: DerivedMaterialInstances)
+	{
+		Collector.AddReferencedObject(DerivedMaterialInstance);
+	}
 	Collector.AddReferencedObject(OriginalMaterial);
 	Collector.AddReferencedObject(MaterialFunction);
 	Collector.AddReferencedObject(ExpressionPreviewMaterial);
@@ -6109,6 +6192,119 @@ void FMaterialEditor::UpdateSubstrateTopologyPreview()
 	}
 }
 
+void FMaterialEditor::CreateDerivedMaterialInstancesPreviews()
+{
+	DerivedMaterialInstances.Empty();
+	OriginalDerivedMaterialInstances.Empty();
+
+	TArray<FString> DisplayNames;
+	DisplayNames.Push(OriginalMaterial->GetName());
+
+	if (MaterialStatsManager->GetProvideDerivedMIFlag())
+	{
+		const int32 MaxCount = CVarMaterialEdMaxDerivedMaterialInstances.GetValueOnGameThread();
+		// TODO consider a mode where we load all MaterialChildList also
+
+		for (TObjectIterator<UMaterialInstance> It; It; ++It)
+		{
+			UMaterialInstance* Instance = *It;
+			if (!Instance->HasStaticParameters())
+			{
+				continue;
+			}
+
+			FMaterialInheritanceChain Chain;
+			Instance->GetMaterialInheritanceChain(Chain);
+			if (Chain.GetBaseMaterial() != OriginalMaterial)
+			{
+				continue;
+			}
+
+			if (Instance->IsEditorOnly())
+			{
+				continue;
+			}
+
+			bool bSkip = false;
+			for (int32 i = 1; i < Chain.MaterialInstances.Num(); ++i)
+			{
+				auto NestedMaterialInstance = Chain.MaterialInstances[i];
+				if (NestedMaterialInstance->HasStaticParameters())
+				{
+					bSkip = true;
+				}
+			}
+
+			if (bSkip)
+			{
+				UE_LOG(LogMaterialEditor, Display, TEXT("Skipping material instance '%s' with static parameters due to depending on other material instances with static parameters"), *Instance->GetName());
+				continue;
+			}
+
+			const UMaterial* InstanceBaseMaterial = Instance->GetBaseMaterial();
+			const FStaticParameterSet& InstanceStaticParameters = Instance->GetStaticParameters();
+			for(int32_t i = 0; i < OriginalDerivedMaterialInstances.Num(); ++i)
+			{
+				auto ExistingInstance = OriginalDerivedMaterialInstances[i];
+				if (ExistingInstance->GetStaticParameters().Equivalent(InstanceStaticParameters))
+				{
+					bSkip = true;
+				}
+			}
+
+			if (bSkip)
+			{
+				UE_LOG(LogMaterialEditor, Display, TEXT("Skipping material instance '%s' because instance with same static parameters and base material is already present"), *Instance->GetName());
+				continue;
+			}
+
+			const bool bIsLandscapeMaterial = Instance->IsA<ULandscapeMaterialInstanceConstant>();
+
+			FString DisplayName;
+			if (bIsLandscapeMaterial)
+			{
+				const auto LandscapeMaterial = Cast<ULandscapeMaterialInstanceConstant>(Instance);
+				if (LandscapeMaterial->bEditorToolUsage)
+				{
+					UE_LOG(LogMaterialEditor, Display, TEXT("Skipping material instance '%s' because it's used by landscape editor"), *Instance->GetName());
+					continue;
+				}
+
+				// Provide a special name for landscape MIC's to improve UX, so user will be able to tell which combination is not compiling.
+				FString BaseMaterialName = Instance->Parent ? Instance->Parent->GetName() : Instance->GetName();
+				FString LayerNames = FString::JoinBy(LandscapeMaterial->GetEditorOnlyStaticParameters().TerrainLayerWeightParameters,
+					TEXT(","),
+					[](const FStaticTerrainLayerWeightParameter& x) -> FString { return x.LayerName.ToString(); });
+
+				DisplayName = FString::Format(TEXT("{0}({1})"), {*BaseMaterialName, *LayerNames});
+			}
+			else
+			{
+				DisplayName = Instance->GetName();
+			}
+
+			UMaterialInstance* DerivedInstance = Cast<UMaterialInstance>(StaticDuplicateObject(Instance, GetTransientPackage(), NAME_None, ~RF_Standalone, Instance->GetClass()));
+
+			// Beware that this potentially ruins inheritance chain:
+			// - Let's say we have Base material <- Material instance 1 with no static params <- Material instance 2 with static params
+			// - Material instance 1 doesn't influence shader compilation
+			// - So it's safe to change inheritance to Base material <- Material instance 2 with static params
+			DerivedInstance->Parent = Material;
+
+			DerivedMaterialInstances.Add(DerivedInstance);
+			OriginalDerivedMaterialInstances.Add(Instance);
+			DisplayNames.Push(DisplayName);
+
+			if (MaxCount >= 0 && OriginalDerivedMaterialInstances.Num() >= MaxCount)
+			{
+				break;
+			}
+		}
+	}
+
+	MaterialStatsManager->SetMaterialsDisplayNames(DisplayNames);
+}
+
 void FMaterialEditor::UpdateMaterialAfterGraphChange()
 {
 	FlushRenderingCommands();
diff --git a/Engine/Source/Editor/MaterialEditor/Private/MaterialEditor.h b/Engine/Source/Editor/MaterialEditor/Private/MaterialEditor.h
--- a/Engine/Source/Editor/MaterialEditor/Private/MaterialEditor.h
+++ b/Engine/Source/Editor/MaterialEditor/Private/MaterialEditor.h
@@ -415,6 +415,7 @@ public:
 	virtual FMatExpressionPreview* GetExpressionPreview(UMaterialExpression* InExpression) override;
 	virtual void DeleteNodes(const TArray<class UEdGraphNode*>& NodesToDelete) override;
 	virtual void GenerateInheritanceMenu(class UToolMenu* Menu) override;
+	virtual void RefreshStatsMaterials() override;
 
 	void DeleteSelectedNodes(bool bShowConfirmation);
 	void DeleteNodes(const TArray<class UEdGraphNode*>& NodesToDelete, bool bShowConfirmation);
@@ -458,6 +459,9 @@ public:
 
 	/** The material applied to the preview mesh. */
 	TObjectPtr<UMaterial> Material;
+
+	TArray<TObjectPtr<UMaterialInstance>> DerivedMaterialInstances;
+	TArray<TObjectPtr<UMaterialInstance>> OriginalDerivedMaterialInstances;
 	
 	/** The source material being edited by this material editor. Only will be updated when Material's settings are copied over this material */
 	TObjectPtr<UMaterial> OriginalMaterial;
@@ -707,8 +711,10 @@ private:
 	/** Update Substrate topology preview */
 	void UpdateSubstrateTopologyPreview();
 
-public:
+	/** Create array of derived material instances used in conjunction with preview material in material stats */
+	void CreateDerivedMaterialInstancesPreviews();
 
+public:
 private:
 	/**
 	 * Load editor settings from disk (docking state, window pos/size, option state, etc).
diff --git a/Engine/Source/Editor/MaterialEditor/Private/MaterialInstanceEditor.cpp b/Engine/Source/Editor/MaterialEditor/Private/MaterialInstanceEditor.cpp
--- a/Engine/Source/Editor/MaterialEditor/Private/MaterialInstanceEditor.cpp
+++ b/Engine/Source/Editor/MaterialEditor/Private/MaterialInstanceEditor.cpp
@@ -414,8 +414,8 @@ void FMaterialInstanceEditor::InitMaterialInstanceEditor( const EToolkitMode::Ty
 	MaterialEditorInstance->SetSourceInstance(InstanceConstant);
 	MaterialEditorInstance->SetSourceFunction(MaterialFunctionOriginal);
 
-	MaterialStatsManager = FMaterialStatsUtils::CreateMaterialStats(this);
-	MaterialStatsManager->SetMaterialDisplayName(MaterialEditorInstance->SourceInstance->GetName());
+	MaterialStatsManager = FMaterialStatsUtils::CreateMaterialStats(this, false);
+	MaterialStatsManager->SetMaterialsDisplayNames({MaterialEditorInstance->SourceInstance->GetName()});
 
 	// Register our commands. This will only register them if not previously registered
 	FMaterialEditorCommands::Register();
diff --git a/Engine/Source/Editor/MaterialEditor/Private/MaterialStats.cpp b/Engine/Source/Editor/MaterialEditor/Private/MaterialStats.cpp
--- a/Engine/Source/Editor/MaterialEditor/Private/MaterialStats.cpp
+++ b/Engine/Source/Editor/MaterialEditor/Private/MaterialStats.cpp
@@ -13,7 +13,6 @@
 #include "MaterialEditorActions.h"
 #include "Materials/MaterialInstance.h"
 #include "IMaterialEditor.h"
-#include "Preferences/MaterialStatsOptions.h"
 #include "MaterialEditorSettings.h"
 #include "ShaderCompiler.h"
 #include "DataDrivenShaderPlatformInfo.h"
@@ -36,13 +35,12 @@ FShaderPlatformSettings::FShaderPlatformSettings(
 	const bool _bAllowPresenceInGrid,
 	const bool _bAllowCodeView,
 	const FString& _Description)
-	:
-	PlatformType(_PlatformType),
-	PlatformShaderID(_ShaderPlatformID),
-	PlatformName(_Name),
-	PlatformDescription(_Description),
-	bAllowCodeView(_bAllowCodeView),
-	bAllowPresenceInGrid(_bAllowPresenceInGrid)
+	: PlatformType(_PlatformType)
+	, PlatformShaderID(_ShaderPlatformID)
+	, PlatformName(_Name)
+	, PlatformDescription(_Description)
+	, bAllowCodeView(_bAllowCodeView)
+	, bAllowPresenceInGrid(_bAllowPresenceInGrid)
 {
 	PlatformNameID = *FMaterialStatsUtils::ShaderPlatformTypeName(PlatformShaderID);
 }
@@ -54,59 +52,64 @@ void FShaderPlatformSettings::ClearResources()
 	// free material resources
 	for (int32 i = 0; i < EMaterialQualityLevel::Num; ++i)
 	{
-		FMaterialResourceStats* Resource = PlatformData[i].MaterialResourcesStats;
-		if (Resource != nullptr)
+		for (int32 InstanceIndex = 0; InstanceIndex < PlatformData[i].Instances.Num(); ++InstanceIndex)
 		{
-			if (Resource->PrepareDestroy_GameThread())
+			auto& Instance = PlatformData[i].Instances[InstanceIndex];
+			FMaterialResourceStats* Resource = Instance.MaterialResourcesStats;
+			if (Resource != nullptr)
 			{
-				MaterialsToDeleteOnRenderThread.Add(Resource);
-			}
-			else
-			{
-				delete Resource;
+				if (Resource->PrepareDestroy_GameThread())
+				{
+					MaterialsToDeleteOnRenderThread.Add(Resource);
+				}
+				else
+				{
+					delete Resource;
+				}
 			}
-			PlatformData[i].MaterialResourcesStats = nullptr;
-		}
 
-		PlatformData[i].ArrShaderEntries.Empty();
+			Instance.ArrShaderEntries.Empty();
+		}
 
-		PlatformData[i].bCompilingShaders = false;
-		PlatformData[i].bNeedShaderRecompilation = true;
+		PlatformData[i].Instances.Empty();
 	}
 
 	FMaterial::DeleteMaterialsOnRenderThread(MaterialsToDeleteOnRenderThread);
 }
 
-FText FShaderPlatformSettings::GetSelectedShaderViewComboText(EMaterialQualityLevel::Type QualityLevel) const
+FText FShaderPlatformSettings::GetSelectedShaderViewComboText(EMaterialQualityLevel::Type QualityLevel, const int32 InstanceIndex) const
 {
-	if (PlatformData[QualityLevel].ArrShaderEntries.Num() == 0)
+	const auto& Instance = PlatformData[QualityLevel].Instances[InstanceIndex];
+	if (Instance.ArrShaderEntries.Num() == 0)
 	{
 		return FText::FromString(TEXT("-Compiling-Shaders-"));
 	}
 
-	return FText::FromString(PlatformData[QualityLevel].ComboBoxSelectedEntry.Text);
+	return FText::FromString(Instance.ComboBoxSelectedEntry.Text);
 }
 
-void FShaderPlatformSettings::OnShaderViewComboSelectionChanged(TSharedPtr<FMaterialShaderEntry> Item, EMaterialQualityLevel::Type QualityType)
+void FShaderPlatformSettings::OnShaderViewComboSelectionChanged(TSharedPtr<FMaterialShaderEntry> Item, EMaterialQualityLevel::Type QualityType, const int32 InstanceIndex)
 {
 	if (Item.IsValid())
 	{
-		PlatformData[QualityType].ComboBoxSelectedEntry = *Item.Get();
-		PlatformData[QualityType].bUpdateShaderCode = true;
+		auto& Instance = PlatformData[QualityType].Instances[InstanceIndex];
+		Instance.ComboBoxSelectedEntry = *Item.Get();
+		Instance.bUpdateShaderCode = true;
 	}
 }
 
-FText FShaderPlatformSettings::GetShaderCode(const EMaterialQualityLevel::Type QualityType)
+FText FShaderPlatformSettings::GetShaderCode(const EMaterialQualityLevel::Type QualityType, const int32 InstanceIndex)
 {
+	auto& Instance = PlatformData[QualityType].Instances[InstanceIndex];
 	// if there were no change to the material return the cached shader code
-	if (!PlatformData[QualityType].bUpdateShaderCode)
+	if (!Instance.bUpdateShaderCode)
 	{
-		return PlatformData[QualityType].ShaderCode;
+		return Instance.ShaderCode;
 	}
 
-	PlatformData[QualityType].ShaderCode = LOCTEXT("ShaderCodeMsg", "Shader code compiling or not available!");
+	Instance.ShaderCode = LOCTEXT("ShaderCodeMsg", "Shader code compiling or not available!");
 
-	FMaterialResource *Resource = PlatformData[QualityType].MaterialResourcesStats;
+	FMaterialResource *Resource = Instance.MaterialResourcesStats;
 	const FMaterialShaderMap* MaterialShaderMap = Resource->GetGameThreadShaderMap();
 
 	const bool bCompilationFinished = Resource->IsCompilationFinished() && (MaterialShaderMap != nullptr);
@@ -117,7 +120,7 @@ FText FShaderPlatformSettings::GetShaderCode(const EMaterialQualityLevel::Type Q
 		TMap<FShaderId, TShaderRef<FShader>> ShaderMap;
 		MaterialShaderMap->GetShaderList(ShaderMap);
 
-		const FShaderId& ShaderId = PlatformData[QualityType].ComboBoxSelectedEntry.ShaderId;
+		const FShaderId& ShaderId = Instance.ComboBoxSelectedEntry.ShaderId;
 
 		const auto Entry = ShaderMap.Find(ShaderId);
 		if (Entry != nullptr)
@@ -126,13 +129,13 @@ FText FShaderPlatformSettings::GetShaderCode(const EMaterialQualityLevel::Type Q
 			const FMemoryImageString* ShaderSource = MaterialShaderMap->GetShaderSource(Shader.GetVertexFactoryType(), Shader.GetType(), ShaderId.PermutationId);
 			if (ShaderSource != nullptr)
 			{
-				PlatformData[QualityType].bUpdateShaderCode = false;
-				PlatformData[QualityType].ShaderCode = FText::FromString(*ShaderSource);
+				Instance.bUpdateShaderCode = false;
+				Instance.ShaderCode = FText::FromString(*ShaderSource);
 			}
 		}
 	}
 
-	return PlatformData[QualityType].ShaderCode;
+	return Instance.ShaderCode;
 }
 
 void FShaderPlatformSettings::AllocateMaterialResources()
@@ -143,157 +146,226 @@ void FShaderPlatformSettings::AllocateMaterialResources()
 
 	for (int32 QualityLevelIndex = 0; QualityLevelIndex < EMaterialQualityLevel::Num; QualityLevelIndex++)
 	{
-		PlatformData[QualityLevelIndex].MaterialResourcesStats = new FMaterialResourceStats();
-		PlatformData[QualityLevelIndex].MaterialResourcesStats->SetMaterial(Material, MaterialInstance, (ERHIFeatureLevel::Type)TargetFeatureLevel, (EMaterialQualityLevel::Type)QualityLevelIndex);
+		auto& Data = PlatformData[QualityLevelIndex].Instances.AddDefaulted_GetRef();
+
+		Data.MaterialResourcesStats = new FMaterialResourceStats();
+		Data.MaterialResourcesStats->SetMaterial(Material, MaterialInstance, TargetFeatureLevel, (EMaterialQualityLevel::Type)QualityLevelIndex);
+
+		Data.ShaderStatsInfo.Reset();
+
+		Data.bNeedShaderRecompilation = true;
+
+		for (int32 i = 0; i < DerivedMaterialInstances.Num(); ++i)
+		{
+			auto& Data2 = PlatformData[QualityLevelIndex].Instances.AddDefaulted_GetRef();
+
+			Data2.MaterialResourcesStats = new FMaterialResourceStats();
+			Data2.MaterialResourcesStats->SetMaterial(DerivedMaterialInstances[i]->GetBaseMaterial(), DerivedMaterialInstances[i], TargetFeatureLevel, (EMaterialQualityLevel::Type)QualityLevelIndex);
+
+			Data2.ShaderStatsInfo.Reset();
+
+			Data2.bNeedShaderRecompilation = true;
+		}
+
+		// ensure we start compiling shaders soon due to bNeedShaderRecompilation being set above
+		PlatformData[QualityLevelIndex].LastTimeCompilationRequested = 0.0;
 	}
 }
 
-void FShaderPlatformSettings::SetMaterial(UMaterial* InMaterial)
+void FShaderPlatformSettings::SetMaterial(UMaterial *InBaseMaterial, UMaterialInstance *InBaseMaterialInstance, const TArray<TObjectPtr<UMaterialInstance>>& InDerivedMaterialInstances)
 {
-	// if this is a different material, clear away the old one's resources and compile new shaders
-	if (Material != InMaterial)
+	bool bReallocate = false;
+
+	if (InBaseMaterial != nullptr && InBaseMaterialInstance == nullptr && Material != InBaseMaterial)
 	{
-		Material = InMaterial;
+		Material = InBaseMaterial;
 		MaterialInstance = nullptr;
+		bReallocate = true;
+	}
 
-		AllocateMaterialResources();
+	if (InBaseMaterial == nullptr && InBaseMaterialInstance != nullptr && MaterialInstance != InBaseMaterialInstance)
+	{
+		Material = InBaseMaterialInstance->GetMaterial();
+		MaterialInstance = InBaseMaterialInstance;
+		bReallocate = true;
 	}
-}
 
-void FShaderPlatformSettings::SetMaterial(UMaterialInstance* InMaterialInstance)
-{
-	if (MaterialInstance != InMaterialInstance)
+	if (DerivedMaterialInstances != InDerivedMaterialInstances)
 	{
-		MaterialInstance = InMaterialInstance;
-		Material = MaterialInstance->GetMaterial();
+		// TODO avoid copy
+		DerivedMaterialInstances = InDerivedMaterialInstances;
+		bReallocate = true;
+	}
 
+	if (bReallocate)
+	{
 		AllocateMaterialResources();
 	}
 }
 
-bool FShaderPlatformSettings::CheckShaders()
+bool FShaderPlatformSettings::CheckShaders(bool bIgnoreCooldown)
 {
-	bool bRetValue = false;
+	if (Material == nullptr)
+	{
+		return false;
+	}
+
+	bool bRefreshStatUI = false;
 
-	// prevent stats 
-	const double kMinimumTimeBetweenCompilationsSeconds = 5.0;
+	// don't refresh stats too often
+	const double kMinimumTimeBetweenCompilationsSeconds = 1.5;
 
-	if (Material != nullptr)
+	const double CurrentTime = FPlatformTime::Seconds();
+
+	// check and triggers shader recompilation if needed
+	for (int32 QualityLevelIndex = 0; QualityLevelIndex < EMaterialQualityLevel::Num; QualityLevelIndex++)
 	{
-		// check and triggers shader recompilation if needed
-		for (int32 QualityLevelIndex = 0; QualityLevelIndex < EMaterialQualityLevel::Num; QualityLevelIndex++)
+		auto &Data = PlatformData[QualityLevelIndex];
+
+		const bool bNeedsShaders = (bPresentInGrid && Data.bExtractStats) || Data.bExtractCode;
+		const bool bCooledDown = bIgnoreCooldown || (CurrentTime - Data.LastTimeCompilationRequested) > kMinimumTimeBetweenCompilationsSeconds;
+
+		bool bTriggeredCompilation = false;
+
+		for (int32 InstanceIndex = 0; InstanceIndex < Data.Instances.Num(); ++InstanceIndex)
 		{
-			auto &Data = PlatformData[QualityLevelIndex];
-			const bool bNeedsShaders = (bPresentInGrid && Data.bExtractStats) || Data.bExtractCode;
-			const double CurrentTime = FPlatformTime::Seconds();
-			if (Data.bNeedShaderRecompilation && bNeedsShaders && (CurrentTime - Data.LastTimeCompilationRequested) > kMinimumTimeBetweenCompilationsSeconds)
+			auto& Instance = Data.Instances[InstanceIndex]; 
+
+			if (Instance.bNeedShaderRecompilation && bNeedsShaders)
 			{
-				Data.MaterialResourcesStats->CancelCompilation();
+				bRefreshStatUI = true;
 
-				Material->UpdateCachedExpressionData();
+				// reset even if we don't immediately recompile to improve UI feedback time
+				Instance.ShaderStatsInfo.Reset();
 
-				if (MaterialInstance != nullptr)
+				if (bCooledDown)
 				{
-					MaterialInstance->UpdateCachedData();
-				}
+					Instance.MaterialResourcesStats->CancelCompilation();
+
+					Material->UpdateCachedExpressionData();
 
-				Data.LastTimeCompilationRequested = CurrentTime;
+					if (MaterialInstance != nullptr)
+					{
+						MaterialInstance->UpdateCachedData();
+					}
 
-				TMap<FName, TArray<FMaterialStatsUtils::FRepresentativeShaderInfo>> ShaderTypeNamesAndDescriptions;
-				FMaterialStatsUtils::GetRepresentativeShaderTypesAndDescriptions(ShaderTypeNamesAndDescriptions, Data.MaterialResourcesStats);
+					TMap<FName, TArray<FMaterialStatsUtils::FRepresentativeShaderInfo>> ShaderTypeNamesAndDescriptions;
+					FMaterialStatsUtils::GetRepresentativeShaderTypesAndDescriptions(ShaderTypeNamesAndDescriptions, Instance.MaterialResourcesStats);
 
-				TArray<const FVertexFactoryType*> VFTypes;
-				TArray<const FShaderPipelineType*> PipelineTypes;
-				TArray<const FShaderType*> ShaderTypes;
-				for (auto& DescriptionPair : ShaderTypeNamesAndDescriptions)
-				{
-					const FVertexFactoryType* VFType = FindVertexFactoryType(DescriptionPair.Key);
-					check(VFType);
+					TArray<const FVertexFactoryType*> VFTypes;
+					TArray<const FShaderPipelineType*> PipelineTypes;
+					TArray<const FShaderType*> ShaderTypes;
 
-					auto& DescriptionArray = DescriptionPair.Value;
-					for (const FMaterialStatsUtils::FRepresentativeShaderInfo& ShaderInfo : DescriptionArray)
+					for (auto& DescriptionPair : ShaderTypeNamesAndDescriptions)
 					{
-						const FShaderType* ShaderType = FindShaderTypeByName(ShaderInfo.ShaderName);
-						if (ShaderType && VFType)
+						const FVertexFactoryType* VFType = FindVertexFactoryType(DescriptionPair.Key);
+						check(VFType);
+
+						auto& DescriptionArray = DescriptionPair.Value;
+						for (const FMaterialStatsUtils::FRepresentativeShaderInfo& ShaderInfo : DescriptionArray)
 						{
-							VFTypes.Add(VFType);
-							ShaderTypes.Add(ShaderType);
-							PipelineTypes.Add(nullptr);
+							const FShaderType* ShaderType = FindShaderTypeByName(ShaderInfo.ShaderName);
+
+							// in compile only setting we only care if derived MI's are compiling at all
+							// so we only take one shader combination not to slow down material stats window too much
+							// it's likely sufficient to treat first shader as most complex, otherwise we need to refactor FindShaderTypeByName
+							if (Instance.bOnlyCompileMostComplexShader && VFTypes.Num() >= 1)
+							{
+								break;
+							}
+
+							if (ShaderType && VFType)
+							{
+								VFTypes.Add(VFType);
+								ShaderTypes.Add(ShaderType);
+								PipelineTypes.Add(nullptr);
+							}
 						}
 					}
-				}
-
-				// Prepare the resource for compilation, but don't compile the completed shader map.
-				const bool bSuccess = Data.MaterialResourcesStats->CacheShaders(PlatformShaderID, EMaterialShaderPrecompileMode::None);
 
-				if (bSuccess)
-				{
-					// Compile just the types we want.
-					Data.MaterialResourcesStats->CacheGivenTypes(PlatformShaderID, VFTypes, PipelineTypes, ShaderTypes);
-				}
+					// Prepare the resource for compilation, but don't compile the completed shader map.
+					const bool bSuccess = Instance.MaterialResourcesStats->CacheShaders(PlatformShaderID, EMaterialShaderPrecompileMode::None);
 
-				Data.bCompilingShaders = true;
-				Data.bUpdateShaderCode = true;
-				Data.bNeedShaderRecompilation = false;
+					if (bSuccess)
+					{
+						// Compile just the types we want.
+						Instance.MaterialResourcesStats->CacheGivenTypes(PlatformShaderID, VFTypes, PipelineTypes, ShaderTypes);
+					}
 
-				Data.ShaderStatsInfo.Reset();
+					Instance.bCompilingShaders = true;
+					Instance.bUpdateShaderCode = true;
+					Instance.bNeedShaderRecompilation = false;
 
-				bRetValue = true;
+					bTriggeredCompilation = true;
+				}
 			}
 		}
+
+		if (bTriggeredCompilation)
+		{
+			Data.LastTimeCompilationRequested = CurrentTime;
+		}
 	}
 
-	return bRetValue;
+	return bRefreshStatUI;
 }
 
 bool FShaderPlatformSettings::Update()
 {
-	bool bRetValue = CheckShaders();
+	bool bRetValue = CheckShaders(false);
 
 	// if a shader compilation has been requested check if completed and extract shader names needed by code viewer combo-box
-	for (int32 i = 0; i < EMaterialQualityLevel::Num; ++i)
+	for (int32 QualityLevelIndex = 0; QualityLevelIndex < EMaterialQualityLevel::Num; ++QualityLevelIndex)
 	{
-		auto& QualityItem = PlatformData[i];
-		if (QualityItem.bCompilingShaders)
+		auto &Data = PlatformData[QualityLevelIndex];
+
+		for (int32 InstanceIndex = 0; InstanceIndex < Data.Instances.Num(); ++InstanceIndex)
 		{
-			FMaterialResource* Resource = QualityItem.MaterialResourcesStats;
+			auto& Instance = Data.Instances[InstanceIndex]; 
 
-			// if compilation is complete extract the list of compiled shader names
-			const bool bCompilationFinished = Resource->IsCompilationFinished();
-			if (bCompilationFinished)
+			if (Instance.bCompilingShaders)
 			{
-				QualityItem.bCompilingShaders = false;
-				QualityItem.bUpdateShaderCode = true;
+				FMaterialResource* Resource = Instance.MaterialResourcesStats;
 
-				const FMaterialShaderMap* MaterialShaderMap = Resource->GetGameThreadShaderMap();
-				if (MaterialShaderMap != nullptr)
+				// if compilation is complete extract the list of compiled shader names
+				const bool bCompilationFinished = Resource->IsCompilationFinished();
+				if (bCompilationFinished)
 				{
-					TMap<FShaderId, TShaderRef<FShader>> ShaderMap;
-					MaterialShaderMap->GetShaderList(ShaderMap);
+					Instance.bCompilingShaders = false;
+					Instance.bUpdateShaderCode = true;
 
-					QualityItem.ArrShaderEntries.Empty();
-					for (const auto& Entry : ShaderMap)
+					const FMaterialShaderMap* MaterialShaderMap = Resource->GetGameThreadShaderMap();
+					if (MaterialShaderMap != nullptr)
 					{
-						FShaderType* EntryShader = Entry.Value.GetType();
-						FVertexFactoryType* VertexFactory = Entry.Value.GetVertexFactoryType();
-						FString ShaderName = FString::Printf(TEXT("%s/%s"), VertexFactory ? VertexFactory->GetName() : TEXT("NullVF"), EntryShader->GetName());
-						QualityItem.ArrShaderEntries.Add(MakeShareable(new FMaterialShaderEntry{ Entry.Key, ShaderName }));
-					}
+						TMap<FShaderId, TShaderRef<FShader>> ShaderMap;
+						MaterialShaderMap->GetShaderList(ShaderMap);
 
-					if (QualityItem.ArrShaderEntries.Num() > 0)
-					{
-						QualityItem.ArrShaderEntries.Sort([](const TSharedPtr<FMaterialShaderEntry>& First, const TSharedPtr<FMaterialShaderEntry>& Second)
-							{
-								return First->Text < Second->Text;
-							});
+						Instance.ArrShaderEntries.Empty();
+						for (const auto& Entry : ShaderMap)
+						{
+							FShaderType* EntryShader = Entry.Value.GetType();
+							FVertexFactoryType* VertexFactory = Entry.Value.GetVertexFactoryType();
+							FString ShaderName = FString::Printf(TEXT("%s/%s"), VertexFactory ? VertexFactory->GetName() : TEXT("NullVF"), EntryShader->GetName());
+							Instance.ArrShaderEntries.Add(MakeShareable(new FMaterialShaderEntry{ Entry.Key, ShaderName }));
+						}
 
-						QualityItem.ComboBoxSelectedEntry = *QualityItem.ArrShaderEntries[0];
+						if (Instance.ArrShaderEntries.Num() > 0)
+						{
+							Instance.ArrShaderEntries.Sort([](const TSharedPtr<FMaterialShaderEntry>& First, const TSharedPtr<FMaterialShaderEntry>& Second)
+								{
+									return First->Text < Second->Text;
+								});
+
+							Instance.ComboBoxSelectedEntry = *Instance.ArrShaderEntries[0];
+						}
 					}
-				}
 
-				FMaterialStatsUtils::ExtractMatertialStatsInfo(PlatformShaderID, QualityItem.ShaderStatsInfo, Resource);
+					FMaterialStatsUtils::ExtractMatertialStatsInfo(PlatformShaderID, Instance.ShaderStatsInfo, Resource);
+
+					Instance.bNeedToWarnAboutCompilationErrors = Resource->GetCompileErrors().Num() > 0;
 
-				bRetValue = true;
+					bRetValue = true;
+				}
 			}
 		}
 	}
@@ -301,6 +373,35 @@ bool FShaderPlatformSettings::Update()
 	return bRetValue;
 }
 
+bool FShaderPlatformSettings::CachePendingShaders()
+{
+	const bool bNeedsGridRefresh = CheckShaders(true);
+
+	FString MaterialName;
+	TArray<FMaterial*> MaterialsToCompile;
+	for (int32 QualityLevelIndex = 0; QualityLevelIndex < EMaterialQualityLevel::Num; QualityLevelIndex++)
+	{
+		for (int32 InstanceIndex = 0; InstanceIndex < PlatformData[QualityLevelIndex].Instances.Num(); ++InstanceIndex)
+		{
+			auto& Instance = PlatformData[QualityLevelIndex].Instances[InstanceIndex]; 
+
+			if (Instance.bCompilingShaders)
+			{
+				if (MaterialName.IsEmpty())
+				{
+					MaterialName = Instance.MaterialResourcesStats->GetFriendlyName(); 
+				}
+
+				MaterialsToCompile.Add(Instance.MaterialResourcesStats);
+			}
+		}
+	}
+
+	FMaterial::FinishCompilation(*MaterialName, MaterialsToCompile);
+
+	return bNeedsGridRefresh;
+}
+
 /*end FShaderPlatformSettings functions*/
 /***********************************************************************************************************************/
 
@@ -312,7 +413,7 @@ FMaterialStats::~FMaterialStats()
 	SaveSettings();
 }
 
-void FMaterialStats::Initialize(IMaterialEditor* InMaterialEditor)
+void FMaterialStats::Initialize(IMaterialEditor* InMaterialEditor, const bool ShowMaterialInstancesMenu)
 {
 	MaterialEditor = InMaterialEditor;
 
@@ -325,7 +426,8 @@ void FMaterialStats::Initialize(IMaterialEditor* InMaterialEditor)
 	StatsGrid->BuildGrid();
 
 	GridStatsWidget = SNew(SMaterialEditorStatsWidget)
-		.MaterialStatsWPtr(SharedThis(this));
+		.MaterialStatsWPtr(SharedThis(this))
+		.ShowMaterialInstancesMenu(ShowMaterialInstancesMenu);
 
 	FMessageLogModule& MessageLogModule = FModuleManager::LoadModuleChecked<FMessageLogModule>("MessageLog");
 	FMessageLogInitializationOptions LogOptions;
@@ -380,6 +482,8 @@ void FMaterialStats::LoadSettings()
 			}
 		}
 	}
+
+	MaterialStatsDerivedMIOption = Options->MaterialStatsDerivedMIOption;
 }
 
 void FMaterialStats::SaveSettings()
@@ -400,6 +504,8 @@ void FMaterialStats::SaveSettings()
 		Options->bMaterialQualityUsed[i] = bQualityPresent ? 1 : 0;
 	}
 
+	Options->MaterialStatsDerivedMIOption = MaterialStatsDerivedMIOption;
+
 	Options->SaveConfig();
 }
 
@@ -542,7 +648,7 @@ void FMaterialStats::SignalMaterialChanged()
 	for (auto Entry : ShaderPlatformStatsDB)
 	{
 		for (int32 i = 0; i < EMaterialQualityLevel::Num; ++i)
-			Entry.Value->SetNeedShaderCompilation((EMaterialQualityLevel::Type)i, true);
+			Entry.Value->SetNeedShaderCompilation((EMaterialQualityLevel::Type)i, true, MaterialStatsDerivedMIOption == EMaterialStatsDerivedMIOption::CompileOnly);
 	}
 }
 
@@ -590,6 +696,18 @@ bool FMaterialStats::SwitchStatsQualityFlag(EMaterialQualityLevel::Type Quality)
 	return bValue;
 }
 
+void FMaterialStats::SetMaterialStatsDerivedMIOption(const EMaterialStatsDerivedMIOption value)
+{
+	MaterialStatsDerivedMIOption = value;
+
+	if (MaterialEditor != nullptr)
+	{
+		MaterialEditor->RefreshStatsMaterials();
+	}
+
+	SaveSettings();
+}
+
 void FMaterialStats::SetShaderPlatformUseCodeView(const EShaderPlatform PlatformID, const EMaterialQualityLevel::Type Quality, const bool bValue)
 {
 	auto* Item = ShaderPlatformStatsDB.Find(PlatformID);
@@ -643,7 +761,7 @@ TSharedPtr<FShaderPlatformSettings> FMaterialStats::GetPlatformSettings(const FN
 	return GetPlatformSettings(PlatformID);
 }
 
-FText FMaterialStats::GetShaderCode(const EShaderPlatform PlatformID, const EMaterialQualityLevel::Type QualityType)
+FText FMaterialStats::GetShaderCode(const EShaderPlatform PlatformID, const EMaterialQualityLevel::Type QualityType, const int32 InstanceIndex)
 {
 	auto* Entry = ShaderPlatformStatsDB.Find(PlatformID);
 	if (Entry == nullptr)
@@ -651,31 +769,133 @@ FText FMaterialStats::GetShaderCode(const EShaderPlatform PlatformID, const EMat
 		return FText::FromString(TEXT("Shader code compiling or not available!"));
 	}
 
-	return (*Entry)->GetShaderCode(QualityType);
+	return (*Entry)->GetShaderCode(QualityType, InstanceIndex);
 }
 
 void FMaterialStats::Update()
 {
 	const bool bNeedsUpdate = IsShowingStats() || IsCodeViewWindowActive();
+
+	// don't update materials if we don't need to verify derived MI's
+	// this is to preserve old behaviour of speeding up material editor by hiding platform stats
+	if (!bNeedsUpdate && !GetProvideDerivedMIFlag())
+	{
+		return;
+	}
+	
+	// otherwise always update material compilation info so we can check if shaders are compiling at material editor save
+	for (const auto& Entry : ShaderPlatformStatsDB)
+	{
+		auto PlatformStats = Entry.Value;
+		bNeedsGridRefresh |= PlatformStats->Update();
+	}
+
 	if (!bNeedsUpdate)
 	{
 		return;
 	}
 
-	bool bInfoChanged = false;
+	if (bNeedsGridRefresh)
+	{
+		// in compile only mode the amount of columns might change if some of MI's fail to compile 
+		if (GetMaterialStatsDerivedMIOption() == EMaterialStatsDerivedMIOption::CompileOnly)
+		{
+			auto GridPtr = GetGridStatsWidget();
+			if (GridPtr.IsValid())
+			{
+				GridPtr->OnColumnNumChanged();
+			}
+		}
+		else
+		{
+			GetStatsGrid()->OnShaderChanged();
+		}
+
+		bNeedsGridRefresh = false;
+	}
+
+	ComputeGridWarnings();
+}
 
+void FMaterialStats::CacheAndCompilePendingShaders()
+{
 	for (const auto& Entry : ShaderPlatformStatsDB)
 	{
 		auto PlatformStats = Entry.Value;
-		bInfoChanged |= PlatformStats->Update();
+		bNeedsGridRefresh |= PlatformStats->CachePendingShaders();
 	}
+}
 
-	if (bInfoChanged)
+void FMaterialStats::SetMaterial(UMaterial *InMaterial, const TArray<TObjectPtr<UMaterialInstance>>& InDerivedMaterialInstances)
+{
+	if (MaterialInterface != InMaterial || DerivedMaterialInstances != InDerivedMaterialInstances)
 	{
-		GetStatsGrid()->OnShaderChanged();
+		MaterialInterface = InMaterial;
+		DerivedMaterialInstances = InDerivedMaterialInstances;
+
+		for (const auto& Entry : ShaderPlatformStatsDB)
+		{
+			auto& Platform = Entry.Value;
+			if (Platform.IsValid())
+			{
+				Platform->SetMaterial(InMaterial, nullptr, DerivedMaterialInstances);
+			}
+		}
+
+		auto GridPtr = GetGridStatsWidget();
+		if (GridPtr.IsValid())
+		{
+			GridPtr->OnColumnNumChanged();
+		}
 	}
+}
 
-	ComputeGridWarnings();
+void FMaterialStats::SetMaterial(UMaterialInstance *InMaterialInstance)
+{
+	if (MaterialInterface != InMaterialInstance || DerivedMaterialInstances.Num() > 0)
+	{
+		MaterialInterface = InMaterialInstance;
+		DerivedMaterialInstances.Empty();
+
+		for (const auto& Entry : ShaderPlatformStatsDB)
+		{
+			auto& Platform = Entry.Value;
+			if (Platform.IsValid())
+			{
+				Platform->SetMaterial(nullptr, InMaterialInstance, DerivedMaterialInstances);
+			}
+		}
+
+		auto GridPtr = GetGridStatsWidget();
+		if (GridPtr.IsValid())
+		{
+			GridPtr->OnColumnNumChanged();
+		}
+	}
+}
+
+bool FMaterialStats::AnyNewCompilationErrors(const int32 StartingFromInstanceIndex)
+{
+	bool bNewCompilationErrors = false;
+
+	for (const auto& Entry : ShaderPlatformStatsDB)
+	{
+		auto& Platform = Entry.Value;
+		if (Platform.IsValid())
+		{
+			for (int32 QualityLevel = 0; QualityLevel < EMaterialQualityLevel::Num; ++QualityLevel)
+			{
+				for (int32 InstanceIndex = StartingFromInstanceIndex; InstanceIndex < Platform->GetPlatformData(static_cast<EMaterialQualityLevel::Type>(QualityLevel)).Instances.Num(); ++InstanceIndex)
+				{
+					auto &Data = Platform->GetInstanceData(static_cast<EMaterialQualityLevel::Type>(QualityLevel), InstanceIndex);
+					bNewCompilationErrors |= Data.bNeedToWarnAboutCompilationErrors;
+					Data.bNeedToWarnAboutCompilationErrors = false;
+				}
+			}
+		}
+	}
+
+	return bNewCompilationErrors;
 }
 
 TSharedRef<class SDockTab> FMaterialStats::SpawnTab_HLSLCode(const class FSpawnTabArgs& Args)
@@ -746,7 +966,7 @@ TSharedRef<class SDockTab> FMaterialStats::SpawnTab_HLSLCode(const class FSpawnT
 	return SpawnedTab;
 }
 
-TSharedRef<SDockTab> FMaterialStats::SpawnTab_ShaderCode(const FSpawnTabArgs& Args, const EShaderPlatform PlatformID, const EMaterialQualityLevel::Type QualityLevel)
+TSharedRef<SDockTab> FMaterialStats::SpawnTab_ShaderCode(const FSpawnTabArgs& Args, const EShaderPlatform PlatformID, const EMaterialQualityLevel::Type QualityLevel, const int32 InstanceIndex)
 {
 	SetShaderPlatformUseCodeView(PlatformID, QualityLevel, true);
 
@@ -757,12 +977,12 @@ TSharedRef<SDockTab> FMaterialStats::SpawnTab_ShaderCode(const FSpawnTabArgs& Ar
 	check(PlatformPtr.IsValid());
 
 	TSharedRef<SComboBox<TSharedPtr<FMaterialShaderEntry>>> ShaderBox = SNew(SComboBox<TSharedPtr<FMaterialShaderEntry>>)
-		.OptionsSource(PlatformPtr->GetShaderEntries(QualityLevel))
+		.OptionsSource(PlatformPtr->GetShaderEntries(QualityLevel, InstanceIndex))
 		.OnGenerateWidget_Lambda([](TSharedPtr<FMaterialShaderEntry> Value) { return SNew(STextBlock).Text(FText::FromString(Value->Text)); })
-		.OnSelectionChanged_Lambda([PlatformPtr, QualityLevel](TSharedPtr<FMaterialShaderEntry> Item, ESelectInfo::Type SelectInfo) { PlatformPtr->OnShaderViewComboSelectionChanged(Item, QualityLevel); })
+		.OnSelectionChanged_Lambda([PlatformPtr, QualityLevel, InstanceIndex](TSharedPtr<FMaterialShaderEntry> Item, ESelectInfo::Type SelectInfo) { PlatformPtr->OnShaderViewComboSelectionChanged(Item, QualityLevel, InstanceIndex); })
 		[
 			SNew(STextBlock)
-			.Text_Lambda([PlatformPtr, QualityLevel]() { return PlatformPtr->GetSelectedShaderViewComboText(QualityLevel); })
+			.Text_Lambda([PlatformPtr, QualityLevel, InstanceIndex]() { return PlatformPtr->GetSelectedShaderViewComboText(QualityLevel, InstanceIndex); })
 		];
 
 	TSharedRef<SDockTab> SpawnedTab = SNew(SDockTab)
@@ -789,12 +1009,12 @@ TSharedRef<SDockTab> FMaterialStats::SpawnTab_ShaderCode(const FSpawnTabArgs& Ar
 						.ContentPadding(3)
 						.OnClicked_Lambda
 						(
-							[MaterialStats = TWeakPtr<FMaterialStats>(SharedThis(this)), PlatformID, QualityLevel]()
+							[MaterialStats = TWeakPtr<FMaterialStats>(SharedThis(this)), PlatformID, QualityLevel, InstanceIndex]()
 							{
 								auto StatsPtr = MaterialStats.Pin();
 								if (StatsPtr.IsValid())
 								{
-									FText ShaderCode = StatsPtr->GetShaderCode(PlatformID, QualityLevel);
+									FText ShaderCode = StatsPtr->GetShaderCode(PlatformID, QualityLevel, InstanceIndex);
 									FPlatformApplicationMisc::ClipboardCopy(*ShaderCode.ToString());
 									return FReply::Handled();
 								}
@@ -839,9 +1059,9 @@ TSharedRef<SDockTab> FMaterialStats::SpawnTab_ShaderCode(const FSpawnTabArgs& Ar
 	return SpawnedTab;
 }
 
-FName FMaterialStats::MakeTabName(const EPlatformCategoryType PlatformType, const EShaderPlatform ShaderPlatformType, const EMaterialQualityLevel::Type QualityLevel)
+FName FMaterialStats::MakeTabName(const EPlatformCategoryType PlatformType, const EShaderPlatform ShaderPlatformType, const EMaterialQualityLevel::Type QualityLevel, const int32 InstanceIndex)
 {
-	const FString TabName = FMaterialStatsUtils::GetPlatformTypeName(PlatformType) + FMaterialStatsUtils::ShaderPlatformTypeName(ShaderPlatformType) + FMaterialStatsUtils::MaterialQualityToString(QualityLevel);
+	const FString TabName = FMaterialStatsUtils::GetPlatformTypeName(PlatformType) + FMaterialStatsUtils::ShaderPlatformTypeName(ShaderPlatformType) + FMaterialStatsUtils::MaterialQualityToString(QualityLevel) + FString::FromInt(InstanceIndex);
 
 	return FName(*TabName);
 }
@@ -860,6 +1080,8 @@ void FMaterialStats::BuildViewShaderCodeMenus()
 		.SetGroup( PlatformGroupMenuItem.ToSharedRef() )
 		.SetIcon( FSlateIcon(FAppStyle::GetAppStyleSetName(), "MaterialEditor.Tabs.HLSLCode") );
 
+	const int32 InstanceIndex = 0; // only work on base material
+
 	for (auto MapEntry : PlatformTypeDB)
 	{
 		const EPlatformCategoryType PlatformType = MapEntry.Key;
@@ -867,11 +1089,11 @@ void FMaterialStats::BuildViewShaderCodeMenus()
 		const FString PlatformName = FMaterialStatsUtils::GetPlatformTypeName(PlatformType);
 		TSharedPtr<FWorkspaceItem> PlatformMenuItem = PlatformGroupMenuItem->AddGroup(FText::FromString(PlatformName), FSlateIcon(FAppStyle::GetAppStyleSetName(), "MaterialEditor.Tabs.HLSLCode"));
 
-		TArray<TSharedPtr<FShaderPlatformSettings>>& ArrShaderPlatforms = MapEntry.Value;
+		const auto& ArrShaderPlatforms = MapEntry.Value;
 
 		for (int32 i = 0; i < ArrShaderPlatforms.Num(); ++i)
 		{
-			TSharedPtr<FShaderPlatformSettings> PlatformPtr = ArrShaderPlatforms[i];
+			const auto& PlatformPtr = ArrShaderPlatforms[i];
 
 			if (!PlatformPtr.IsValid())
 			{
@@ -893,9 +1115,9 @@ void FMaterialStats::BuildViewShaderCodeMenus()
 				const EMaterialQualityLevel::Type QualityLevel = (EMaterialQualityLevel::Type)q;
 
 				const FString MaterialQualityName = FMaterialStatsUtils::MaterialQualityToString(QualityLevel);
-				const FName TabName = MakeTabName(PlatformType, PlatformID, QualityLevel);
+				const FName TabName = MakeTabName(PlatformType, PlatformID, QualityLevel, InstanceIndex);
 
-				TabManager->RegisterTabSpawner(TabName, FOnSpawnTab::CreateSP(this, &FMaterialStats::SpawnTab_ShaderCode, PlatformID, QualityLevel))
+				TabManager->RegisterTabSpawner(TabName, FOnSpawnTab::CreateSP(this, &FMaterialStats::SpawnTab_ShaderCode, PlatformID, QualityLevel, InstanceIndex))
 					.SetGroup(ShaderPlatformMenuItem.ToSharedRef())
 					.SetDisplayName(FText::FromString(MaterialQualityName));
 
@@ -908,7 +1130,7 @@ void FMaterialStats::BuildViewShaderCodeMenus()
 							auto StatsPtr = MaterialStats.Pin();
 							if (StatsPtr.IsValid())
 							{
-								return StatsPtr->GetShaderCode(PlatformID, QualityLevel);
+								return StatsPtr->GetShaderCode(PlatformID, QualityLevel, InstanceIndex);
 							}
 
 							return FText::FromString(TEXT("Error reading shader code!"));
@@ -943,8 +1165,6 @@ TSharedRef<SDockTab> FMaterialStats::SpawnTab_Stats(const FSpawnTabArgs& Args)
 {
 	check(Args.GetTabId() == StatsTabId);
 
-	FString TabName = FString::Printf(TEXT(""), *GetMaterialName().ToString());
-
 	TSharedRef<SDockTab> SpawnedTab = SNew(SDockTab)
 		.Label(LOCTEXT("Platform Stats", "Platform Stats"))
 		.OnTabClosed_Lambda([&bShowStats = bShowStats](TSharedRef<SDockTab>) { bShowStats = false; })
@@ -1019,6 +1239,8 @@ void FMaterialStats::RegisterTabs()
 void FMaterialStats::UnregisterTabs()
 {
 	auto TabManager = MaterialEditor->GetTabManager();
+	
+	const int32 InstanceIndex = 0;
 
 	for (auto MapEntry : PlatformTypeDB)
 	{
@@ -1038,7 +1260,7 @@ void FMaterialStats::UnregisterTabs()
 
 			for (int32 q = 0; q < EMaterialQualityLevel::Num; ++q)
 			{
-				const FName TabName = MakeTabName(PlatformType, ShaderPlatformID, (EMaterialQualityLevel::Type)q);
+				const FName TabName = MakeTabName(PlatformType, ShaderPlatformID, (EMaterialQualityLevel::Type)q, InstanceIndex);
 
 				TabManager->UnregisterTabSpawner(TabName);
 			}
@@ -1063,22 +1285,58 @@ void FMaterialStats::ComputeGridWarnings()
 		return;
 	}
 
-	int32 Warnings = 0;
-	TArray<EShaderPlatform> CompilerWarnings;
-	TArray<FString> WarningMessages;
+	// no need to update warnings that often
+	const double kMinimumTimeBetweenWarningsRefreshSeconds = 0.1;
+	const double CurrentTime = FPlatformTime::Seconds();
+	const bool bCooledDown = (CurrentTime - LastGridMessagesUpdate) > kMinimumTimeBetweenWarningsRefreshSeconds;
+
+	if (!bCooledDown)
+	{
+		return;
+	}
+
+	LastGridMessagesUpdate = CurrentTime;
+
+	TArray<FString> Messages;
+	TArray<bool> IsErrorMessage;
 
 	bool bAnyQualityPresent = false;
 	for (int32 q = 0; q < EMaterialQualityLevel::Num; ++q)
 	{
 		auto Quality = static_cast<EMaterialQualityLevel::Type>(q);
-
 		bAnyQualityPresent |= GetStatsQualityFlag(Quality);
 	}
 
 	if (!bAnyQualityPresent)
 	{
-		Warnings |= WarningNoQuality;
-		WarningMessages.Emplace("No material quality selected. Please use the 'Settings' button and choose the desired quality level to be analyzed.");
+		Messages.Emplace("No material quality selected. Please use the 'Settings' button and choose the desired quality level to be analyzed.");
+		IsErrorMessage.Add(false);
+	}
+
+	for (const auto& Entry : ShaderPlatformStatsDB)
+	{
+		auto& Platform = Entry.Value;
+		if (Platform.IsValid())
+		{
+			for (int32 QualityLevel = 0; QualityLevel < EMaterialQualityLevel::Num; ++QualityLevel)
+			{
+				for (int32 InstanceIndex = 0; InstanceIndex < Platform->GetPlatformData(static_cast<EMaterialQualityLevel::Type>(QualityLevel)).Instances.Num(); ++InstanceIndex)
+				{
+					auto& Data = Platform->GetInstanceData(static_cast<EMaterialQualityLevel::Type>(QualityLevel), InstanceIndex);
+					auto& Errors = Data.MaterialResourcesStats->GetCompileErrors();
+					for (int32 ErrorIndex = 0; ErrorIndex < Errors.Num(); ++ErrorIndex)
+					{
+						const auto& PlatformName = Platform->GetPlatformName().ToString();
+						const auto& AssetName = GetMaterialName(InstanceIndex);
+						const FString QualityName = FMaterialStatsUtils::MaterialQualityToShortString((EMaterialQualityLevel::Type)QualityLevel);
+
+						FString Message = FString::Printf(TEXT("[%s/%s:%s] %s"), *AssetName, *PlatformName, *QualityName, *Errors[ErrorIndex]);
+						Messages.Add(Message);
+						IsErrorMessage.Add(true);
+					}
+				}
+			}
+		}
 	}
 
 	bool bAnyPlatformPresent = false;
@@ -1101,9 +1359,9 @@ void FMaterialStats::ComputeGridWarnings()
 
 				if (!bCompilerAvailable)
 				{
-					CompilerWarnings.Add(ShaderPlatformType);
 					auto WarningString = FString::Printf(TEXT("Platform %s needs an offline shader compiler to extract instruction count. Please check 'Editor Preferences' -> 'Content Editors' -> 'Material Editor' for additional settings."), *PlatformPtr->GetPlatformName().ToString());
-					WarningMessages.Add(MoveTemp(WarningString));
+					Messages.Add(MoveTemp(WarningString));
+					IsErrorMessage.Add(false);
 				}
 			}
 		}
@@ -1111,33 +1369,20 @@ void FMaterialStats::ComputeGridWarnings()
 
 	if (!bAnyPlatformPresent)
 	{
-		Warnings |= WarningNoPlatform;
-		WarningMessages.Emplace("No platform selected. Please use the 'Settings' button and choose desired platform to be analyzed.");
-	}
-
-	bool bRefreshWarnings = (Warnings != LastGenericWarning) || (CompilerWarnings.Num() != LastMissingCompilerWarnings.Num());
-	if (!bRefreshWarnings)
-	{
-		for (int32 i = 0; i < CompilerWarnings.Num(); ++i)
-		{
-			if (CompilerWarnings[i] != LastMissingCompilerWarnings[i])
-			{
-				bRefreshWarnings = true;
-				break;
-			}
-		}
+		Messages.Emplace("No platform selected. Please use the 'Settings' button and choose desired platform to be analyzed.");
+		IsErrorMessage.Add(false);
 	}
 
-	if (bRefreshWarnings)
+	if (Messages != LastGridMessages)
 	{
-		LastGenericWarning = Warnings;
-		LastMissingCompilerWarnings = CompilerWarnings;
+		LastGridMessages = Messages;
 
-		GridPtr->ClearWarningMessages();
+		GridPtr->ClearMessages();
 
-		for (int32 i = 0; i < WarningMessages.Num(); ++i)
+		check(Messages.Num() == IsErrorMessage.Num());
+		for (int32 i = 0; i < Messages.Num(); ++i)
 		{
-			GridPtr->AddWarningMessage(WarningMessages[i]);
+			GridPtr->AddMessage(Messages[i], IsErrorMessage[i]);
 		}
 	}
 }
diff --git a/Engine/Source/Editor/MaterialEditor/Private/MaterialStats.h b/Engine/Source/Editor/MaterialEditor/Private/MaterialStats.h
--- a/Engine/Source/Editor/MaterialEditor/Private/MaterialStats.h
+++ b/Engine/Source/Editor/MaterialEditor/Private/MaterialStats.h
@@ -5,6 +5,7 @@
 #include "Materials/Material.h"
 #include "MaterialStatsCommon.h"
 #include "UObject/GCObject.h"
+#include "Preferences/MaterialStatsOptions.h"
 
 /** structure used to store various statistics extracted from compiled shaders */
 struct FShaderStatsInfo
@@ -41,7 +42,7 @@ struct FShaderStatsInfo
 
 		ShaderCount.StrDescription = TEXT("Compiling...");
 		ShaderCount.StrDescriptionLong = TEXT("Compiling...");
-		
+
 		StrShaderErrors.Empty();
 	}
 
@@ -84,13 +85,11 @@ struct FMaterialShaderEntry
 struct FShaderPlatformSettings
 {
 public:
-	/////////////////////
-	/** inner structure used to hold properties for a single material platform and with a specific quality level */
-	struct FPlatformData
+	struct FInstanceData
 	{
 		/** pointer to the material resource created for this platform
 		*  mainly used to compile the shaders and extract information from them */
-		FMaterialResourceStats* MaterialResourcesStats = nullptr;
+		FMaterialResourceStats* MaterialResourcesStats;
 
 		/** array of shader ids for this platform; needed to fill ComboBox in MaterialEditor's shader viewer
 		* generated from ShaderID.ShaderType->GetFName() */
@@ -99,27 +98,40 @@ public:
 		/** ComboBox current entry */
 		FMaterialShaderEntry ComboBoxSelectedEntry;
 
-		/** flag that marks the usage of this data structure */
-		bool bExtractStats = false;
-
-		/** true when code is listed in its own tab */
-		bool bExtractCode = false;
-
-		/** cached shader code computed by FShaderPlatformSettings::GetShaderCode() */
-		FText ShaderCode;
 		/** when true we should update the content of [FText ShaderCode] variable */
 		bool bUpdateShaderCode = false;
+		/** cached shader code computed by FShaderPlatformSettings::GetShaderCode() */
+		FText ShaderCode;
 		/** flag that marks an ongoing shader compilation */
 		bool bCompilingShaders = false;
 		/** flag suggests we needed to recompile shaders due to changes in the material */
 		bool bNeedShaderRecompilation = false;
+		/** flag suggests we only need to compile the most complex shader from available set */
+		bool bOnlyCompileMostComplexShader = false;
+		/** flag suggests we needed to warn the user due to compilation errors */
+		bool bNeedToWarnAboutCompilationErrors = false;
+
+		FShaderStatsInfo ShaderStatsInfo;
+	};
+	
+	/////////////////////
+	/** inner structure used to hold properties for a single material platform and with a specific quality level */
+	struct FPlatformData
+	{
+		/* Array of material instances including base material and any material instances derived from it */
+		TArray<FInstanceData> Instances;
+		
+		/** flag that marks the usage of this data structure */
+		bool bExtractStats = false;
+
+		/** true when code is listed in its own tab */
+		bool bExtractCode = false;
+
 		/** object used to display the content of [FText ShaderCode] */
 		TSharedPtr<class SScrollBox> CodeScrollBox;
 		/** weak pointer to the spawned shader code viewer tab */
 		TWeakPtr<class SDockTab> CodeViewerTab;
 
-		FShaderStatsInfo ShaderStatsInfo;
-
 		/** Time when the previous compilation was requested.*/
 		double LastTimeCompilationRequested = 0.0;
 	};
@@ -153,11 +165,13 @@ private:
 	UMaterial *Material = nullptr;
 	/** pointer to the material instance whose stats are analyzed */
 	UMaterialInstance *MaterialInstance = nullptr;
+	/** array of derived material instances whose stats are analyzed */ 
+	TArray<TObjectPtr<UMaterialInstance>> DerivedMaterialInstances;
 
 private:
 	/** function used to trigger shader rebuilding when needed */
 	/** returns true if shaders are being recompiled, false otherwise */
-	bool CheckShaders();
+	bool CheckShaders(bool bIgnoreCooldown);
 
 	/** builds material resources needed to compile shaders */
 	void AllocateMaterialResources();
@@ -197,7 +211,7 @@ public:
 	FORCEINLINE TWeakPtr<class SDockTab> GetCodeViewerTab(const EMaterialQualityLevel::Type QualityLevel);
 
 	/** returns an array with the names of all the compiled shaders for this material with specified quality level */
-	FORCEINLINE const TArray<TSharedPtr<FMaterialShaderEntry>> *GetShaderEntries(const EMaterialQualityLevel::Type QualityLevel);
+	FORCEINLINE const TArray<TSharedPtr<FMaterialShaderEntry>> *GetShaderEntries(const EMaterialQualityLevel::Type QualityLevel, const int32 InstanceIndex);
 
 	/** when set this flag will indicate the presence of this material with a particular quality level inside the stats widget */
 	void SetExtractStatsFlag(const EMaterialQualityLevel::Type QualityType, const bool bValue);
@@ -223,7 +237,7 @@ public:
 	FORCEINLINE void SetExtractStatsQualityLevel(const EMaterialQualityLevel::Type Quality, const bool bActive);
 
 	/** flags shader compilation for a specific quality level */
-	FORCEINLINE void SetNeedShaderCompilation(const EMaterialQualityLevel::Type QualityLevel, const bool bValue);
+	FORCEINLINE void SetNeedShaderCompilation(const EMaterialQualityLevel::Type QualityLevel, const bool bValue, const bool bOnlyCompileDerivedMI);
 
 	/** returns the shader type used by this platform as set at construction time */
 	FORCEINLINE EShaderPlatform GetPlatformShaderType() const;
@@ -231,22 +245,26 @@ public:
 	/** returns a reference to the platform settings used for a specified material quality level; see the definition of FPlatformData inner structure */
 	FORCEINLINE FPlatformData& GetPlatformData(const EMaterialQualityLevel::Type QualityLevel);
 
+	/** returns a reference to the platform instance data used for a specified material quality level and instance index; see the definition of FInstanceData inner structure */
+	FORCEINLINE FInstanceData& GetInstanceData(const EMaterialQualityLevel::Type QualityLevel, const int32 InstanceIndex);
+
 	/** returns the shader name chosen in the shader viewer combo-box */
-	FText GetSelectedShaderViewComboText(EMaterialQualityLevel::Type QualityLevel) const;
+	FText GetSelectedShaderViewComboText(EMaterialQualityLevel::Type QualityLevel, const int32 InstanceIndex) const;
 
 	/** callback function called when we change the content of the shader viewer combo-box, used to select a different shader to be displayed */
-	void OnShaderViewComboSelectionChanged(TSharedPtr<FMaterialShaderEntry> Item, EMaterialQualityLevel::Type QualityType);
+	void OnShaderViewComboSelectionChanged(TSharedPtr<FMaterialShaderEntry> Item, EMaterialQualityLevel::Type QualityType, const int32 InstanceIndex);
 
 	/** returns the actual shader source selected by the shaders viewer's combo-box */
-	FText GetShaderCode(const EMaterialQualityLevel::Type QualityType);
+	FText GetShaderCode(const EMaterialQualityLevel::Type QualityType, const int32 InstanceIndex);
 
 	/** call this whenever the analyzed material or material instance is changed */
-	void SetMaterial(UMaterial *InMaterial);
-	void SetMaterial(UMaterialInstance *InMaterialInstance);
+	void SetMaterial(UMaterial *InBaseMaterial, UMaterialInstance *InBaseMaterialInstance, const TArray<TObjectPtr<UMaterialInstance>>& InDerivedMaterialInstances);
 
 	/** main function used to update the state of this platform; will be called from FMaterialStats::Update() */
 	/** returns true if something changed for this the update call, false otherwise */
 	bool Update();
+
+	bool CachePendingShaders();
 };
 
 /** name alias used bellow in FMaterialStats */
@@ -275,8 +293,11 @@ class FMaterialStats : public FGCObject, public TSharedFromThis<FMaterialStats>
 	/** array of bools that flag a specific global material quality setting of the stats grid widget */
 	bool bArrStatsQualitySelector[EMaterialQualityLevel::Num] = { false };
 
-	/** name of the analyzed material */
-	FText MaterialName;
+	/** inspect all derived material instances also */
+	EMaterialStatsDerivedMIOption MaterialStatsDerivedMIOption = EMaterialStatsDerivedMIOption::CompileOnly;
+
+	/** names of the analyzed materials */
+	TArray<FString> MaterialNames;
 
 	/** the id of the stats grid widget tab */
 	static const FName StatsTabId;
@@ -303,12 +324,14 @@ class FMaterialStats : public FGCObject, public TSharedFromThis<FMaterialStats>
 	/** pointer to the material interface whose stats are analyzed */
 	class UMaterialInterface *MaterialInterface = nullptr;
 
-	/** grid warning messages ids */
-	static const int32 WarningNoQuality = 1;
-	static const int32 WarningNoPlatform = 2;
+	/** array of pointers to derived material instances whose stats are analyzed */
+	TArray<TObjectPtr<UMaterialInstance>> DerivedMaterialInstances;
+
+	/** grid warning messages */
+	TArray<FString> LastGridMessages;
+	double LastGridMessagesUpdate = 0.0;
 
-	int32 LastGenericWarning = 0;
-	TArray<EShaderPlatform> LastMissingCompilerWarnings;
+	bool bNeedsGridRefresh = false;
 
 private:
 	/** adds a specified platform in the grid widget for analysis; usually called from BuildShaderPlatformDB() */
@@ -319,11 +342,11 @@ private:
 	void BuildShaderPlatformDB();
 
 	/** this will spawn the window that will display the a specific set of shaders from the analyzed material */
-	TSharedRef<class SDockTab> SpawnTab_ShaderCode(const class FSpawnTabArgs& Args, const EShaderPlatform PlatformID, const EMaterialQualityLevel::Type QualityType);
+	TSharedRef<class SDockTab> SpawnTab_ShaderCode(const class FSpawnTabArgs& Args, const EShaderPlatform PlatformID, const EMaterialQualityLevel::Type QualityType, const int32 InstanceIndex);
 	TSharedRef<class SDockTab> SpawnTab_HLSLCode(const class FSpawnTabArgs& Args);
 
 	/** utility function used to build names for the shader viewing tabs  */
-	static FName MakeTabName(const EPlatformCategoryType PlatformType, const EShaderPlatform ShaderPlatformType, const EMaterialQualityLevel::Type QualityLevel);
+	static FName MakeTabName(const EPlatformCategoryType PlatformType, const EShaderPlatform ShaderPlatformType, const EMaterialQualityLevel::Type QualityLevel, const int32 InstanceIndex);
 
 	/** functions responsible for construction and spawning of the window that will display the stats grid widget */
 	void BuildStatsTab();
@@ -379,7 +402,7 @@ public:
 	static FORCEINLINE FName GetGridOldStatsTabName();
 
 	/** returns the material name we're analyzing */
-	FORCEINLINE FText GetMaterialName() const;
+	FORCEINLINE FString GetMaterialName(const int32 InstanceIndex) const;
 
 	/** returns a pointer to the material stats grid widget */
 	FORCEINLINE TSharedPtr<class SMaterialEditorStatsWidget> GetGridStatsWidget();
@@ -397,12 +420,15 @@ public:
 	void SetStatusQualityFlag(const EMaterialQualityLevel::Type Quality, const bool bValue);
 	FORCEINLINE bool GetStatsQualityFlag(const EMaterialQualityLevel::Type Quality);
 
+	void SetMaterialStatsDerivedMIOption(const EMaterialStatsDerivedMIOption value);
+	FORCEINLINE EMaterialStatsDerivedMIOption GetMaterialStatsDerivedMIOption() const;
+	FORCEINLINE bool GetProvideDerivedMIFlag() const;
+
 	/** switches on or off the presence of a specified shader platform inside the stats grid widget for this material */
 	bool SwitchShaderPlatformUseStats(const EShaderPlatform PlatformID);
 
 	/** Sets the name of the material that will be displayed in the stats grid widget */
-	template <typename TString>
-	void SetMaterialDisplayName(TString&& Name);
+	void SetMaterialsDisplayNames(const TArray<FString>& Names);
 
 	/** returns a platform name given at its construction time */
 	FName GetPlatformName(const EShaderPlatform InEnumValue) const;
@@ -415,7 +441,7 @@ public:
 	TSharedPtr<FShaderPlatformSettings> GetPlatformSettings(const FName PlatformName);
 
 	/** returns the shader code computed by the specified platform with some quality level */
-	FText GetShaderCode(const EShaderPlatform PlatformID, const EMaterialQualityLevel::Type QualityType);
+	FText GetShaderCode(const EShaderPlatform PlatformID, const EMaterialQualityLevel::Type QualityType, const int32 InstanceIndex);
 
 	/** call this whenever some material property is changed, as it will trigger shader recompilation */
 	void SignalMaterialChanged();
@@ -439,7 +465,7 @@ public:
 
 private:
 	/** this function will do the setup procedure and its called from FMaterialStatsUtils::CreateMaterialStats()  */
-	void Initialize(IMaterialEditor *MaterialEditor);
+	void Initialize(IMaterialEditor *MaterialEditor, const bool ShowMaterialInstancesMenu);
 public:
 	//end Setup Functions
 	///////////////////////////////////////////
@@ -449,13 +475,18 @@ public:
 	/** call this from chosen IMaterialEditor whenever an update to this analysis tool in appropriate */
 	void Update();
 
-	/** this will set the material to be analyzed by this class */
-	/** TMaterial should be UMaterial or UMaterialInstance */
-	template <typename TMaterial>
-	void SetMaterial(TMaterial& InMaterial);
+	/** call this to request caching required shaders regardless of cooldown */
+	void CacheAndCompilePendingShaders();
+
+	/** this will set the material or material instance to be analyzed by this class */
+	void SetMaterial(UMaterial *InMaterial, const TArray<TObjectPtr<UMaterialInstance>>& InDerivedMaterialInstances);
+	void SetMaterial(UMaterialInstance *InMaterialInstance);
 
 	// end Update Functions
 	///////////////////////////////////////////
+
+	/** returns true if any new compilation errors are discovered */
+	bool AnyNewCompilationErrors(const int32 StartingFromInstanceIndex = 0);
 };
 
 
@@ -492,15 +523,21 @@ FORCEINLINE bool FMaterialStats::IsShowingOldStats() const
 	return bShowOldStats;
 }
 
-template <typename TString>
-void FMaterialStats::SetMaterialDisplayName(TString&& Name)
+FORCEINLINE void FMaterialStats::SetMaterialsDisplayNames(const TArray<FString>& Names)
 {
-	MaterialName = FText::FromString(Forward<TString>(Name));
+	MaterialNames = Names;
 }
 
-FORCEINLINE FText FMaterialStats::GetMaterialName() const
+FORCEINLINE FString FMaterialStats::GetMaterialName(const int32 InstanceIndex) const
 {
-	return MaterialName;
+	if (InstanceIndex < MaterialNames.Num())
+	{
+		return MaterialNames[InstanceIndex];
+	}
+	else
+	{
+		return TEXT("");
+	}
 }
 
 FORCEINLINE FName FMaterialStats::GetGridStatsTabName()
@@ -519,27 +556,19 @@ FORCEINLINE bool FMaterialStats::GetStatsQualityFlag(const EMaterialQualityLevel
 	return bArrStatsQualitySelector[(int32)Quality];
 }
 
-FORCEINLINE TSharedPtr<class SMaterialEditorStatsWidget> FMaterialStats::GetGridStatsWidget()
+FORCEINLINE EMaterialStatsDerivedMIOption FMaterialStats::GetMaterialStatsDerivedMIOption() const
 {
-	return GridStatsWidget;
+	return MaterialStatsDerivedMIOption;
 }
 
-//TMaterial should be UMaterial or UMaterialInstance
-template <typename TMaterial>
-void FMaterialStats::SetMaterial(TMaterial& MaterialPtr)
+FORCEINLINE bool FMaterialStats::GetProvideDerivedMIFlag() const
 {
-	if (MaterialInterface != MaterialPtr)
-	{
-		MaterialInterface = MaterialPtr;
-		for (const auto& Entry : ShaderPlatformStatsDB)
-		{
-			auto& Platform = Entry.Value;
-			if (Platform.IsValid())
-			{
-				Platform->SetMaterial(MaterialPtr);
-			}
-		}
-	}
+	return MaterialStatsDerivedMIOption != EMaterialStatsDerivedMIOption::Ignore;
+}
+
+FORCEINLINE TSharedPtr<class SMaterialEditorStatsWidget> FMaterialStats::GetGridStatsWidget()
+{
+	return GridStatsWidget;
 }
 
 // end FMaterialStats implementation
@@ -590,10 +619,10 @@ FORCEINLINE TWeakPtr<class SDockTab> FShaderPlatformSettings::GetCodeViewerTab(c
 	return SomePlatformData.CodeViewerTab;
 }
 
-FORCEINLINE const TArray<TSharedPtr<FMaterialShaderEntry>> *FShaderPlatformSettings::GetShaderEntries(const EMaterialQualityLevel::Type QualityLevel)
+FORCEINLINE const TArray<TSharedPtr<FMaterialShaderEntry>> *FShaderPlatformSettings::GetShaderEntries(const EMaterialQualityLevel::Type QualityLevel, const int32 InstanceIndex)
 {
-	FPlatformData& SomePlatformData = GetPlatformData(QualityLevel);
-	return &SomePlatformData.ArrShaderEntries;
+	auto& InstanceData = GetInstanceData(QualityLevel, InstanceIndex);
+	return &InstanceData.ArrShaderEntries;
 }
 
 FORCEINLINE void FShaderPlatformSettings::SetExtractStatsFlag(const EMaterialQualityLevel::Type QualityType, const bool bValue)
@@ -639,10 +668,15 @@ FORCEINLINE void FShaderPlatformSettings::SetExtractStatsQualityLevel(const EMat
 	PlatformData[Quality].bExtractStats = bActive;
 }
 
-FORCEINLINE void FShaderPlatformSettings::SetNeedShaderCompilation(const EMaterialQualityLevel::Type QualityLevel, const bool bValue)
+FORCEINLINE void FShaderPlatformSettings::SetNeedShaderCompilation(const EMaterialQualityLevel::Type QualityLevel, const bool bValue, const bool bOnlyCompileDerivedMI)
 {
-	check(QualityLevel != EMaterialQualityLevel::Num);
-	PlatformData[QualityLevel].bNeedShaderRecompilation = bValue;
+	check(QualityLevel < EMaterialQualityLevel::Num);
+	for (int32 InstancesIndex = 0; InstancesIndex < PlatformData[QualityLevel].Instances.Num(); ++InstancesIndex)
+	{
+		auto& Instance = GetInstanceData(QualityLevel, InstancesIndex);
+		Instance.bNeedShaderRecompilation = bValue;
+		Instance.bOnlyCompileMostComplexShader = (InstancesIndex > 0) && bOnlyCompileDerivedMI;
+	}
 }
 
 FORCEINLINE EShaderPlatform FShaderPlatformSettings::GetPlatformShaderType() const
@@ -652,8 +686,15 @@ FORCEINLINE EShaderPlatform FShaderPlatformSettings::GetPlatformShaderType() con
 
 FORCEINLINE FShaderPlatformSettings::FPlatformData& FShaderPlatformSettings::GetPlatformData(const EMaterialQualityLevel::Type QualityLevel)
 {
-	check(QualityLevel != EMaterialQualityLevel::Num);
+	check(QualityLevel < EMaterialQualityLevel::Num);
 	return PlatformData[QualityLevel];
 }
+
+FORCEINLINE FShaderPlatformSettings::FInstanceData& FShaderPlatformSettings::GetInstanceData(const EMaterialQualityLevel::Type QualityLevel, const int32 InstanceIndex)
+{
+	check(QualityLevel < EMaterialQualityLevel::Num);
+	return PlatformData[QualityLevel].Instances[InstanceIndex];
+}
+
 // end FShaderPlatformSettings implementation
 //////////////////////////////////////////////////////////////////////////////////////////////////
diff --git a/Engine/Source/Editor/MaterialEditor/Private/MaterialStatsCommon.cpp b/Engine/Source/Editor/MaterialEditor/Private/MaterialStatsCommon.cpp
--- a/Engine/Source/Editor/MaterialEditor/Private/MaterialStatsCommon.cpp
+++ b/Engine/Source/Editor/MaterialEditor/Private/MaterialStatsCommon.cpp
@@ -27,10 +27,10 @@ void FMaterialResourceStats::SetupExtraCompilationSettings(const EShaderPlatform
 /***********************************************************************************************************************/
 /*begin FMaterialStatsUtils */
 
-TSharedPtr<FMaterialStats> FMaterialStatsUtils::CreateMaterialStats(class IMaterialEditor* MaterialEditor)
+TSharedPtr<FMaterialStats> FMaterialStatsUtils::CreateMaterialStats(class IMaterialEditor* MaterialEditor, const bool ShowMaterialInstancesMenu)
 {
 	TSharedPtr<FMaterialStats> MaterialStats = MakeShareable(new FMaterialStats());
-	MaterialStats->Initialize(MaterialEditor);
+	MaterialStats->Initialize(MaterialEditor, ShowMaterialInstancesMenu);
 
 	return MaterialStats;
 }
diff --git a/Engine/Source/Editor/MaterialEditor/Private/MaterialStatsGrid.cpp b/Engine/Source/Editor/MaterialEditor/Private/MaterialStatsGrid.cpp
--- a/Engine/Source/Editor/MaterialEditor/Private/MaterialStatsGrid.cpp
+++ b/Engine/Source/Editor/MaterialEditor/Private/MaterialStatsGrid.cpp
@@ -32,9 +32,9 @@ TSharedPtr<FGridCell> FStatsGridRow::GetCell(const FName ColumnName)
 	return RetCell;
 }
 
-void FStatsGridRow::RemovePlatform(TSharedPtr<FMaterialStats> StatsManager, const TSharedPtr<FShaderPlatformSettings> Platform, const EMaterialQualityLevel::Type QualityLevel)
+void FStatsGridRow::RemovePlatform(TSharedPtr<FMaterialStats> StatsManager, const TSharedPtr<FShaderPlatformSettings> Platform, const EMaterialQualityLevel::Type QualityLevel, const int32 InstanceIndex)
 {
-	const FName ColumnName = FMaterialStatsGrid::MakePlatformColumnName(Platform, QualityLevel);
+	const FName ColumnName = FMaterialStatsGrid::MakePlatformColumnName(Platform, QualityLevel, InstanceIndex);
 	RemoveCell(ColumnName);
 }
 
@@ -55,8 +55,11 @@ void FStatsGridRow::FillPlatformCellsHelper(TSharedPtr<FMaterialStats> StatsMana
 
 			if (StatsManager->GetStatsQualityFlag(QualityLevel))
 			{
-				// call implementation specific function to build the needed cell
-				AddPlatform(StatsManager, Platform, QualityLevel);
+				for (int32 InstanceIndex = 0; InstanceIndex < Platform->GetPlatformData(QualityLevel).Instances.Num(); ++InstanceIndex)
+				{
+					// call implementation specific function to build the needed cell
+					AddPlatform(StatsManager, Platform, QualityLevel, InstanceIndex);
+				}
 			}
 		}
 	}
@@ -77,11 +80,11 @@ void FStatsGridRow_Empty::CreateRow(TSharedPtr<FMaterialStats> StatsManager)
 	FillPlatformCellsHelper(StatsManager);
 }
 
-void FStatsGridRow_Empty::AddPlatform(TSharedPtr<FMaterialStats> StatsManager, const TSharedPtr<FShaderPlatformSettings> Platform, const EMaterialQualityLevel::Type QualityLevel)
+void FStatsGridRow_Empty::AddPlatform(TSharedPtr<FMaterialStats> StatsManager, const TSharedPtr<FShaderPlatformSettings> Platform, const EMaterialQualityLevel::Type QualityLevel, const int32 InstanceIndex)
 {
 	TSharedPtr<FGridCell_Empty> Cell = MakeShareable(new FGridCell_Empty());
 
-	const FName ColumnName = FMaterialStatsGrid::MakePlatformColumnName(Platform, QualityLevel);
+	const FName ColumnName = FMaterialStatsGrid::MakePlatformColumnName(Platform, QualityLevel, InstanceIndex);
 	AddCell(ColumnName, Cell);
 }
 
@@ -89,9 +92,9 @@ void FStatsGridRow_Empty::AddPlatform(TSharedPtr<FMaterialStats> StatsManager, c
 /*==============================================================================================================*/
 
 /*==============================================================================================================*/
-/* FStatsGridRow_Quality functions*/
+/* FStatsGridRow_Name functions*/
 
-void FStatsGridRow_Quality::CreateRow(TSharedPtr<FMaterialStats> StatsManager)
+void FStatsGridRow_Name::CreateRow(TSharedPtr<FMaterialStats> StatsManager)
 {
 	// we don't use a descriptor for this row
 	AddCell(FMaterialStatsGrid::DescriptorColumnName, MakeShareable(new FGridCell_Empty()));
@@ -100,52 +103,46 @@ void FStatsGridRow_Quality::CreateRow(TSharedPtr<FMaterialStats> StatsManager)
 	FillPlatformCellsHelper(StatsManager);
 }
 
-void FStatsGridRow_Quality::AddPlatform(TSharedPtr<FMaterialStats> StatsManager, const TSharedPtr<FShaderPlatformSettings> Platform, const EMaterialQualityLevel::Type QualityLevel)
+void FStatsGridRow_Name::AddPlatform(TSharedPtr<FMaterialStats> StatsManager, const TSharedPtr<FShaderPlatformSettings> Platform, const EMaterialQualityLevel::Type QualityLevel, const int32 InstanceIndex)
 {
-	// translate material quality to string and store it inside a StaticString cell
-	const FString CellContent = FMaterialStatsUtils::MaterialQualityToShortString(QualityLevel);
-	TSharedPtr<FGridCell_StaticString> Cell = MakeShareable(new FGridCell_StaticString(CellContent, CellContent));
-
-	Cell->SetContentBold(true);
-	FSlateColor CellColor = FMaterialStatsUtils::QualitySettingColor(QualityLevel);
-	Cell->SetColor(CellColor);
+	// add a cell that will query any available errors for this platform
+	TSharedPtr<FGridCell_ShaderValue> Cell = MakeShareable(new FGridCell_ShaderValue(StatsManager, EShaderInfoType::Name, ERepresentativeShader::Num, QualityLevel, Platform->GetPlatformShaderType(), InstanceIndex));
+	Cell->SetHorizontalAlignment(HAlign_Fill);
 
-	const FName ColumnName = FMaterialStatsGrid::MakePlatformColumnName(Platform, QualityLevel);
+	const FName ColumnName = FMaterialStatsGrid::MakePlatformColumnName(Platform, QualityLevel, InstanceIndex);
 	AddCell(ColumnName, Cell);
 }
 
-/*end FStatsGridRow_Quality functions*/
+/*end FStatsGridRow_Name functions*/
 /*==============================================================================================================*/
 
 /*==============================================================================================================*/
-/* FStatsGridRow_Errors functions*/
+/* FStatsGridRow_Quality functions*/
 
-void FStatsGridRow_Errors::CreateRow(TSharedPtr<FMaterialStats> StatsManager)
+void FStatsGridRow_Quality::CreateRow(TSharedPtr<FMaterialStats> StatsManager)
 {
-	// add an "Error" string in the descriptor column
-	TSharedPtr<FGridCell> HeaderCell = MakeShareable(new FGridCell_StaticString(TEXT("Errors"), TEXT("Errors")));
-	HeaderCell->SetColor(FStyleColors::AccentOrange);
-	HeaderCell->SetContentBold(true);
-	AddCell(FMaterialStatsGrid::DescriptorColumnName, HeaderCell);
-
+	// we don't use a descriptor for this row
+	AddCell(FMaterialStatsGrid::DescriptorColumnName, MakeShareable(new FGridCell_Empty()));
 	AddCell(FMaterialStatsGrid::ShaderColumnName, MakeShareable(new FGridCell_Empty()));
 
 	FillPlatformCellsHelper(StatsManager);
 }
 
-void FStatsGridRow_Errors::AddPlatform(TSharedPtr<FMaterialStats> StatsManager, const TSharedPtr<FShaderPlatformSettings> Platform, const EMaterialQualityLevel::Type QualityLevel)
+void FStatsGridRow_Quality::AddPlatform(TSharedPtr<FMaterialStats> StatsManager, const TSharedPtr<FShaderPlatformSettings> Platform, const EMaterialQualityLevel::Type QualityLevel, const int32 InstanceIndex)
 {
-	// add a cell that will query any available errors for this platform
+	// translate material quality to string and store it inside a StaticString cell
 	const FString CellContent = FMaterialStatsUtils::MaterialQualityToShortString(QualityLevel);
-	TSharedPtr<FGridCell_ShaderValue> Cell = MakeShareable(new FGridCell_ShaderValue(StatsManager, EShaderInfoType::Errors, ERepresentativeShader::Num, QualityLevel, Platform->GetPlatformShaderType()));
-	Cell->SetColor(FStyleColors::AccentOrange);
-	Cell->SetHorizontalAlignment(HAlign_Fill);
+	TSharedPtr<FGridCell_StaticString> Cell = MakeShareable(new FGridCell_StaticString(CellContent, CellContent));
+
+	Cell->SetContentBold(true);
+	FSlateColor CellColor = FMaterialStatsUtils::QualitySettingColor(QualityLevel);
+	Cell->SetColor(CellColor);
 
-	const FName ColumnName = FMaterialStatsGrid::MakePlatformColumnName(Platform, QualityLevel);
+	const FName ColumnName = FMaterialStatsGrid::MakePlatformColumnName(Platform, QualityLevel, InstanceIndex);
 	AddCell(ColumnName, Cell);
 }
 
-/*end FStatsGridRow_Errors functions*/
+/*end FStatsGridRow_Quality functions*/
 /*==============================================================================================================*/
 
 /*==============================================================================================================*/
@@ -213,13 +210,13 @@ void FStatsGridRow_Shaders::CreateRow(TSharedPtr<FMaterialStats> StatsManager)
 	FillPlatformCellsHelper(StatsManager);
 }
 
-void FStatsGridRow_Shaders::AddPlatform(TSharedPtr<FMaterialStats> StatsManager, const TSharedPtr<FShaderPlatformSettings> Platform, const EMaterialQualityLevel::Type QualityLevel)
+void FStatsGridRow_Shaders::AddPlatform(TSharedPtr<FMaterialStats> StatsManager, const TSharedPtr<FShaderPlatformSettings> Platform, const EMaterialQualityLevel::Type QualityLevel, const int32 InstanceIndex)
 {
 	// add a cell that display the instruction count for this platform
 	const FString CellContent = FMaterialStatsUtils::MaterialQualityToShortString(QualityLevel);
-	TSharedPtr<FGridCell_ShaderValue> Cell = MakeShareable(new FGridCell_ShaderValue(StatsManager, EShaderInfoType::InstructionsCount, ShaderType, QualityLevel, Platform->GetPlatformShaderType()));
+	TSharedPtr<FGridCell_ShaderValue> Cell = MakeShareable(new FGridCell_ShaderValue(StatsManager, EShaderInfoType::InstructionsCount, ShaderType, QualityLevel, Platform->GetPlatformShaderType(), InstanceIndex));
 
-	const FName ColumnName = FMaterialStatsGrid::MakePlatformColumnName(Platform, QualityLevel);
+	const FName ColumnName = FMaterialStatsGrid::MakePlatformColumnName(Platform, QualityLevel, InstanceIndex);
 	AddCell(ColumnName, Cell);
 }
 
@@ -242,13 +239,13 @@ void FStatsGridRow_Samplers::CreateRow(TSharedPtr<FMaterialStats> StatsManager)
 	FillPlatformCellsHelper(StatsManager);
 }
 
-void FStatsGridRow_Samplers::AddPlatform(TSharedPtr<FMaterialStats> StatsManager, const TSharedPtr<FShaderPlatformSettings> Platform, const EMaterialQualityLevel::Type QualityLevel)
+void FStatsGridRow_Samplers::AddPlatform(TSharedPtr<FMaterialStats> StatsManager, const TSharedPtr<FShaderPlatformSettings> Platform, const EMaterialQualityLevel::Type QualityLevel, const int32 InstanceIndex)
 {
 	// cell that will enumerate the total number of samplers in this material
 	const FString CellContent = FMaterialStatsUtils::MaterialQualityToShortString(QualityLevel);
-	TSharedPtr<FGridCell_ShaderValue> Cell = MakeShareable(new FGridCell_ShaderValue(StatsManager, EShaderInfoType::SamplersCount, ERepresentativeShader::Num, QualityLevel, Platform->GetPlatformShaderType()));
+	TSharedPtr<FGridCell_ShaderValue> Cell = MakeShareable(new FGridCell_ShaderValue(StatsManager, EShaderInfoType::SamplersCount, ERepresentativeShader::Num, QualityLevel, Platform->GetPlatformShaderType(), InstanceIndex));
 
-	const FName ColumnName = FMaterialStatsGrid::MakePlatformColumnName(Platform, QualityLevel);
+	const FName ColumnName = FMaterialStatsGrid::MakePlatformColumnName(Platform, QualityLevel, InstanceIndex);
 	AddCell(ColumnName, Cell);
 }
 
@@ -271,13 +268,13 @@ void FStatsGridRow_Interpolators::CreateRow(TSharedPtr<FMaterialStats> StatsMana
 	FillPlatformCellsHelper(StatsManager);
 }
 
-void FStatsGridRow_Interpolators::AddPlatform(TSharedPtr<FMaterialStats> StatsManager, const TSharedPtr<FShaderPlatformSettings> Platform, const EMaterialQualityLevel::Type QualityLevel)
+void FStatsGridRow_Interpolators::AddPlatform(TSharedPtr<FMaterialStats> StatsManager, const TSharedPtr<FShaderPlatformSettings> Platform, const EMaterialQualityLevel::Type QualityLevel, const int32 InstanceIndex)
 {
 	// cell that will enumerate the total number of interpolators in this material
 	const FString CellContent = FMaterialStatsUtils::MaterialQualityToShortString(QualityLevel);
-	TSharedPtr<FGridCell_ShaderValue> Cell = MakeShareable(new FGridCell_ShaderValue(StatsManager, EShaderInfoType::InterpolatorsCount, ERepresentativeShader::Num, QualityLevel, Platform->GetPlatformShaderType()));
+	TSharedPtr<FGridCell_ShaderValue> Cell = MakeShareable(new FGridCell_ShaderValue(StatsManager, EShaderInfoType::InterpolatorsCount, ERepresentativeShader::Num, QualityLevel, Platform->GetPlatformShaderType(), InstanceIndex));
 
-	const FName ColumnName = FMaterialStatsGrid::MakePlatformColumnName(Platform, QualityLevel);
+	const FName ColumnName = FMaterialStatsGrid::MakePlatformColumnName(Platform, QualityLevel, InstanceIndex);
 	AddCell(ColumnName, Cell);
 }
 
@@ -300,13 +297,13 @@ void FStatsGridRow_NumTextureSamples::CreateRow(TSharedPtr<FMaterialStats> Stats
 	FillPlatformCellsHelper(StatsManager);
 }
 
-void FStatsGridRow_NumTextureSamples::AddPlatform(TSharedPtr<FMaterialStats> StatsManager, const TSharedPtr<FShaderPlatformSettings> Platform, const EMaterialQualityLevel::Type QualityLevel)
+void FStatsGridRow_NumTextureSamples::AddPlatform(TSharedPtr<FMaterialStats> StatsManager, const TSharedPtr<FShaderPlatformSettings> Platform, const EMaterialQualityLevel::Type QualityLevel, const int32 InstanceIndex)
 {
 	// cell that will enumerate the total number of texture samples in this material
 	const FString CellContent = FMaterialStatsUtils::MaterialQualityToShortString(QualityLevel);
-	TSharedPtr<FGridCell_ShaderValue> Cell = MakeShareable(new FGridCell_ShaderValue(StatsManager, EShaderInfoType::TextureSampleCount, ERepresentativeShader::Num, QualityLevel, Platform->GetPlatformShaderType()));
+	TSharedPtr<FGridCell_ShaderValue> Cell = MakeShareable(new FGridCell_ShaderValue(StatsManager, EShaderInfoType::TextureSampleCount, ERepresentativeShader::Num, QualityLevel, Platform->GetPlatformShaderType(), InstanceIndex));
 
-	const FName ColumnName = FMaterialStatsGrid::MakePlatformColumnName(Platform, QualityLevel);
+	const FName ColumnName = FMaterialStatsGrid::MakePlatformColumnName(Platform, QualityLevel, InstanceIndex);
 	AddCell(ColumnName, Cell);
 }
 
@@ -329,13 +326,13 @@ void FStatsGridRow_NumVirtualTextureLookups::CreateRow(TSharedPtr<FMaterialStats
 	FillPlatformCellsHelper(StatsManager);
 }
 
-void FStatsGridRow_NumVirtualTextureLookups::AddPlatform(TSharedPtr<FMaterialStats> StatsManager, const TSharedPtr<FShaderPlatformSettings> Platform, const EMaterialQualityLevel::Type QualityLevel)
+void FStatsGridRow_NumVirtualTextureLookups::AddPlatform(TSharedPtr<FMaterialStats> StatsManager, const TSharedPtr<FShaderPlatformSettings> Platform, const EMaterialQualityLevel::Type QualityLevel, const int32 InstanceIndex)
 {
 	// cell that will enumerate the total number of texture samples in this material
 	const FString CellContent = FMaterialStatsUtils::MaterialQualityToShortString(QualityLevel);
-	TSharedPtr<FGridCell_ShaderValue> Cell = MakeShareable(new FGridCell_ShaderValue(StatsManager, EShaderInfoType::VirtualTextureLookupCount, ERepresentativeShader::Num, QualityLevel, Platform->GetPlatformShaderType()));
+	TSharedPtr<FGridCell_ShaderValue> Cell = MakeShareable(new FGridCell_ShaderValue(StatsManager, EShaderInfoType::VirtualTextureLookupCount, ERepresentativeShader::Num, QualityLevel, Platform->GetPlatformShaderType(), InstanceIndex));
 
-	const FName ColumnName = FMaterialStatsGrid::MakePlatformColumnName(Platform, QualityLevel);
+	const FName ColumnName = FMaterialStatsGrid::MakePlatformColumnName(Platform, QualityLevel, InstanceIndex);
 	AddCell(ColumnName, Cell);
 }
 
@@ -358,13 +355,13 @@ void FStatsGridRow_NumShaders::CreateRow(TSharedPtr<FMaterialStats> StatsManager
 	FillPlatformCellsHelper(StatsManager);
 }
 
-void FStatsGridRow_NumShaders::AddPlatform(TSharedPtr<FMaterialStats> StatsManager, const TSharedPtr<FShaderPlatformSettings> Platform, const EMaterialQualityLevel::Type QualityLevel)
+void FStatsGridRow_NumShaders::AddPlatform(TSharedPtr<FMaterialStats> StatsManager, const TSharedPtr<FShaderPlatformSettings> Platform, const EMaterialQualityLevel::Type QualityLevel, const int32 InstanceIndex)
 {
 	// cell that will enumerate the number of shaders
 	const FString CellContent = FMaterialStatsUtils::MaterialQualityToShortString(QualityLevel);
-	TSharedPtr<FGridCell_ShaderValue> Cell = MakeShareable(new FGridCell_ShaderValue(StatsManager, EShaderInfoType::ShaderCount, ERepresentativeShader::Num, QualityLevel, Platform->GetPlatformShaderType()));
+	TSharedPtr<FGridCell_ShaderValue> Cell = MakeShareable(new FGridCell_ShaderValue(StatsManager, EShaderInfoType::ShaderCount, ERepresentativeShader::Num, QualityLevel, Platform->GetPlatformShaderType(), InstanceIndex));
 
-	const FName ColumnName = FMaterialStatsGrid::MakePlatformColumnName(Platform, QualityLevel);
+	const FName ColumnName = FMaterialStatsGrid::MakePlatformColumnName(Platform, QualityLevel, InstanceIndex);
 	AddCell(ColumnName, Cell);
 }
 
@@ -417,13 +414,13 @@ TSharedPtr<FGridCell> FMaterialStatsGrid::GetCell(int32 RowID, FName ColumnName)
 	return MakeShareable(new FGridCell_Empty());
 }
 
-void FMaterialStatsGrid::CollectShaderInfo(const TSharedPtr<FShaderPlatformSettings>& PlatformPtr, const EMaterialQualityLevel::Type QualityLevel)
+void FMaterialStatsGrid::CollectShaderInfo(const TSharedPtr<FShaderPlatformSettings>& PlatformPtr, const EMaterialQualityLevel::Type QualityLevel, const int32 InstanceIndex)
 {
 	const FShaderPlatformSettings::FPlatformData& PlatformData = PlatformPtr->GetPlatformData(QualityLevel);
 
 	for (int32 i = 0; i < (int32)ERepresentativeShader::Num; ++i)
 	{
-		UsedShaders[i] |= PlatformData.ShaderStatsInfo.ShaderInstructionCount.Contains((ERepresentativeShader)i);
+		UsedShaders[i] |= PlatformData.Instances[InstanceIndex].ShaderStatsInfo.ShaderInstructionCount.Contains((ERepresentativeShader)i);
 	}
 }
 
@@ -454,7 +451,10 @@ void FMaterialStatsGrid::CollectShaderInfo()
 		{
 			EMaterialQualityLevel::Type QualitySetting = static_cast<EMaterialQualityLevel::Type>(q);
 
-			CollectShaderInfo(Platform, QualitySetting);
+			for (int32 InstanceIndex = 0; InstanceIndex < Platform->GetPlatformData(QualitySetting).Instances.Num(); ++InstanceIndex)
+			{
+				CollectShaderInfo(Platform, QualitySetting, InstanceIndex);
+			}
 		}
 	}
 }
@@ -508,15 +508,18 @@ void FMaterialStatsGrid::CheckForErrors()
 				EMaterialQualityLevel::Type QualitySetting = static_cast<EMaterialQualityLevel::Type>(q);
 				if (StatsManager->GetStatsQualityFlag(QualitySetting))
 				{
-					auto &Data = Platform->GetPlatformData(QualitySetting);
-
-					if (Data.ShaderStatsInfo.HasErrors())
+					for (int32 InstanceIndex = 0; InstanceIndex < Platform->GetPlatformData(QualitySetting).Instances.Num(); ++InstanceIndex)
 					{
-						PlatformErrorsType = PlatformErrorsType == EGlobalErrorsType::SpecificPlatformErrors ? EGlobalErrorsType::SpecificPlatformErrors : EGlobalErrorsType::GlobalPlatformErrors;
-					}
-					else
-					{
-						PlatformErrorsType = PlatformErrorsType == EGlobalErrorsType::NoErrors ? EGlobalErrorsType::NoErrors : EGlobalErrorsType::SpecificPlatformErrors;
+						auto &Data = Platform->GetInstanceData(QualitySetting, InstanceIndex);
+
+						if (Data.ShaderStatsInfo.HasErrors())
+						{
+							PlatformErrorsType = PlatformErrorsType == EGlobalErrorsType::SpecificPlatformErrors ? EGlobalErrorsType::SpecificPlatformErrors : EGlobalErrorsType::GlobalPlatformErrors;
+						}
+						else
+						{
+							PlatformErrorsType = PlatformErrorsType == EGlobalErrorsType::NoErrors ? EGlobalErrorsType::NoErrors : EGlobalErrorsType::SpecificPlatformErrors;
+						}
 					}
 				}
 			}
@@ -528,14 +531,9 @@ void FMaterialStatsGrid::BuildRowIds()
 {
 	RowIDs.Reset();
 
+	BuildKeyAndInsert(ERowType::Name);
 	BuildKeyAndInsert(ERowType::Quality);
 
-	// add errors row if at least one platform has issues
-	if (PlatformErrorsType != EGlobalErrorsType::NoErrors)
-	{
-		BuildKeyAndInsert(ERowType::Errors);
-	}
-
 	// add the rest of the rows only if there's at least one error free platform
 	if (PlatformErrorsType != EGlobalErrorsType::GlobalPlatformErrors)
 	{
@@ -569,13 +567,26 @@ void FMaterialStatsGrid::BuildRowIds()
 
 void FMaterialStatsGrid::OnShaderChanged()
 {
+	SCOPED_NAMED_EVENT_TEXT("FMaterialStatsGrid::OnShaderChanged", FColor::Red);
+
 	CollectShaderInfo();
+
 	BuildShaderRows();
+	BuildColumnInfo();
 	CheckForErrors();
 
 	BuildRowIds();
 }
 
+void FMaterialStatsGrid::OnColumnNumChanged()
+{
+	SCOPED_NAMED_EVENT_TEXT("FMaterialStatsGrid::OnColumnNumChanged", FColor::Red);
+
+	// nuke all rows to rebuild columns
+	BuildStaticRows();
+	OnShaderChanged();
+}
+
 void FMaterialStatsGrid::BuildStaticRows()
 {
 	TSharedPtr<FMaterialStats> StatsManager = StatsManagerWPtr.Pin();
@@ -591,14 +602,14 @@ void FMaterialStatsGrid::BuildStaticRows()
 		StaticRows.Add(ERowType::Empty, Row);
 	}
 	{
-		TSharedPtr<FStatsGridRow> Row = MakeShareable(new FStatsGridRow_Quality());
+		TSharedPtr<FStatsGridRow> Row = MakeShareable(new FStatsGridRow_Name());
 		Row->CreateRow(StatsManager);
-		StaticRows.Add(ERowType::Quality, Row);
+		StaticRows.Add(ERowType::Name, Row);
 	}
 	{
-		TSharedPtr<FStatsGridRow> Row = MakeShareable(new FStatsGridRow_Errors());
+		TSharedPtr<FStatsGridRow> Row = MakeShareable(new FStatsGridRow_Quality());
 		Row->CreateRow(StatsManager);
-		StaticRows.Add(ERowType::Errors, Row);
+		StaticRows.Add(ERowType::Quality, Row);
 	}
 	{
 		TSharedPtr<FStatsGridRow> Row = MakeShareable(new FStatsGridRow_Samplers());
@@ -669,6 +680,7 @@ void FMaterialStatsGrid::BuildShaderRows()
 
 void FMaterialStatsGrid::BuildColumnInfo()
 {
+	GridColumnContent.Empty();
 	GridColumnContent.Add(DescriptorColumnName, FColumnInfo());
 	GridColumnContent.Add(ShaderColumnName, FColumnInfo());
 
@@ -695,7 +707,14 @@ void FMaterialStatsGrid::BuildColumnInfo()
 				continue;
 			}
 
-			AddColumnInfo(PlatformPair.Value, QualityLevel);
+			const int32 NumInstances = PlatformPair.Value->GetPlatformData(QualityLevel).Instances.Num();
+			// skip all derived MI's if not asked to show stats as errors are shown separately
+			const int32 IterateForNum = StatsManager->GetMaterialStatsDerivedMIOption() == EMaterialStatsDerivedMIOption::ShowStats ? NumInstances : FMath::Min(1, NumInstances);  
+
+			for (int32 InstanceIndex = 0; InstanceIndex < IterateForNum; ++InstanceIndex)
+			{
+				AddColumnInfo(PlatformPair.Value, QualityLevel, InstanceIndex);
+			}
 		}
 	}
 }
@@ -712,7 +731,7 @@ void FMaterialStatsGrid::BuildGrid()
 	BuildRowIds();
 }
 
-void FMaterialStatsGrid::AddColumnInfo(TSharedPtr<FShaderPlatformSettings> PlatformPtr, const EMaterialQualityLevel::Type QualityLevel)
+void FMaterialStatsGrid::AddColumnInfo(TSharedPtr<FShaderPlatformSettings> PlatformPtr, const EMaterialQualityLevel::Type QualityLevel, const int32 InstanceIndex)
 {
 	FColumnInfo Info;
 
@@ -720,17 +739,17 @@ void FMaterialStatsGrid::AddColumnInfo(TSharedPtr<FShaderPlatformSettings> Platf
 	Info.Content = PlatformPtr->GetPlatformName().ToString();
 	Info.ContentLong = PlatformPtr->GetPlatformDescription();
 
-	const FName ColumnName = MakePlatformColumnName(PlatformPtr, QualityLevel);
+	const FName ColumnName = MakePlatformColumnName(PlatformPtr, QualityLevel, InstanceIndex);
 	GridColumnContent.Add(ColumnName, Info);
 }
 
-void FMaterialStatsGrid::RemoveColumnInfo(TSharedPtr<FShaderPlatformSettings> PlatformPtr, const EMaterialQualityLevel::Type QualityLevel)
+void FMaterialStatsGrid::RemoveColumnInfo(TSharedPtr<FShaderPlatformSettings> PlatformPtr, const EMaterialQualityLevel::Type QualityLevel, const int32 InstanceIndex)
 {
-	const FName ColumnName = MakePlatformColumnName(PlatformPtr, QualityLevel);
+	const FName ColumnName = MakePlatformColumnName(PlatformPtr, QualityLevel, InstanceIndex);
 	GridColumnContent.Remove(ColumnName);
 }
 
-void FMaterialStatsGrid::AddOrRemovePlatform(TSharedPtr<FShaderPlatformSettings> PlatformPtr, const bool bAdd, const EMaterialQualityLevel::Type QualityLevel)
+void FMaterialStatsGrid::AddOrRemovePlatform(TSharedPtr<FShaderPlatformSettings> PlatformPtr, const bool bAdd, const EMaterialQualityLevel::Type QualityLevel, const int32 InstanceIndex)
 {
 	TSharedPtr<FMaterialStats> StatsManager = StatsManagerWPtr.Pin();
 	if (!StatsManager.IsValid())
@@ -741,22 +760,22 @@ void FMaterialStatsGrid::AddOrRemovePlatform(TSharedPtr<FShaderPlatformSettings>
 	// update column record
 	if (bAdd)
 	{
-		AddColumnInfo(PlatformPtr, QualityLevel);
+		AddColumnInfo(PlatformPtr, QualityLevel, InstanceIndex);
 	}
 	else
 	{
-		RemoveColumnInfo(PlatformPtr, QualityLevel);
+		RemoveColumnInfo(PlatformPtr, QualityLevel, InstanceIndex);
 	}
 
 	for (auto RowPair : StaticRows)
 	{
 		if (bAdd)
 		{
-			RowPair.Value->AddPlatform(StatsManager, PlatformPtr, QualityLevel);
+			RowPair.Value->AddPlatform(StatsManager, PlatformPtr, QualityLevel, InstanceIndex);
 		}
 		else
 		{
-			RowPair.Value->RemovePlatform(StatsManager, PlatformPtr, QualityLevel);
+			RowPair.Value->RemovePlatform(StatsManager, PlatformPtr, QualityLevel, InstanceIndex);
 		}
 	}
 
@@ -764,11 +783,11 @@ void FMaterialStatsGrid::AddOrRemovePlatform(TSharedPtr<FShaderPlatformSettings>
 	{
 		if (bAdd)
 		{
-			VertexShaderRows[i]->AddPlatform(StatsManager, PlatformPtr, QualityLevel);
+			VertexShaderRows[i]->AddPlatform(StatsManager, PlatformPtr, QualityLevel, InstanceIndex);
 		}
 		else
 		{
-			VertexShaderRows[i]->RemovePlatform(StatsManager, PlatformPtr, QualityLevel);
+			VertexShaderRows[i]->RemovePlatform(StatsManager, PlatformPtr, QualityLevel, InstanceIndex);
 		}
 	}
 
@@ -776,11 +795,11 @@ void FMaterialStatsGrid::AddOrRemovePlatform(TSharedPtr<FShaderPlatformSettings>
 	{
 		if (bAdd)
 		{
-			FragmentShaderRows[i]->AddPlatform(StatsManager, PlatformPtr, QualityLevel);
+			FragmentShaderRows[i]->AddPlatform(StatsManager, PlatformPtr, QualityLevel, InstanceIndex);
 		}
 		else
 		{
-			FragmentShaderRows[i]->RemovePlatform(StatsManager, PlatformPtr, QualityLevel);
+			FragmentShaderRows[i]->RemovePlatform(StatsManager, PlatformPtr, QualityLevel, InstanceIndex);
 		}
 	}
 }
@@ -804,7 +823,10 @@ void FMaterialStatsGrid::OnAddOrRemovePlatform(TSharedPtr<FShaderPlatformSetting
 			continue;
 		}
 
-		AddOrRemovePlatform(PlatformPtr, bAdded, QualityLevel);
+		for (int32 InstanceIndex = 0; InstanceIndex < PlatformPtr->GetPlatformData(QualityLevel).Instances.Num(); ++InstanceIndex)
+		{
+			AddOrRemovePlatform(PlatformPtr, bAdded, QualityLevel, InstanceIndex);
+		}
 	}
 
 	// recheck shader rows in case something changed
@@ -824,7 +846,10 @@ void FMaterialStatsGrid::OnQualitySettingChanged(const EMaterialQualityLevel::Ty
 	{
 		if (PlatformPair.Value->IsPresentInGrid())
 		{
-			AddOrRemovePlatform(PlatformPair.Value, bQualityOn, QualityLevel);
+			for (int32 InstanceIndex = 0; InstanceIndex < PlatformPair.Value->GetPlatformData(QualityLevel).Instances.Num(); ++InstanceIndex)
+			{
+				AddOrRemovePlatform(PlatformPair.Value, bQualityOn, QualityLevel, InstanceIndex);
+			}
 		}
 	}
 
@@ -832,9 +857,9 @@ void FMaterialStatsGrid::OnQualitySettingChanged(const EMaterialQualityLevel::Ty
 	OnShaderChanged();
 }
 
-FName FMaterialStatsGrid::MakePlatformColumnName(const TSharedPtr<FShaderPlatformSettings>& Platform, const EMaterialQualityLevel::Type Quality)
+FName FMaterialStatsGrid::MakePlatformColumnName(const TSharedPtr<FShaderPlatformSettings>& Platform, const EMaterialQualityLevel::Type Quality, const int32 InstanceIndex)
 {
-	FName RetName = *(Platform->GetPlatformID().ToString() + "_" + FMaterialStatsUtils::MaterialQualityToString(Quality));
+	FName RetName = *(Platform->GetPlatformID().ToString() + "_" + FMaterialStatsUtils::MaterialQualityToString(Quality) + "_" + FString::FromInt(InstanceIndex));
 	return RetName;
 }
 
@@ -876,7 +901,7 @@ FString FGridCell_StaticString::GetCellContentLong()
 }
 
 FGridCell_ShaderValue::FGridCell_ShaderValue(const TWeakPtr<FMaterialStats>& _MaterialStatsWPtr, const EShaderInfoType _InfoType, const ERepresentativeShader _ShaderType,
-	const EMaterialQualityLevel::Type _QualityLevel, const EShaderPlatform _PlatformType)
+	const EMaterialQualityLevel::Type _QualityLevel, const EShaderPlatform _PlatformType, const int32 _InstanceIndex)
 {
 	MaterialStatsWPtr = _MaterialStatsWPtr;
 
@@ -884,6 +909,7 @@ FGridCell_ShaderValue::FGridCell_ShaderValue(const TWeakPtr<FMaterialStats>& _Ma
 	ShaderType = _ShaderType;
 	QualityLevel = _QualityLevel;
 	PlatformType = _PlatformType;
+	InstanceIndex = _InstanceIndex;
 }
 
 FString FGridCell_ShaderValue::InternalGetContent(bool bLongContent)
@@ -895,22 +921,22 @@ FString FGridCell_ShaderValue::InternalGetContent(bool bLongContent)
 	}
 
 	auto Platform = MaterialStats->GetPlatformSettings(PlatformType);
-	if (!Platform.IsValid())
+	if (!Platform.IsValid() || Platform->GetPlatformData(QualityLevel).Instances.Num() <= InstanceIndex)
 	{
 		return TEXT("");
 	}
 
-	auto PlatformData = Platform->GetPlatformData(QualityLevel);
+	const auto& InstanceData = Platform->GetInstanceData(QualityLevel, InstanceIndex);
 
 	switch (InfoType)
 	{
-		case EShaderInfoType::Errors:
-			return PlatformData.ShaderStatsInfo.StrShaderErrors;
+		case EShaderInfoType::Name:
+			return MaterialStats->GetMaterialName(InstanceIndex);
 		break;
 
 		case EShaderInfoType::InstructionsCount:
 		{
-			auto* Count = PlatformData.ShaderStatsInfo.ShaderInstructionCount.Find(ShaderType);
+			auto* Count = InstanceData.ShaderStatsInfo.ShaderInstructionCount.Find(ShaderType);
 			if (Count)
 			{
 				return bLongContent ? Count->StrDescriptionLong : Count->StrDescription;
@@ -919,23 +945,23 @@ FString FGridCell_ShaderValue::InternalGetContent(bool bLongContent)
 		break;
 
 		case EShaderInfoType::InterpolatorsCount:
-			return bLongContent ? PlatformData.ShaderStatsInfo.InterpolatorsCount.StrDescriptionLong : PlatformData.ShaderStatsInfo.InterpolatorsCount.StrDescription;
+			return bLongContent ? InstanceData.ShaderStatsInfo.InterpolatorsCount.StrDescriptionLong : InstanceData.ShaderStatsInfo.InterpolatorsCount.StrDescription;
 		break;
 
 		case EShaderInfoType::TextureSampleCount:
-			return bLongContent ? PlatformData.ShaderStatsInfo.TextureSampleCount.StrDescriptionLong : PlatformData.ShaderStatsInfo.TextureSampleCount.StrDescription;
+			return bLongContent ? InstanceData.ShaderStatsInfo.TextureSampleCount.StrDescriptionLong : InstanceData.ShaderStatsInfo.TextureSampleCount.StrDescription;
 		break;
 
 		case EShaderInfoType::SamplersCount:
-			return bLongContent ? PlatformData.ShaderStatsInfo.SamplersCount.StrDescriptionLong : PlatformData.ShaderStatsInfo.SamplersCount.StrDescription;
+			return bLongContent ? InstanceData.ShaderStatsInfo.SamplersCount.StrDescriptionLong : InstanceData.ShaderStatsInfo.SamplersCount.StrDescription;
 		break;
 
 		case EShaderInfoType::VirtualTextureLookupCount:
-			return bLongContent ? PlatformData.ShaderStatsInfo.VirtualTextureLookupCount.StrDescriptionLong : PlatformData.ShaderStatsInfo.VirtualTextureLookupCount.StrDescription;
+			return bLongContent ? InstanceData.ShaderStatsInfo.VirtualTextureLookupCount.StrDescriptionLong : InstanceData.ShaderStatsInfo.VirtualTextureLookupCount.StrDescription;
 		break;
 
 		case EShaderInfoType::ShaderCount:
-			return bLongContent ? PlatformData.ShaderStatsInfo.ShaderCount.StrDescriptionLong : PlatformData.ShaderStatsInfo.ShaderCount.StrDescription;
+			return bLongContent ? InstanceData.ShaderStatsInfo.ShaderCount.StrDescriptionLong : InstanceData.ShaderStatsInfo.ShaderCount.StrDescription;
 		break;
 	}
 
@@ -951,3 +977,27 @@ FString FGridCell_ShaderValue::GetCellContentLong()
 {
 	return InternalGetContent(true);
 }
+
+FGridCell::EIcon FGridCell_ShaderValue::GetIcon() const
+{
+	auto MaterialStats = MaterialStatsWPtr.Pin();
+	if (!MaterialStats.IsValid())
+	{
+		return EIcon::None;
+	}
+
+	auto Platform = MaterialStats->GetPlatformSettings(PlatformType);
+	if (!Platform.IsValid() || Platform->GetPlatformData(QualityLevel).Instances.Num() <= InstanceIndex)
+	{
+		return EIcon::None;
+	}
+
+	const auto& InstanceData = Platform->GetInstanceData(QualityLevel, InstanceIndex);
+	switch (InfoType)
+	{
+		case EShaderInfoType::Name:
+			return InstanceData.MaterialResourcesStats->GetCompileErrors().Num() > 0 ? EIcon::Error : EIcon::None;
+		default:
+			return EIcon::None;
+	}
+}
diff --git a/Engine/Source/Editor/MaterialEditor/Private/MaterialStatsGrid.h b/Engine/Source/Editor/MaterialEditor/Private/MaterialStatsGrid.h
--- a/Engine/Source/Editor/MaterialEditor/Private/MaterialStatsGrid.h
+++ b/Engine/Source/Editor/MaterialEditor/Private/MaterialStatsGrid.h
@@ -14,6 +14,13 @@
 /** class to represent a single cell inside the material stats grid */
 class FGridCell
 {
+public:
+	enum class EIcon
+	{
+		None,
+		Error
+	};
+
 protected:
 	/** attributes used at display time to configure widgets */
 	FSlateColor CellColor;
@@ -42,6 +49,11 @@ public:
 
 	FORCEINLINE EVerticalAlignment GetVerticalAlignment() const;
 	FORCEINLINE void SetVerticalAlignment(EVerticalAlignment Align);
+
+	virtual EIcon GetIcon() const
+	{
+		return EIcon::None;
+	}
 };
 
 /** this time of cell with just return an empty string and its mainly used to separate rows */
@@ -68,7 +80,7 @@ public:
 /** enumeration used to classify arguments for FGridCell_ShaderValue */
 enum class EShaderInfoType
 {
-	Errors,
+	Name,
 	InstructionsCount,
 	SamplersCount,
 	InterpolatorsCount,
@@ -87,15 +99,17 @@ private:
 	ERepresentativeShader ShaderType;
 	EMaterialQualityLevel::Type QualityLevel;
 	EShaderPlatform PlatformType;
+	int32 InstanceIndex;
 
 	FString InternalGetContent(bool bLongContent);
 
 public:
 	FGridCell_ShaderValue(const TWeakPtr<FMaterialStats>& _MaterialStatsWPtr, const EShaderInfoType _InfoType, const ERepresentativeShader _ShaderType,
-		const EMaterialQualityLevel::Type _QualityLevel, const EShaderPlatform _PlatformType);
+		const EMaterialQualityLevel::Type _QualityLevel, const EShaderPlatform _PlatformType, const int32 _InstanceIndex);
 
 	FString GetCellContent() override;
 	FString GetCellContentLong() override;
+	EIcon GetIcon() const override;
 };
 
 /** virtual class to model grid row generation */
@@ -120,10 +134,15 @@ public:
 	virtual void CreateRow(TSharedPtr<FMaterialStats> StatsManager) = 0;
 
 	/** Add/RemovePlatforms should be called when a platform is added or removed from the grid */
-	virtual void AddPlatform(TSharedPtr<FMaterialStats> StatsManager, const TSharedPtr<FShaderPlatformSettings> PlatformPtr, const EMaterialQualityLevel::Type QualityLevel) = 0;
-	virtual void RemovePlatform(TSharedPtr<FMaterialStats> StatsManager, const TSharedPtr<FShaderPlatformSettings> PlatformPtr, const EMaterialQualityLevel::Type QualityLevel);
+	virtual void AddPlatform(TSharedPtr<FMaterialStats> StatsManager, const TSharedPtr<FShaderPlatformSettings> PlatformPtr, const EMaterialQualityLevel::Type QualityLevel, const int32 InstanceIndex) = 0;
+	virtual void RemovePlatform(TSharedPtr<FMaterialStats> StatsManager, const TSharedPtr<FShaderPlatformSettings> PlatformPtr, const EMaterialQualityLevel::Type QualityLevel, const int32 InstanceIndex);
 
 	TSharedPtr<FGridCell> GetCell(const FName ColumnName);
+
+	void RemoveAll()
+	{
+		RowCells.Empty();
+	}
 };
 
 /** separator row */
@@ -132,25 +151,25 @@ class FStatsGridRow_Empty : public FStatsGridRow
 public:
 	void CreateRow(TSharedPtr<FMaterialStats> StatsManager) override;
 
-	void AddPlatform(TSharedPtr<FMaterialStats> StatsManager, const TSharedPtr<FShaderPlatformSettings> Platform, const EMaterialQualityLevel::Type QualityLevel) override;
+	void AddPlatform(TSharedPtr<FMaterialStats> StatsManager, const TSharedPtr<FShaderPlatformSettings> Platform, const EMaterialQualityLevel::Type QualityLevel, const int32 InstanceIndex) override;
 };
 
-/** row that will produce static string from EMaterialQualityLevel::Type */
-class FStatsGridRow_Quality : public FStatsGridRow
+/** material/instance name */
+class FStatsGridRow_Name : public FStatsGridRow
 {
 public:
 	void CreateRow(TSharedPtr<FMaterialStats> StatsManager) override;
 
-	void AddPlatform(TSharedPtr<FMaterialStats> StatsManager, const TSharedPtr<FShaderPlatformSettings> Platform, const EMaterialQualityLevel::Type QualityLevel) override;
+	void AddPlatform(TSharedPtr<FMaterialStats> StatsManager, const TSharedPtr<FShaderPlatformSettings> Platform, const EMaterialQualityLevel::Type QualityLevel, const int32 InstanceIndex) override;
 };
 
-/** row that will display eventual shader errors */
-class FStatsGridRow_Errors : public FStatsGridRow
+/** row that will produce static string from EMaterialQualityLevel::Type */
+class FStatsGridRow_Quality : public FStatsGridRow
 {
 public:
 	void CreateRow(TSharedPtr<FMaterialStats> StatsManager) override;
 
-	void AddPlatform(TSharedPtr<FMaterialStats> StatsManager, const TSharedPtr<FShaderPlatformSettings> Platform, const EMaterialQualityLevel::Type QualityLevel) override;
+	void AddPlatform(TSharedPtr<FMaterialStats> StatsManager, const TSharedPtr<FShaderPlatformSettings> Platform, const EMaterialQualityLevel::Type QualityLevel, const int32 InstanceIndex) override;
 };
 
 /** row that will extract and the number of instructions for each used shader */
@@ -177,7 +196,7 @@ public:
 
 	void CreateRow(TSharedPtr<FMaterialStats> StatsManager) override;
 
-	void AddPlatform(TSharedPtr<FMaterialStats> StatsManager, const TSharedPtr<FShaderPlatformSettings> Platform, const EMaterialQualityLevel::Type QualityLevel) override;
+	void AddPlatform(TSharedPtr<FMaterialStats> StatsManager, const TSharedPtr<FShaderPlatformSettings> Platform, const EMaterialQualityLevel::Type QualityLevel, const int32 InstanceIndex) override;
 };
 
 /** this row will display the global number of samplers present in the material for a specified platform */
@@ -186,7 +205,7 @@ class FStatsGridRow_Samplers : public FStatsGridRow
 public:
 	void CreateRow(TSharedPtr<FMaterialStats> StatsManager) override;
 
-	void AddPlatform(TSharedPtr<FMaterialStats> StatsManager, const TSharedPtr<FShaderPlatformSettings> Platform, const EMaterialQualityLevel::Type QualityLevel) override;
+	void AddPlatform(TSharedPtr<FMaterialStats> StatsManager, const TSharedPtr<FShaderPlatformSettings> Platform, const EMaterialQualityLevel::Type QualityLevel, const int32 InstanceIndex) override;
 };
 
 /** this row will display the global number of interpolators present in the material for a specified platform */
@@ -195,7 +214,7 @@ class FStatsGridRow_Interpolators : public FStatsGridRow
 public:
 	void CreateRow(TSharedPtr<FMaterialStats> StatsManager) override;
 
-	void AddPlatform(TSharedPtr<FMaterialStats> StatsManager, const TSharedPtr<FShaderPlatformSettings> Platform, const EMaterialQualityLevel::Type QualityLevel) override;
+	void AddPlatform(TSharedPtr<FMaterialStats> StatsManager, const TSharedPtr<FShaderPlatformSettings> Platform, const EMaterialQualityLevel::Type QualityLevel, const int32 InstanceIndex) override;
 };
 
 /** this row will display the global number of texture samples present in the material for a specified platform */
@@ -204,7 +223,7 @@ class FStatsGridRow_NumTextureSamples : public FStatsGridRow
 public:
 	void CreateRow(TSharedPtr<FMaterialStats> StatsManager) override;
 
-	void AddPlatform(TSharedPtr<FMaterialStats> StatsManager, const TSharedPtr<FShaderPlatformSettings> Platform, const EMaterialQualityLevel::Type QualityLevel) override;
+	void AddPlatform(TSharedPtr<FMaterialStats> StatsManager, const TSharedPtr<FShaderPlatformSettings> Platform, const EMaterialQualityLevel::Type QualityLevel, const int32 InstanceIndex) override;
 };
 
 /** this row will display the global number of virtual texture lookups present in the material for a specified platform */
@@ -213,7 +232,7 @@ class FStatsGridRow_NumVirtualTextureLookups : public FStatsGridRow
 public:
 	void CreateRow(TSharedPtr<FMaterialStats> StatsManager) override;
 
-	void AddPlatform(TSharedPtr<FMaterialStats> StatsManager, const TSharedPtr<FShaderPlatformSettings> Platform, const EMaterialQualityLevel::Type QualityLevel) override;
+	void AddPlatform(TSharedPtr<FMaterialStats> StatsManager, const TSharedPtr<FShaderPlatformSettings> Platform, const EMaterialQualityLevel::Type QualityLevel, const int32 InstanceIndex) override;
 };
 
 /** this row will display the total number of shaders present in the material for a specified platform */
@@ -222,7 +241,7 @@ class FStatsGridRow_NumShaders : public FStatsGridRow
 public:
 	void CreateRow(TSharedPtr<FMaterialStats> StatsManager) override;
 
-	void AddPlatform(TSharedPtr<FMaterialStats> StatsManager, const TSharedPtr<FShaderPlatformSettings> Platform, const EMaterialQualityLevel::Type QualityLevel) override;
+	void AddPlatform(TSharedPtr<FMaterialStats> StatsManager, const TSharedPtr<FShaderPlatformSettings> Platform, const EMaterialQualityLevel::Type QualityLevel, const int32 InstanceIndex) override;
 };
 
 /** class that models the logical material stats grid */
@@ -243,8 +262,8 @@ class FMaterialStatsGrid
 	enum class ERowType
 	{
 		Empty,
+		Name,
 		Quality,
-		Errors,
 		Samplers,
 		Interpolators,
 		TextureSamples,
@@ -293,8 +312,8 @@ public:
 	const static FName ShaderColumnName;
 
 private:
-	void AddColumnInfo(TSharedPtr<FShaderPlatformSettings> PlatformPtr, const EMaterialQualityLevel::Type QualityLevel);
-	void RemoveColumnInfo(TSharedPtr<FShaderPlatformSettings> PlatformPtr, const EMaterialQualityLevel::Type QualityLevel);
+	void AddColumnInfo(TSharedPtr<FShaderPlatformSettings> PlatformPtr, const EMaterialQualityLevel::Type QualityLevel, const int32 InstanceIndex);
+	void RemoveColumnInfo(TSharedPtr<FShaderPlatformSettings> PlatformPtr, const EMaterialQualityLevel::Type QualityLevel, const int32 InstanceIndex);
 
 	/** this function will go through all available shader platforms and will call build the content of GridColumnContent */
 	void BuildColumnInfo();
@@ -307,7 +326,7 @@ private:
 	void CheckForErrors();
 
 	/** this functions will build the content of UsedShaders array */
-	void CollectShaderInfo(const TSharedPtr<FShaderPlatformSettings>& PlatformPtr, const EMaterialQualityLevel::Type QualityLevel);
+	void CollectShaderInfo(const TSharedPtr<FShaderPlatformSettings>& PlatformPtr, const EMaterialQualityLevel::Type QualityLevel, const int32 InstanceIndex);
 	void CollectShaderInfo();
 
 	/** functions that build the row ids array used by GridStatsWidget to identify available rows inside this logical grid */
@@ -318,7 +337,7 @@ private:
 	FORCEINLINE int32 AssembleRowKey(const ERowType RowType, const int16 Index);
 	FORCEINLINE void DissasambleRowKey(ERowType& RowType, int32& Index, const int32 Key);
 
-	void AddOrRemovePlatform(TSharedPtr<FShaderPlatformSettings> PlatformPtr, const bool bAdd, const EMaterialQualityLevel::Type QualityLevel);
+	void AddOrRemovePlatform(TSharedPtr<FShaderPlatformSettings> PlatformPtr, const bool bAdd, const EMaterialQualityLevel::Type QualityLevel, const int32 InstanceIndex);
 
 public:
 	FMaterialStatsGrid(TWeakPtr<FMaterialStats> _StatsManager);
@@ -331,6 +350,7 @@ public:
 	FORCEINLINE TArray<FName> GetVisibleColumnNames() const;
 
 	void OnShaderChanged();
+	void OnColumnNumChanged();
 
 	void OnAddOrRemovePlatform(TSharedPtr<FShaderPlatformSettings> PlatformPtr);
 	void OnQualitySettingChanged(const EMaterialQualityLevel::Type QualityLevel);
@@ -343,7 +363,7 @@ public:
 	FSlateColor GetColumnColor(const FName ColumnName) const;
 
 	/** helper function that will assemble a column name from the given arguments */
-	static FName MakePlatformColumnName(const TSharedPtr<FShaderPlatformSettings>& Platform, const EMaterialQualityLevel::Type Quality);
+	static FName MakePlatformColumnName(const TSharedPtr<FShaderPlatformSettings>& Platform, const EMaterialQualityLevel::Type Quality, const int32 InstanceIndex);
 };
 
 //////////////////////////////////////////////////////////////////////////////////////////////////
@@ -412,4 +432,4 @@ FORCEINLINE EVerticalAlignment FGridCell::GetVerticalAlignment() const
 FORCEINLINE void FGridCell::SetVerticalAlignment(EVerticalAlignment Align)
 {
 	VAlignment = Align;
-}
\ No newline at end of file
+}
diff --git a/Engine/Source/Editor/MaterialEditor/Private/SMaterialEditorStatsWidget.cpp b/Engine/Source/Editor/MaterialEditor/Private/SMaterialEditorStatsWidget.cpp
--- a/Engine/Source/Editor/MaterialEditor/Private/SMaterialEditorStatsWidget.cpp
+++ b/Engine/Source/Editor/MaterialEditor/Private/SMaterialEditorStatsWidget.cpp
@@ -8,6 +8,7 @@
 #include "Widgets/Views/SListView.h"
 #include "Widgets/Input/SCheckBox.h"
 #include "Widgets/Input/SComboButton.h"
+#include "Widgets/Input/SComboBox.h"
 #include "Widgets/Images/SImage.h"
 #include "Framework/MultiBox/MultiBoxBuilder.h"
 #include "Fonts/FontMeasure.h"
@@ -63,12 +64,24 @@ private:
 	TWeakPtr<FMaterialStats> MaterialStatsWPtr;
 };
 
+FName GetGridCellIconBrushName(FGridCell::EIcon Icon)
+{
+	switch(Icon)
+	{
+	case FGridCell::EIcon::Error:
+		return "MessageLog.Error";
+	default:
+		return "";
+	}
+}
+
 TSharedRef<SWidget> SMaterialStatsViewRow::GenerateWidgetForColumn(const FName& ColumnName)
 {
 	EHorizontalAlignment HAlign = EHorizontalAlignment::HAlign_Fill;
 	EVerticalAlignment VALign = EVerticalAlignment::VAlign_Top;
+	FGridCell::EIcon Icon = FGridCell::EIcon::None;
 
-	FName UsedFontStyle = SMaterialEditorStatsWidget::GetRegulatFontStyleName();
+	FName UsedFontStyle = SMaterialEditorStatsWidget::GetRegularFontStyleName();
 
 	auto StatsPtr = MaterialStatsWPtr.Pin();
 	if (StatsPtr.IsValid() && PtrRowID.IsValid())
@@ -77,22 +90,53 @@ TSharedRef<SWidget> SMaterialStatsViewRow::GenerateWidgetForColumn(const FName&
 
 		auto Cell = StatsPtr->GetStatsGrid()->GetCell(RowID, ColumnName);
 
-		UsedFontStyle = Cell->IsContentBold() ? SMaterialEditorStatsWidget::GetBoldFontStyleName() : SMaterialEditorStatsWidget::GetRegulatFontStyleName();
+		UsedFontStyle = Cell->IsContentBold() ? SMaterialEditorStatsWidget::GetBoldFontStyleName() : SMaterialEditorStatsWidget::GetRegularFontStyleName();
 		HAlign = Cell->GetHorizontalAlignment();
 		VALign = Cell->GetVerticalAlignment();
+		Icon = Cell->GetIcon();
 	}
 
-	return SNew(SBox)
-		.Padding(FMargin(4, 2, 4, 2))
-		.HAlign(HAlign)
-		.VAlign(VALign)
-		[
-			SNew(STextBlock)
-			.TextStyle(FAppStyle::Get(), UsedFontStyle)
-			.Text(this, &SMaterialStatsViewRow::GetTextForCell, ColumnName, false)
-			.ToolTipText(this, &SMaterialStatsViewRow::GetTextForCell, ColumnName, true)
-			.AutoWrapText(true)
-		];
+	if (Icon != FGridCell::EIcon::None)
+	{
+		return SNew(SBox)
+			.Padding(FMargin(4, 2, 4, 2))
+			.HAlign(HAlign)
+			.VAlign(VALign)
+			[
+				SNew(SHorizontalBox)
+				+ SHorizontalBox::Slot()
+				.AutoWidth()
+				.VAlign(VAlign_Center)
+				[
+					SNew(SImage)
+					.Image(FAppStyle::GetBrush(GetGridCellIconBrushName(Icon)))
+				]
+				+ SHorizontalBox::Slot()
+				.AutoWidth()
+				.VAlign(VAlign_Center)
+				[
+					SNew(STextBlock)
+					.TextStyle(FAppStyle::Get(), UsedFontStyle)
+					.Text(this, &SMaterialStatsViewRow::GetTextForCell, ColumnName, false)
+					.ToolTipText(this, &SMaterialStatsViewRow::GetTextForCell, ColumnName, true)
+					.AutoWrapText(true)
+				]
+			];
+	}
+	else
+	{
+		return SNew(SBox)
+			.Padding(FMargin(4, 2, 4, 2))
+			.HAlign(HAlign)
+			.VAlign(VALign)
+			[
+				SNew(STextBlock)
+				.TextStyle(FAppStyle::Get(), UsedFontStyle)
+				.Text(this, &SMaterialStatsViewRow::GetTextForCell, ColumnName, false)
+				.ToolTipText(this, &SMaterialStatsViewRow::GetTextForCell, ColumnName, true)
+				.AutoWrapText(true)
+			];
+	}
 }
 
 FText SMaterialStatsViewRow::GetTextForCell(const FName Name, const bool bToolTip) const
@@ -183,10 +227,12 @@ float SMaterialEditorStatsWidget::GetColumnSize(const FName ColumnName) const
 					const auto Cell = StatsPtr->GetStatsGrid()->GetCell(*ArrRowIds[i], ColumnName);
 
 					const FString Content = Cell->GetCellContent();
+					const FGridCell::EIcon Icon = Cell->GetIcon();
 
 					FVector2D FontMeasure = FontMeasureService->Measure(Content, FontInfo);
+					const float IconSize = Icon != FGridCell::EIcon::None ? FAppStyle::GetBrush(GetGridCellIconBrushName(Icon))->GetImageSize().X : 0.0f;
 
-					ColumnSize = FMath::Clamp(FontMeasure.X, ColumnSize, ColumnSizeExtraLarge);
+					ColumnSize = FMath::Clamp(FontMeasure.X + IconSize, ColumnSize, ColumnSizeExtraLarge);
 				}
 			}
 		}
@@ -359,25 +405,8 @@ void SMaterialEditorStatsWidget::CreatePlatformMenus(FMenuBuilder& Builder, EPla
 			auto MaterialStats = WidgetPtr->MaterialStatsWPtr.Pin();
 			if (PlatformPtr.IsValid() && MaterialStats.IsValid())
 			{
-				const bool bSwitchValue = MaterialStats->SwitchShaderPlatformUseStats(PlatformPtr->GetPlatformShaderType());
-
-				for (int32 q = 0; q < EMaterialQualityLevel::Num; ++q)
-				{
-					EMaterialQualityLevel::Type QualityLevel = static_cast<EMaterialQualityLevel::Type>(q);
-					if (MaterialStats->GetStatsQualityFlag(QualityLevel))
-					{
-						const FName PlatformColumnName = FMaterialStatsGrid::MakePlatformColumnName(PlatformPtr, QualityLevel);
-
-						if (bSwitchValue)
-						{
-							WidgetPtr->AddColumn(PlatformColumnName);
-						}
-						else
-						{
-							WidgetPtr->RemoveColumn(PlatformColumnName);
-						}
-					}
-				}
+				MaterialStats->SwitchShaderPlatformUseStats(PlatformPtr->GetPlatformShaderType());
+				WidgetPtr->OnColumnNumChanged();
 			}
 
 			WidgetPtr->RequestRefresh();
@@ -418,48 +447,11 @@ void SMaterialEditorStatsWidget::OnFlipQualityState(const ECheckBoxState NewStat
 	auto StatsPtr = MaterialStatsWPtr.Pin();
 	if (StatsPtr.IsValid())
 	{
-		bool bSwitchValue = StatsPtr->SwitchStatsQualityFlag(QualityLevel);
+		StatsPtr->SwitchStatsQualityFlag(QualityLevel);
 
 		StatsPtr->GetStatsGrid()->OnQualitySettingChanged(QualityLevel);
 
-		auto& PlatformDB = StatsPtr->GetPlatformsDB();
-		for (auto Pair : PlatformDB)
-		{
-			TSharedPtr<FShaderPlatformSettings> Platform = Pair.Value;
-			if (Platform->IsPresentInGrid())
-			{
-				const FName ColumnName = FMaterialStatsGrid::MakePlatformColumnName(Platform, QualityLevel);
-
-				if (bSwitchValue)
-				{
-					// find insert spot, right after the a column used by the same platform, with other quality settings
-					EMaterialQualityLevel::Type InsertAfterQuality = QualityLevel;
-
-					for (int32 i = 0; i < EMaterialQualityLevel::Num; ++i)
-					{
-						if (StatsPtr->GetStatsQualityFlag((EMaterialQualityLevel::Type)i) && i != QualityLevel)
-						{
-							InsertAfterQuality = (EMaterialQualityLevel::Type)i;
-							break;
-						}
-					}
-
-					if (QualityLevel != InsertAfterQuality)
-					{
-						const FName PreviousColumnName = FMaterialStatsGrid::MakePlatformColumnName(Platform, InsertAfterQuality);
-						InsertColumnAfter(ColumnName, PreviousColumnName);
-					}
-					else
-					{
-						AddColumn(ColumnName);
-					}
-				}
-				else
-				{
-					RemoveColumn(ColumnName);
-				}
-			}
-		}
+		OnColumnNumChanged();
 
 		RequestRefresh();
 	}
@@ -497,6 +489,64 @@ void SMaterialEditorStatsWidget::CreateQualityMenus(FMenuBuilder& Builder)
 	}
 }
 
+FText SMaterialEditorStatsWidget::MaterialStatsDerivedMIOptionToDescription(const EMaterialStatsDerivedMIOption Option)
+{
+	static_assert(static_cast<int32>(EMaterialStatsDerivedMIOption::InvalidOrMax) == 3, "Not all cases are handled in switch below!?");
+	switch (Option)
+	{
+	case EMaterialStatsDerivedMIOption::Ignore:			return LOCTEXT("MaterialPlatformStats_IgnoreMIs", "Ignore derived material instances");
+	case EMaterialStatsDerivedMIOption::CompileOnly:	return LOCTEXT("MaterialPlatformStats_CompileMIs", "Compile derived material instances");
+	case EMaterialStatsDerivedMIOption::ShowStats:		return LOCTEXT("MaterialPlatformStats_ShowMIs", "Show stats for derived material instances");
+	default:											return LOCTEXT("InvalidOrMax", "InvalidOrMax");
+	}
+}
+
+void SMaterialEditorStatsWidget::CreateDerivedMaterialsMenu(class FMenuBuilder& Builder)
+{
+	DerivedMaterialInstancesComboBoxItems.SetNum(static_cast<int32>(EMaterialStatsDerivedMIOption::InvalidOrMax));
+	auto AddOption = [&](EMaterialStatsDerivedMIOption Option)
+	{
+		DerivedMaterialInstancesComboBoxItems[static_cast<int32>(Option)] = MakeShared<EMaterialStatsDerivedMIOption>(Option);
+	};
+	AddOption(EMaterialStatsDerivedMIOption::Ignore);
+	AddOption(EMaterialStatsDerivedMIOption::CompileOnly);
+	AddOption(EMaterialStatsDerivedMIOption::ShowStats);
+
+	auto PlatformWidget = SNew(SComboBox<TSharedPtr<EMaterialStatsDerivedMIOption>>)
+		.OptionsSource(&DerivedMaterialInstancesComboBoxItems)
+		.OnGenerateWidget_Lambda([](TSharedPtr<EMaterialStatsDerivedMIOption> Option)
+		{
+			return SNew(STextBlock)
+				.Text(Option.IsValid() ? MaterialStatsDerivedMIOptionToDescription(*Option) : MaterialStatsDerivedMIOptionToDescription(EMaterialStatsDerivedMIOption::InvalidOrMax));
+		})
+		.OnSelectionChanged_Lambda([WidgetPtr = this](TSharedPtr<EMaterialStatsDerivedMIOption> Option, ESelectInfo::Type InSelectType)
+		{
+			auto MaterialStats = WidgetPtr->MaterialStatsWPtr.Pin();
+			if (MaterialStats.IsValid() && Option.IsValid())
+			{
+				MaterialStats->SetMaterialStatsDerivedMIOption(*Option);
+			}
+		})
+	[
+		SNew(STextBlock)
+		.Text_Lambda([WidgetPtr = this]()
+		{
+			auto MaterialStats = WidgetPtr->MaterialStatsWPtr.Pin();
+			if (MaterialStats.IsValid())
+			{
+				return MaterialStatsDerivedMIOptionToDescription(MaterialStats->GetMaterialStatsDerivedMIOption());
+			}
+			else
+			{
+				return MaterialStatsDerivedMIOptionToDescription(EMaterialStatsDerivedMIOption::InvalidOrMax);
+			}
+		})
+	]
+	;
+
+	Builder.AddMenuEntry(FUIAction(), PlatformWidget);
+}
+
 void SMaterialEditorStatsWidget::CreateGlobalQualityMenu(class FMenuBuilder& Builder)
 {
 	auto GlobalQualityWidget = SNew(STextBlock)
@@ -510,6 +560,11 @@ TSharedRef<SWidget> SMaterialEditorStatsWidget::GetSettingsButtonContent()
 	FMenuBuilder Builder(false, nullptr);
 
 	CreatePlatformCategoryMenus(Builder);
+	if (ShowMaterialInstancesMenu)
+	{
+		Builder.AddMenuSeparator();
+		CreateDerivedMaterialsMenu(Builder);
+	}
 	Builder.AddMenuSeparator();
 	CreateGlobalQualityMenu(Builder);
 
@@ -525,14 +580,14 @@ TSharedPtr<SWidget> SMaterialEditorStatsWidget::BuildMessageArea()
 	return MessageBoxWidget;
 }
 
-void SMaterialEditorStatsWidget::AddWarningMessage(const FString& Message)
+void SMaterialEditorStatsWidget::AddMessage(const FString& Message, const bool bIsError)
 {
 	MessageBoxWidget->AddSlot()
 		.AutoHeight()
 		.Padding(2.5, 2.5)
 		[
 			SNew(SBorder)
-			.BorderImage(FAppStyle::GetBrush(TEXT("ToolPanel.GroupBorder")))
+			.BorderImage(FAppStyle::Get().GetBrush(TEXT("Brushes.Recessed")))
 			[
 				SNew(SHorizontalBox)
 				+ SHorizontalBox::Slot()
@@ -540,7 +595,7 @@ void SMaterialEditorStatsWidget::AddWarningMessage(const FString& Message)
 				.VAlign(VAlign_Center)
 				[
 					SNew(SImage)
-					.Image(FAppStyle::GetBrush("MessageLog.Warning"))
+					.Image(FAppStyle::GetBrush(bIsError ? "MessageLog.Error" : "MessageLog.Warning"))
 				]
 				+ SHorizontalBox::Slot()
 				.AutoWidth()
@@ -548,7 +603,7 @@ void SMaterialEditorStatsWidget::AddWarningMessage(const FString& Message)
 				[
 					SNew(STextBlock)
 					.TextStyle(FAppStyle::Get(), TEXT("RichTextBlock.Bold"))
-					.ColorAndOpacity(FStyleColors::Warning)
+					.ColorAndOpacity(bIsError ? FStyleColors::Error : FStyleColors::Warning)
 					.Text(FText::FromString(Message))
 					.ToolTipText(FText::FromString(Message))
 				]
@@ -556,7 +611,7 @@ void SMaterialEditorStatsWidget::AddWarningMessage(const FString& Message)
 		];
 }
 
-void SMaterialEditorStatsWidget::ClearWarningMessages()
+void SMaterialEditorStatsWidget::ClearMessages()
 {
 	MessageBoxWidget->ClearChildren();
 }
@@ -564,7 +619,8 @@ void SMaterialEditorStatsWidget::ClearWarningMessages()
 void SMaterialEditorStatsWidget::Construct(const FArguments& InArgs)
 {
 	MaterialStatsWPtr = InArgs._MaterialStatsWPtr;
-	
+	ShowMaterialInstancesMenu = InArgs._ShowMaterialInstancesMenu;
+
 	const auto StatsPtr = MaterialStatsWPtr.Pin();
 
 	if (!StatsPtr.IsValid())
@@ -634,72 +690,58 @@ void SMaterialEditorStatsWidget::Construct(const FArguments& InArgs)
 			]
 		]
 		+ SVerticalBox::Slot()
-		.AutoHeight()
-		[
-			SNew(SBorder)
-			.Padding(0.0f)
-			.HAlign(HAlign_Fill)
-			.BorderImage(FAppStyle::Get().GetBrush(TEXT("Brushes.Recessed")))
-			[
-				MessageArea.ToSharedRef()
-			]
-		]
-		+ SVerticalBox::Slot() // this will contain the stats grid
 		[
-			SNew(SBorder)
-			.BorderImage(FAppStyle::Get().GetBrush(TEXT("Brushes.Recessed")))
-			.Padding(0.0f)
-			.HAlign(HAlign_Fill)
+			SNew(SHorizontalBox)
+			+ SHorizontalBox::Slot()
 			[
-				SNew(SHorizontalBox)
-				+ SHorizontalBox::Slot()
+				SNew(SBorder)
+				.BorderImage(FAppStyle::Get().GetBrush(TEXT("Brushes.Recessed")))
+				.Padding(0.0f)
+				.HAlign(HAlign_Fill)
 				[
 					SNew(SScrollBox)
 					.Orientation(Orient_Vertical)
 					.ExternalScrollbar(VerticalScrollbar)
 					+ SScrollBox::Slot()
 					[
-						// ########## Material stats grid ##########
-						SNew(SVerticalBox)
-						+ SVerticalBox::Slot()
-						.VAlign(VAlign_Fill)
-						.AutoHeight()
+						SNew(SBorder)
+						.Padding(0.0f)
+						.HAlign(HAlign_Fill)
+						.BorderImage(FAppStyle::Get().GetBrush(TEXT("Brushes.Recessed")))
 						[
-							SNew(SHorizontalBox)
-							+ SHorizontalBox::Slot()
-							.HAlign(HAlign_Fill)
-							.FillWidth(1.f)
-							[
-								SNew(SScrollBox)
-								.Orientation(Orient_Horizontal)
-								.ExternalScrollbar(HorizontalScrollbar)
-								+ SScrollBox::Slot()
-								[
-									SAssignNew(MaterialInfoList, SListView<TSharedPtr<int32>>)
-									.ExternalScrollbar(VerticalScrollbar)
-									.ListItemsSource(StatsPtr->GetStatsGrid()->GetGridRowIDs())
-									.OnGenerateRow(this, &SMaterialEditorStatsWidget::MakeMaterialInfoWidget)
-									.Visibility(EVisibility::Visible)
-									.SelectionMode(ESelectionMode::Single)
-									.HeaderRow(PlatformColumnHeader)
-								]
-							]
+							MessageArea.ToSharedRef()
 						]
-						+ SVerticalBox::Slot()
-						.AutoHeight()
+					]
+					+ SScrollBox::Slot()
+					[
+						SNew(SScrollBox)
+						.Orientation(Orient_Horizontal)
+						.ExternalScrollbar(HorizontalScrollbar)
+						+ SScrollBox::Slot()
 						[
-							HorizontalScrollbar
+							SAssignNew(MaterialInfoList, SListView<TSharedPtr<int32>>)
+							.ExternalScrollbar(VerticalScrollbar)
+							.ListItemsSource(StatsPtr->GetStatsGrid()->GetGridRowIDs())
+							.OnGenerateRow(this, &SMaterialEditorStatsWidget::MakeMaterialInfoWidget)
+							.Visibility(EVisibility::Visible)
+							.SelectionMode(ESelectionMode::Single)
+							.HeaderRow(PlatformColumnHeader)
 						]
 					]
 				]
-				+ SHorizontalBox::Slot()
-				.AutoWidth()
-				[
-					VerticalScrollbar
-				]
 			]
+			+ SHorizontalBox::Slot()
+			.AutoWidth()
+			[
+				VerticalScrollbar
+			]
+		]
+		+ SVerticalBox::Slot()
+		.AutoHeight()
+		[
+			HorizontalScrollbar
 		]
-	];
+];
 }
 
 TSharedRef<ITableRow> SMaterialEditorStatsWidget::MakeMaterialInfoWidget(const TSharedPtr<int32> PtrRowID, const TSharedRef<STableViewBase>& OwnerTable) const
@@ -714,7 +756,7 @@ void SMaterialEditorStatsWidget::RequestRefresh()
 	MaterialInfoList->RequestListRefresh();
 }
 
-FName SMaterialEditorStatsWidget::GetRegulatFontStyleName()
+FName SMaterialEditorStatsWidget::GetRegularFontStyleName()
 {
 	return RegularFontStyle;
 }
@@ -724,4 +766,14 @@ FName SMaterialEditorStatsWidget::GetBoldFontStyleName()
 	return BoldFontStyle;
 }
 
+void SMaterialEditorStatsWidget::OnColumnNumChanged()
+{
+	auto StatsPtr = MaterialStatsWPtr.Pin();
+	if (StatsPtr.IsValid() && StatsPtr->GetStatsGrid().IsValid())
+	{
+		StatsPtr->GetStatsGrid()->OnColumnNumChanged();
+		RebuildColumns();
+	}
+}
+
 #undef LOCTEXT_NAMESPACE
\ No newline at end of file
diff --git a/Engine/Source/Editor/MaterialEditor/Private/SMaterialEditorStatsWidget.h b/Engine/Source/Editor/MaterialEditor/Private/SMaterialEditorStatsWidget.h
--- a/Engine/Source/Editor/MaterialEditor/Private/SMaterialEditorStatsWidget.h
+++ b/Engine/Source/Editor/MaterialEditor/Private/SMaterialEditorStatsWidget.h
@@ -8,6 +8,7 @@
 #include "Widgets/Views/STableViewBase.h"
 #include "Widgets/Views/STableRow.h"
 #include "MaterialStats.h"
+#include "MaterialStatsGrid.h"
 
 //////////////////////////////////////////////////////////////////////////
 // SMaterialEditorStatsWidget
@@ -18,9 +19,11 @@ class SMaterialEditorStatsWidget : public SCompoundWidget
 public:
 	SLATE_BEGIN_ARGS(SMaterialEditorStatsWidget)
 		: _MaterialStatsWPtr(nullptr)
+		, _ShowMaterialInstancesMenu(true)
 	{}
 
 		SLATE_ARGUMENT(TWeakPtr<FMaterialStats>, MaterialStatsWPtr)
+		SLATE_ARGUMENT(bool, ShowMaterialInstancesMenu)
 	SLATE_END_ARGS()
 
 	void Construct(const FArguments& InArgs);
@@ -31,14 +34,16 @@ public:
 	/** Request a refresh of the list view */
 	void RequestRefresh();
 
-	void AddWarningMessage(const FString& Message);
-	void ClearWarningMessages();
+	void AddMessage(const FString& Message, const bool bIsError);
+	void ClearMessages();
 
 	void FillWarningMessages();
 
-	static FName GetRegulatFontStyleName();
+	static FName GetRegularFontStyleName();
 	static FName GetBoldFontStyleName();
 
+	void OnColumnNumChanged();
+
 private:
 	void RebuildColumns();
 	uint32 CountSubPlatforms(EPlatformCategoryType Category) const;
@@ -55,7 +60,10 @@ private:
 	void CreatePlatformMenus(class FMenuBuilder& Builder, EPlatformCategoryType Category);
 	void CreatePlatformCategoryMenus(class FMenuBuilder& Builder);
 
+	static FText MaterialStatsDerivedMIOptionToDescription(const EMaterialStatsDerivedMIOption Option);
+
 	void CreateQualityMenus(class FMenuBuilder& Builder);
+	void CreateDerivedMaterialsMenu(class FMenuBuilder& Builder);
 	void CreateGlobalQualityMenu(class FMenuBuilder& Builder);
 
 	void OnFlipQualityState(const ECheckBoxState NewState, const EMaterialQualityLevel::Type QualitySetting);
@@ -74,8 +82,11 @@ protected:
 	TSharedPtr<SVerticalBox> MessageBoxWidget;
 
 	TWeakPtr<FMaterialStats> MaterialStatsWPtr;
+	bool ShowMaterialInstancesMenu;
 
 	TSharedPtr<SListView<TSharedPtr<int32>>> MaterialInfoList;
 
 	TSharedPtr<SHeaderRow> PlatformColumnHeader;
+
+	TArray<TSharedPtr<EMaterialStatsDerivedMIOption>> DerivedMaterialInstancesComboBoxItems;
 };
\ No newline at end of file
diff --git a/Engine/Source/Editor/MaterialEditor/Public/IMaterialEditor.h b/Engine/Source/Editor/MaterialEditor/Public/IMaterialEditor.h
--- a/Engine/Source/Editor/MaterialEditor/Public/IMaterialEditor.h
+++ b/Engine/Source/Editor/MaterialEditor/Public/IMaterialEditor.h
@@ -159,6 +159,9 @@ public:
 
 	/** Fills inheritance drop down menu */
 	virtual void GenerateInheritanceMenu(class UToolMenu* Menu) {};
+
+	/** Updates materials used in material stats */
+	virtual void RefreshStatsMaterials() {};
 private:
 	FMaterialEditorClosedEvent MaterialEditorClosedEvent;
 	FRegisterTabSpawnersEvent RegisterTabSpawnersEvent;
diff --git a/Engine/Source/Editor/MaterialEditor/Public/MaterialStatsCommon.h b/Engine/Source/Editor/MaterialEditor/Public/MaterialStatsCommon.h
--- a/Engine/Source/Editor/MaterialEditor/Public/MaterialStatsCommon.h
+++ b/Engine/Source/Editor/MaterialEditor/Public/MaterialStatsCommon.h
@@ -92,7 +92,7 @@ public:
 
 public:
 	/** call this to create an instance to FMaterialStats */
-	static TSharedPtr<class FMaterialStats> CreateMaterialStats(class IMaterialEditor* MaterialEditor);
+	static TSharedPtr<class FMaterialStats> CreateMaterialStats(class IMaterialEditor* MaterialEditor, const bool ShowMaterialInstancesMenu);
 
 	/** utility functions that translate various enum values to strings */
 	static FString MaterialQualityToString(const EMaterialQualityLevel::Type Quality);
diff --git a/Engine/Source/Editor/UnrealEd/Classes/Preferences/MaterialStatsOptions.h b/Engine/Source/Editor/UnrealEd/Classes/Preferences/MaterialStatsOptions.h
--- a/Engine/Source/Editor/UnrealEd/Classes/Preferences/MaterialStatsOptions.h
+++ b/Engine/Source/Editor/UnrealEd/Classes/Preferences/MaterialStatsOptions.h
@@ -17,6 +17,15 @@
 #include "SceneTypes.h"
 #include "MaterialStatsOptions.generated.h"
 
+UENUM()
+enum class EMaterialStatsDerivedMIOption : uint8
+{
+	Ignore = 0,
+	CompileOnly,
+	ShowStats,
+	InvalidOrMax,
+};
+
 UCLASS(hidecategories = Object, config = EditorPerProjectUserSettings, MinimalAPI)
 class UMaterialStatsOptions : public UObject
 {
@@ -27,4 +36,7 @@ class UMaterialStatsOptions : public UObject
 
 	UPROPERTY(EditAnywhere, config, Category = Options)
 	int32 bMaterialQualityUsed[EMaterialQualityLevel::Num];
+
+	UPROPERTY(EditAnywhere, config, Category=Options)
+	EMaterialStatsDerivedMIOption MaterialStatsDerivedMIOption;
 };
diff --git a/Engine/Source/Editor/UnrealEd/Private/Dialogs/Dialogs.cpp b/Engine/Source/Editor/UnrealEd/Private/Dialogs/Dialogs.cpp
--- a/Engine/Source/Editor/UnrealEd/Private/Dialogs/Dialogs.cpp
+++ b/Engine/Source/Editor/UnrealEd/Private/Dialogs/Dialogs.cpp
@@ -815,6 +815,13 @@ private:
 	TAttribute<FText> MyCheckboxMessage;
 };
 
+TSet<FString> FSuppressableWarningDialog::SuppressedInTheSession = {};
+
+FString SuppressableWarningDialogGetSessionKey(const FString& IniSettingName, const FString& IniSettingFileName)
+{
+	return IniSettingFileName + TEXT("_") + IniSettingName;
+}
+
 FSuppressableWarningDialog::FSuppressableWarningDialog(const FSetupInfo& Info)
 {
 	// Ensure proper usage of the suppression warning.
@@ -828,8 +835,16 @@ FSuppressableWarningDialog::FSuppressableWarningDialog(const FSetupInfo& Info)
 	IniSettingName = Info.IniSettingName;
 	IniSettingFileName = Info.IniSettingFileName;
 	Prompt = Info.Message;
+	bDontPersistSuppressionAcrossSessions = Info.bDontPersistSuppressionAcrossSessions;
 
-	GConfig->GetBool( *ConfigSection, *IniSettingName, bShouldSuppressDialog, IniSettingFileName );
+	if (bDontPersistSuppressionAcrossSessions)
+	{
+		bShouldSuppressDialog = SuppressedInTheSession.Contains(SuppressableWarningDialogGetSessionKey(IniSettingName, IniSettingFileName));
+	}
+	else
+	{
+		GConfig->GetBool( *ConfigSection, *IniSettingName, bShouldSuppressDialog, IniSettingFileName );
+	}
 	
 	if (!bShouldSuppressDialog && FSlateApplication::IsInitialized())
 	{
@@ -862,7 +877,14 @@ FSuppressableWarningDialog::EResult FSuppressableWarningDialog::ShowModal() cons
 	bool bShouldSuppressDialog = false;
 
 	// Get the setting from the config file.
-	GConfig->GetBool( *ConfigSection, *IniSettingName, bShouldSuppressDialog, IniSettingFileName );
+	if (bDontPersistSuppressionAcrossSessions)
+	{
+		bShouldSuppressDialog = SuppressedInTheSession.Contains(SuppressableWarningDialogGetSessionKey(IniSettingName, IniSettingFileName));
+	}
+	else
+	{
+		GConfig->GetBool( *ConfigSection, *IniSettingName, bShouldSuppressDialog, IniSettingFileName );
+	}
 
 	EResult RetCode = Suppressed;
 	if( !bShouldSuppressDialog )
@@ -874,7 +896,17 @@ FSuppressableWarningDialog::EResult FSuppressableWarningDialog::ShowModal() cons
 		{
 			// Set the ini variable to the state of the disable check box
 			bShouldSuppressDialog = MessageBox->GetCheckBoxState();
-			GConfig->SetBool( *ConfigSection, *IniSettingName, bShouldSuppressDialog, IniSettingFileName );
+			if (bDontPersistSuppressionAcrossSessions)
+			{
+				if (bShouldSuppressDialog)
+				{
+					SuppressedInTheSession.Add(SuppressableWarningDialogGetSessionKey(IniSettingName, IniSettingFileName));
+				}
+			}
+			else
+			{
+				GConfig->SetBool( *ConfigSection, *IniSettingName, bShouldSuppressDialog, IniSettingFileName );
+			}
 		}
 	}
 	else
diff --git a/Engine/Source/Editor/UnrealEd/Private/PreferenceStubs.cpp b/Engine/Source/Editor/UnrealEd/Private/PreferenceStubs.cpp
--- a/Engine/Source/Editor/UnrealEd/Private/PreferenceStubs.cpp
+++ b/Engine/Source/Editor/UnrealEd/Private/PreferenceStubs.cpp
@@ -85,6 +85,8 @@ UMaterialStatsOptions::UMaterialStatsOptions(const FObjectInitializer& ObjectIni
 #endif
 
 	bMaterialQualityUsed[EMaterialQualityLevel::High] = 1;
+
+	MaterialStatsDerivedMIOption = EMaterialStatsDerivedMIOption::Ignore;
 }
 
 UAnimationBlueprintEditorOptions::UAnimationBlueprintEditorOptions(const FObjectInitializer& ObjectInitializer)
diff --git a/Engine/Source/Editor/UnrealEd/Public/Dialogs/Dialogs.h b/Engine/Source/Editor/UnrealEd/Public/Dialogs/Dialogs.h
--- a/Engine/Source/Editor/UnrealEd/Public/Dialogs/Dialogs.h
+++ b/Engine/Source/Editor/UnrealEd/Public/Dialogs/Dialogs.h
@@ -99,6 +99,9 @@ public:
 		/** If true the suppress checkbox defaults to true*/
 		bool bDefaultToSuppressInTheFuture;
 
+		/** If true suppression will not persist for future editor sessions */
+		bool bDontPersistSuppressionAcrossSessions;
+
 		/** Text used on the button which will return FSuppressableWarningDialog::Confirm */
 		FText ConfirmText;
 
@@ -125,6 +128,7 @@ public:
 			, IniSettingName(InIniSettingName)
 			, IniSettingFileName(InIniSettingFileName)
 			, bDefaultToSuppressInTheFuture(false)
+			, bDontPersistSuppressionAcrossSessions(false)
 			, ConfirmText()
 			, CancelText()
 			, CheckBoxText(NSLOCTEXT("ModalDialogs", "DefaultCheckBoxMessage", "Don't show this again"))
@@ -168,6 +172,11 @@ private:
 	/** Cached pointer to the message box held within the window */
 	TSharedPtr<class SModalDialogWithCheckbox> MessageBox;
 
+	/** If true suppression will not persist for future editor sessions */
+	bool bDontPersistSuppressionAcrossSessions;
+
+	/** Set of session only suppressions */
+	static TSet<FString> SuppressedInTheSession;
 };
 
 
diff --git a/Engine/Source/Runtime/Engine/Private/Materials/MaterialShared.cpp b/Engine/Source/Runtime/Engine/Private/Materials/MaterialShared.cpp
--- a/Engine/Source/Runtime/Engine/Private/Materials/MaterialShared.cpp
+++ b/Engine/Source/Runtime/Engine/Private/Materials/MaterialShared.cpp
@@ -803,7 +803,7 @@ bool FMaterial::IsRequiredComplete() const
 }
 
 #if WITH_EDITOR
-void FMaterial::GetShaderMapIDsWithUnfinishedCompilation(TArray<int32>& ShaderMapIds)
+void FMaterial::AddShaderMapIDsWithUnfinishedCompilation(TArray<int32>& ShaderMapIds)
 {
 	if (GameThreadCompilingShaderMapId != 0u && GShaderCompilingManager->IsCompilingShaderMap(GameThreadCompilingShaderMapId))
 	{
@@ -840,7 +840,7 @@ void FMaterial::CancelCompilation()
 	}
 
 	TArray<int32> ShaderMapIdsToCancel;
-	GetShaderMapIDsWithUnfinishedCompilation(ShaderMapIdsToCancel);
+	AddShaderMapIDsWithUnfinishedCompilation(ShaderMapIdsToCancel);
 
 	if (ShaderMapIdsToCancel.Num() > 0)
 	{
@@ -854,7 +854,7 @@ void FMaterial::FinishCompilation()
 	FinishCacheShaders();
 
 	TArray<int32> ShaderMapIdsToFinish;
-	GetShaderMapIDsWithUnfinishedCompilation(ShaderMapIdsToFinish);
+	AddShaderMapIDsWithUnfinishedCompilation(ShaderMapIdsToFinish);
 
 	if (ShaderMapIdsToFinish.Num() > 0)
 	{
@@ -863,6 +863,26 @@ void FMaterial::FinishCompilation()
 	}
 }
 
+void FMaterial::FinishCompilation(const TCHAR* MaterialName, const TArray<FMaterial*>& MaterialsToCompile)
+{
+	for(const FMaterial* Material: MaterialsToCompile)
+	{
+		Material->FinishCacheShaders();
+	}
+
+	TArray<int32> ShaderMapIdsToFinish;
+	for(FMaterial* Material: MaterialsToCompile)
+	{
+		Material->AddShaderMapIDsWithUnfinishedCompilation(ShaderMapIdsToFinish);
+	}
+
+	if (ShaderMapIdsToFinish.Num() > 0)
+	{
+		// Block until the shader maps that we will save have finished being compiled
+		GShaderCompilingManager->FinishCompilation(MaterialName, ShaderMapIdsToFinish);
+	}
+}
+
 const FMaterialCachedHLSLTree* FMaterial::GetCachedHLSLTree() const
 {
 	const UMaterialInterface* MaterialInterface = GetMaterialInterface();
diff --git a/Engine/Source/Runtime/Engine/Private/ShaderCompiler/ShaderCompiler.cpp b/Engine/Source/Runtime/Engine/Private/ShaderCompiler/ShaderCompiler.cpp
--- a/Engine/Source/Runtime/Engine/Private/ShaderCompiler/ShaderCompiler.cpp
+++ b/Engine/Source/Runtime/Engine/Private/ShaderCompiler/ShaderCompiler.cpp
@@ -5983,6 +5983,7 @@ void FShaderCompilingManager::BlockOnShaderMapCompletion(const TArray<int32>& Sh
 		}
 
 		FScopedSlowTask SlowTask(NumJobs, FText::Format(LOCTEXT("BlockOnShaderMapCompletion", "Compiling Shaders ({0})"), NumJobs), GIsEditor && !IsRunningCommandlet());
+		SlowTask.MakeDialog(false, true);
 
 		int32 NumPendingJobs = 0;
 		// Keep track of previous number of pending jobs so we can update the slow task with the amount of work done.
diff --git a/Engine/Source/Runtime/Engine/Public/MaterialShared.h b/Engine/Source/Runtime/Engine/Public/MaterialShared.h
--- a/Engine/Source/Runtime/Engine/Public/MaterialShared.h
+++ b/Engine/Source/Runtime/Engine/Public/MaterialShared.h
@@ -2025,6 +2025,11 @@ public:
 	 */
 	ENGINE_API void FinishCompilation();
 
+	/** 
+	 * Blocks until compilation has completed. Returns immediately if a compilation is not outstanding.
+	 */
+	ENGINE_API static void FinishCompilation(const TCHAR* MaterialName, const TArray<FMaterial*>& MaterialsToCompile);
+
 	/**
 	 * Checks if the compilation for this shader is finished
 	 * 
@@ -2327,7 +2332,7 @@ protected:
 	/**
 	* Fills the passed array with IDs of shader maps unfinished compilation jobs.
 	*/
-	void GetShaderMapIDsWithUnfinishedCompilation(TArray<int32>& ShaderMapIds);
+	void AddShaderMapIDsWithUnfinishedCompilation(TArray<int32>& ShaderMapIds);
 
 	uint32 GetGameThreadCompilingShaderMapId() const { return GameThreadCompilingShaderMapId; }
 
diff --git a/Engine/Source/Runtime/Engine/Public/StaticParameterSet.h b/Engine/Source/Runtime/Engine/Public/StaticParameterSet.h
--- a/Engine/Source/Runtime/Engine/Public/StaticParameterSet.h
+++ b/Engine/Source/Runtime/Engine/Public/StaticParameterSet.h
@@ -519,7 +519,7 @@ struct FStaticParameterSet : public FStaticParameterSetRuntimeData
 		return !(*this == ReferenceSet);
 	}
 
-	bool Equivalent(const FStaticParameterSet& ReferenceSet) const;
+	ENGINE_API bool Equivalent(const FStaticParameterSet& ReferenceSet) const;
 
 #if WITH_EDITORONLY_DATA
 	ENGINE_API void SetParameterValue(const FMaterialParameterInfo& ParameterInfo, const FMaterialParameterMetadata& Meta, EMaterialSetParameterValueFlags Flags = EMaterialSetParameterValueFlags::None);
