Clarified the message shown on SMaterialEditorSubstrateWidget when Substrate is disabled.

diff --git a/Engine/Source/Editor/MaterialEditor/Private/SMaterialEditorSubstrateWidget.cpp b/Engine/Source/Editor/MaterialEditor/Private/SMaterialEditorSubstrateWidget.cpp
--- a/Engine/Source/Editor/MaterialEditor/Private/SMaterialEditorSubstrateWidget.cpp
+++ b/Engine/Source/Editor/MaterialEditor/Private/SMaterialEditorSubstrateWidget.cpp
@@ -361,7 +361,7 @@ void SMaterialEditorSubstrateWidget::Construct(const FArguments& InArgs, TWeakPt
 					.ColorAndOpacity(FLinearColor::Yellow)
 					.ShadowColorAndOpacity(FLinearColor::Black)
 					.ShadowOffset(FVector2D::UnitVector)
-					.Text(LOCTEXT("SubstrateWidgetNotEnable", "Substrate is not enabled."))
+					.Text(LOCTEXT("SubstrateWidgetNotEnable", "Details cannot be shown: Substrate (experimental) is not enabled for this project (See the project settings window, rendering settings section)."))
 				]
 			]
 		];
