Fix for conflicting localization messages with material editor errors.

- Warning_CompileErrorsInMaterial was used twice for two different error messages.
- Renamed the second one to Warning_CompileErrorsInMaterialInstance, since it was an error message specific to material instances.

diff --git a/Engine/Source/Editor/MaterialEditor/Private/MaterialEditor.cpp b/Engine/Source/Editor/MaterialEditor/Private/MaterialEditor.cpp
--- a/Engine/Source/Editor/MaterialEditor/Private/MaterialEditor.cpp
+++ b/Engine/Source/Editor/MaterialEditor/Private/MaterialEditor.cpp
@@ -2665,8 +2665,8 @@ bool FMaterialEditor::UpdateOriginalMaterial()
 						const FString QualityName = FMaterialStatsUtils::MaterialQualityToShortString((EMaterialQualityLevel::Type)QualityLevel);
 
 						FSuppressableWarningDialog::FSetupInfo Info(
-						FText::Format(NSLOCTEXT("UnrealEd", "Warning_CompileErrorsInMaterial", "The current material has compilation errors in derived material instance {0} for platform {1} at quality level {2}, so it will not render correctly.\nAre you sure you wish to continue?"),FText::FromString(*AssetName),FText::FromName(PlatformName),FText::FromString(QualityName)),
-						NSLOCTEXT("UnrealEd", "Warning_CompileErrorsInMaterial_Title", "Warning: Compilation errors in this Material" ), "Warning_CompileErrorsInMaterial");
+						FText::Format(NSLOCTEXT("UnrealEd", "Warning_CompileErrorsInMaterialInstance", "The current material has compilation errors in derived material instance {0} for platform {1} at quality level {2}, so it will not render correctly.\nAre you sure you wish to continue?"),FText::FromString(*AssetName),FText::FromName(PlatformName),FText::FromString(QualityName)),
+						NSLOCTEXT("UnrealEd", "Warning_CompileErrorsInMaterial_Title", "Warning: Compilation errors in this Material" ), "Warning_CompileErrorsInMaterialInstance");
 						Info.ConfirmText = NSLOCTEXT("ModalDialogs", "CompileErrorsInMaterialConfirm", "Continue");
 						Info.CancelText = NSLOCTEXT("ModalDialogs", "CompileErrorsInMaterialCancel", "Abort");
 						Info.bDontPersistSuppressionAcrossSessions = true;
