Improve platform stats and material editor compilation warnings.

diff --git a/Engine/Source/Editor/MaterialEditor/Private/MaterialEditor.cpp b/Engine/Source/Editor/MaterialEditor/Private/MaterialEditor.cpp
--- a/Engine/Source/Editor/MaterialEditor/Private/MaterialEditor.cpp
+++ b/Engine/Source/Editor/MaterialEditor/Private/MaterialEditor.cpp
@@ -2605,6 +2605,9 @@ bool FMaterialEditor::UpdateOriginalMaterial()
 		MaterialStatsManager->CacheAndCompilePendingShaders();
 	}
 
+	TArray<FText> Errors;
+	bool bBaseMaterialFailsToCompile = false;
+
 	// If the Material has compilation errors, warn the user
 	for (int32 i = ERHIFeatureLevel::Num - 1; i >= 0; --i)
 	{
@@ -2614,34 +2617,8 @@ bool FMaterialEditor::UpdateOriginalMaterial()
 		{
 			FString FeatureLevelName;
 			GetFeatureLevelName(FeatureLevel, FeatureLevelName);
-
-			if (Material->bUsedAsSpecialEngineMaterial)
-			{
-				FSuppressableWarningDialog::FSetupInfo Info(
-					FText::Format(NSLOCTEXT("UnrealEd", "Error_CompileErrorsInDefaultMaterial", "The current material has compilation errors for feature level {0}.\nThis material is a Default Material which must be available as a code fallback at all times, compilation errors are not allowed."), FText::FromString(*FeatureLevelName)),
-					NSLOCTEXT("UnrealEd", "Warning_CompileErrorsInDefaultMaterial_Title", "Error: Compilation errors in Default Material"), "Error_CompileErrorsInDefaultMaterial");
-				Info.ConfirmText = NSLOCTEXT("ModalDialogs", "CompileErrorsInDefaultMaterialOk", "Ok");
-				Info.bDontPersistSuppressionAcrossSessions = true;
-
-				FSuppressableWarningDialog CompileErrors(Info);
-				CompileErrors.ShowModal();
-				return false;
-			}
-			else
-			{
-				FSuppressableWarningDialog::FSetupInfo Info(
-					FText::Format(NSLOCTEXT("UnrealEd", "Warning_CompileErrorsInMaterial", "The current material has compilation errors, so it will not render correctly in feature level {0}.\nAre you sure you wish to continue?"),FText::FromString(*FeatureLevelName)),
-					NSLOCTEXT("UnrealEd", "Warning_CompileErrorsInMaterial_Title", "Warning: Compilation errors in this Material" ), "Warning_CompileErrorsInMaterial");
-				Info.ConfirmText = NSLOCTEXT("ModalDialogs", "CompileErrorsInMaterialConfirm", "Continue");
-				Info.CancelText = NSLOCTEXT("ModalDialogs", "CompileErrorsInMaterialCancel", "Abort");
-				Info.bDontPersistSuppressionAcrossSessions = true;
-
-				FSuppressableWarningDialog CompileErrorsWarning( Info );
-				if( CompileErrorsWarning.ShowModal() == FSuppressableWarningDialog::Cancel )
-				{
-					return false;
-				}
-			}
+			Errors.Push(FText::Format(NSLOCTEXT("UnrealEd", "Warning_CompileErrorsInMaterial_ListEntryFeatureLevel", "- At feature level {0}."), FText::FromString(*FeatureLevelName)));
+			bBaseMaterialFailsToCompile = true;
 		}
 	}
 
@@ -2665,24 +2642,48 @@ bool FMaterialEditor::UpdateOriginalMaterial()
 						const auto& AssetName = MaterialStatsManager->GetMaterialName(InstanceIndex);
 						const FString QualityName = FMaterialStatsUtils::MaterialQualityToShortString((EMaterialQualityLevel::Type)QualityLevel);
 
-						FSuppressableWarningDialog::FSetupInfo Info(
-						FText::Format(NSLOCTEXT("UnrealEd", "Warning_CompileErrorsInMaterialInstance", "The current material has compilation errors in derived material instance {0} for platform {1} at quality level {2}, so it will not render correctly.\nAre you sure you wish to continue?"),FText::FromString(*AssetName),FText::FromName(PlatformName),FText::FromString(QualityName)),
-						NSLOCTEXT("UnrealEd", "Warning_CompileErrorsInMaterial_Title", "Warning: Compilation errors in this Material" ), "Warning_CompileErrorsInMaterialInstance");
-						Info.ConfirmText = NSLOCTEXT("ModalDialogs", "CompileErrorsInMaterialConfirm", "Continue");
-						Info.CancelText = NSLOCTEXT("ModalDialogs", "CompileErrorsInMaterialCancel", "Abort");
-						Info.bDontPersistSuppressionAcrossSessions = true;
-
-						FSuppressableWarningDialog CompileErrorsWarning( Info );
-						if( CompileErrorsWarning.ShowModal() == FSuppressableWarningDialog::Cancel )
-						{
-							return false;
-						}
+						Errors.Push(FText::Format(NSLOCTEXT("UnrealEd", "Warning_CompileErrorsInMaterial_ListEntryDerivedMaterial", "- Material instance {0} for platform {1} at quality level {2}."), FText::FromString(*AssetName), FText::FromName(PlatformName), FText::FromString(QualityName)));
 					}
 				}
 			}
 		}
 	}
 
+
+	if (Errors.Num() > 0)
+	{
+		const FText JoinedErrors = FText::Join(FText::FromString(TEXT("\n")), Errors);
+		if (Material->bUsedAsSpecialEngineMaterial && bBaseMaterialFailsToCompile)
+		{
+			FSuppressableWarningDialog::FSetupInfo Info(
+				FText::Format(NSLOCTEXT("UnrealEd", "Error_CompileErrorsInDefaultMaterial", "The current material has the following compilation errors:\n{0}\nThis material is a Default Material which must be available as a code fallback at all times, compilation errors are not allowed."), JoinedErrors),
+				NSLOCTEXT("UnrealEd", "Warning_CompileErrorsInDefaultMaterial_Title", "Error: Compilation errors in Default Material"), "Error_CompileErrorsInDefaultMaterial");
+			Info.ConfirmText = NSLOCTEXT("ModalDialogs", "CompileErrorsInDefaultMaterialOk", "Ok");
+			Info.bDontPersistSuppressionAcrossSessions = true;
+			Info.WrapMessageAt = 0.0f;
+
+			FSuppressableWarningDialog CompileErrors(Info);
+			CompileErrors.ShowModal();
+			return false;
+		}
+		else
+		{
+			FSuppressableWarningDialog::FSetupInfo Info(
+				FText::Format(NSLOCTEXT("UnrealEd", "Warning_CompileErrorsInMaterial", "The current material has the following compilation errors:\n{0}\nAre you sure you wish to continue?"), JoinedErrors),
+				NSLOCTEXT("UnrealEd", "Warning_CompileErrorsInMaterial_Title", "Warning: Compilation errors in this Material" ), "Warning_CompileErrorsInMaterial");
+			Info.ConfirmText = NSLOCTEXT("ModalDialogs", "CompileErrorsInMaterialConfirm", "Continue");
+			Info.CancelText = NSLOCTEXT("ModalDialogs", "CompileErrorsInMaterialCancel", "Abort");
+			Info.bDontPersistSuppressionAcrossSessions = true;
+			Info.WrapMessageAt = 0.0f;
+
+			FSuppressableWarningDialog CompileErrorsWarning( Info );
+			if(CompileErrorsWarning.ShowModal() == FSuppressableWarningDialog::Cancel)
+			{
+				return false;
+			}
+		}
+	}
+
 	// Make sure any graph position changes that might not have been copied are taken into account
 	Material->MaterialGraph->LinkMaterialExpressionsFromGraph();
 
diff --git a/Engine/Source/Editor/MaterialEditor/Private/MaterialStatsGrid.cpp b/Engine/Source/Editor/MaterialEditor/Private/MaterialStatsGrid.cpp
--- a/Engine/Source/Editor/MaterialEditor/Private/MaterialStatsGrid.cpp
+++ b/Engine/Source/Editor/MaterialEditor/Private/MaterialStatsGrid.cpp
@@ -567,8 +567,6 @@ void FMaterialStatsGrid::BuildRowIds()
 
 void FMaterialStatsGrid::OnShaderChanged()
 {
-	SCOPED_NAMED_EVENT_TEXT("FMaterialStatsGrid::OnShaderChanged", FColor::Red);
-
 	CollectShaderInfo();
 
 	BuildShaderRows();
@@ -580,8 +578,6 @@ void FMaterialStatsGrid::OnShaderChanged()
 
 void FMaterialStatsGrid::OnColumnNumChanged()
 {
-	SCOPED_NAMED_EVENT_TEXT("FMaterialStatsGrid::OnColumnNumChanged", FColor::Red);
-
 	// nuke all rows to rebuild columns
 	BuildStaticRows();
 	OnShaderChanged();
diff --git a/Engine/Source/Editor/MaterialEditor/Private/SMaterialEditorStatsWidget.cpp b/Engine/Source/Editor/MaterialEditor/Private/SMaterialEditorStatsWidget.cpp
--- a/Engine/Source/Editor/MaterialEditor/Private/SMaterialEditorStatsWidget.cpp
+++ b/Engine/Source/Editor/MaterialEditor/Private/SMaterialEditorStatsWidget.cpp
@@ -53,6 +53,8 @@ public:
 	virtual TSharedRef<SWidget> GenerateWidgetForColumn(const FName& ColumnName) override;
 
 	FText GetTextForCell(const FName Name, const bool bToolTip) const;
+	static const FSlateBrush* GetIconForCell(const FGridCell::EIcon Icon);
+	const FSlateBrush* GetIconForCell(const FName Name) const;
 	FSlateColor GetColorForCell(const FName Name) const;
 	EHorizontalAlignment GetHAlignForCell(const FName Name) const;
 	EVerticalAlignment GetVAlignForCell(const FName Name) const;
@@ -64,22 +66,10 @@ private:
 	TWeakPtr<FMaterialStats> MaterialStatsWPtr;
 };
 
-FName GetGridCellIconBrushName(FGridCell::EIcon Icon)
-{
-	switch(Icon)
-	{
-	case FGridCell::EIcon::Error:
-		return "MessageLog.Error";
-	default:
-		return "";
-	}
-}
-
 TSharedRef<SWidget> SMaterialStatsViewRow::GenerateWidgetForColumn(const FName& ColumnName)
 {
 	EHorizontalAlignment HAlign = EHorizontalAlignment::HAlign_Fill;
 	EVerticalAlignment VALign = EVerticalAlignment::VAlign_Top;
-	FGridCell::EIcon Icon = FGridCell::EIcon::None;
 
 	FName UsedFontStyle = SMaterialEditorStatsWidget::GetRegularFontStyleName();
 
@@ -93,50 +83,32 @@ TSharedRef<SWidget> SMaterialStatsViewRow::GenerateWidgetForColumn(const FName&
 		UsedFontStyle = Cell->IsContentBold() ? SMaterialEditorStatsWidget::GetBoldFontStyleName() : SMaterialEditorStatsWidget::GetRegularFontStyleName();
 		HAlign = Cell->GetHorizontalAlignment();
 		VALign = Cell->GetVerticalAlignment();
-		Icon = Cell->GetIcon();
 	}
 
-	if (Icon != FGridCell::EIcon::None)
-	{
-		return SNew(SBox)
-			.Padding(FMargin(4, 2, 4, 2))
-			.HAlign(HAlign)
-			.VAlign(VALign)
+	return SNew(SBox)
+		.Padding(FMargin(4, 2, 4, 2))
+		.HAlign(HAlign)
+		.VAlign(VALign)
+		[
+			SNew(SHorizontalBox)
+			+ SHorizontalBox::Slot()
+			.AutoWidth()
+			.VAlign(VAlign_Center)
 			[
-				SNew(SHorizontalBox)
-				+ SHorizontalBox::Slot()
-				.AutoWidth()
-				.VAlign(VAlign_Center)
-				[
-					SNew(SImage)
-					.Image(FAppStyle::GetBrush(GetGridCellIconBrushName(Icon)))
-				]
-				+ SHorizontalBox::Slot()
-				.AutoWidth()
-				.VAlign(VAlign_Center)
-				[
-					SNew(STextBlock)
-					.TextStyle(FAppStyle::Get(), UsedFontStyle)
-					.Text(this, &SMaterialStatsViewRow::GetTextForCell, ColumnName, false)
-					.ToolTipText(this, &SMaterialStatsViewRow::GetTextForCell, ColumnName, true)
-					.AutoWrapText(true)
-				]
-			];
-	}
-	else
-	{
-		return SNew(SBox)
-			.Padding(FMargin(4, 2, 4, 2))
-			.HAlign(HAlign)
-			.VAlign(VALign)
+				SNew(SImage)
+				.Image(this, &SMaterialStatsViewRow::GetIconForCell, ColumnName)
+			]
+			+ SHorizontalBox::Slot()
+			.AutoWidth()
+			.VAlign(VAlign_Center)
 			[
 				SNew(STextBlock)
 				.TextStyle(FAppStyle::Get(), UsedFontStyle)
 				.Text(this, &SMaterialStatsViewRow::GetTextForCell, ColumnName, false)
 				.ToolTipText(this, &SMaterialStatsViewRow::GetTextForCell, ColumnName, true)
 				.AutoWrapText(true)
-			];
-	}
+			]
+		];
 }
 
 FText SMaterialStatsViewRow::GetTextForCell(const FName Name, const bool bToolTip) const
@@ -159,6 +131,36 @@ FText SMaterialStatsViewRow::GetTextForCell(const FName Name, const bool bToolTi
 	return FinalText;
 }
 
+const FSlateBrush* SMaterialStatsViewRow::GetIconForCell(const FGridCell::EIcon Icon)
+{
+	switch(Icon)
+	{
+	case FGridCell::EIcon::Error:
+		return FAppStyle::GetBrush(TEXT("MessageLog.Error"));
+	default:
+		return nullptr;
+	}
+}
+
+const FSlateBrush* SMaterialStatsViewRow::GetIconForCell(const FName Name) const
+{
+	FGridCell::EIcon CellContent = FGridCell::EIcon::None;
+
+	const auto StatsPtr = MaterialStatsWPtr.Pin();
+	if (StatsPtr.IsValid() && PtrRowID.IsValid())
+	{
+		const int32 RowID = *PtrRowID;
+
+		const auto Cell = StatsPtr->GetStatsGrid()->GetCell(RowID, Name);
+		if (Cell.IsValid())
+		{
+			CellContent = Cell->GetIcon();
+		}
+	}
+
+	return GetIconForCell(CellContent);
+}
+
 FSlateColor SMaterialStatsViewRow::GetColorForCell(const FName Name) const
 {
 	return FStyleColors::Foreground;
@@ -227,10 +229,10 @@ float SMaterialEditorStatsWidget::GetColumnSize(const FName ColumnName) const
 					const auto Cell = StatsPtr->GetStatsGrid()->GetCell(*ArrRowIds[i], ColumnName);
 
 					const FString Content = Cell->GetCellContent();
-					const FGridCell::EIcon Icon = Cell->GetIcon();
+					const FSlateBrush* Icon = SMaterialStatsViewRow::GetIconForCell(Cell->GetIcon());
 
 					FVector2D FontMeasure = FontMeasureService->Measure(Content, FontInfo);
-					const float IconSize = Icon != FGridCell::EIcon::None ? FAppStyle::GetBrush(GetGridCellIconBrushName(Icon))->GetImageSize().X : 0.0f;
+					const float IconSize = Icon != nullptr ? Icon->GetImageSize().X : 0.0f;
 
 					ColumnSize = FMath::Clamp(FontMeasure.X + IconSize, ColumnSize, ColumnSizeExtraLarge);
 				}
diff --git a/Engine/Source/Editor/UnrealEd/Private/Dialogs/Dialogs.cpp b/Engine/Source/Editor/UnrealEd/Private/Dialogs/Dialogs.cpp
--- a/Engine/Source/Editor/UnrealEd/Private/Dialogs/Dialogs.cpp
+++ b/Engine/Source/Editor/UnrealEd/Private/Dialogs/Dialogs.cpp
@@ -576,7 +576,8 @@ class SModalDialogWithCheckbox : public SCompoundWidget
 {
 public:
 	SLATE_BEGIN_ARGS( SModalDialogWithCheckbox )
-	:_bHasCancelButton(false) 
+	: _bHasCancelButton(false)
+	, _WrapMessageAt(512.0f)
 	{}
 		/** Warning message displayed on the dialog */
 		SLATE_ATTRIBUTE(FText, Message)
@@ -596,6 +597,9 @@ public:
 		/** Default value of the checkbox */
 		SLATE_ARGUMENT(bool, bDefaultCheckValue)
 
+		/** Wrap message at specified length, zero or negative number will disable the wrapping */
+		SLATE_ARGUMENT(float, WrapMessageAt)
+
 		/** Typically an icon to help the user more easily identify the nature of the issue */
 		SLATE_ATTRIBUTE( const FSlateBrush*, Image )
 
@@ -651,8 +655,8 @@ public:
 						.Padding(FMargin(16.f, 0.f, 0.f, 0.f))
 						[
 							SNew( STextBlock )
-							.WrapTextAt(512.0f)
-							.Text( MyMessage )
+							.WrapTextAt(InArgs._WrapMessageAt)
+							.Text(MyMessage)
 							.Font(MessageFont)
 						]
 					]
@@ -853,7 +857,8 @@ FSuppressableWarningDialog::FSuppressableWarningDialog(const FSetupInfo& Info)
 			.bDefaultCheckValue(Info.bDefaultToSuppressInTheFuture)
 			.CheckboxMessage(Info.CheckBoxText)
 			.ParentWindow(ModalWindow)
-			.Image((Info.Image != NULL) ? Info.Image : DefaultImage);
+			.Image((Info.Image != NULL) ? Info.Image : DefaultImage)
+			.WrapMessageAt(Info.WrapMessageAt);
 
 		ModalWindow->SetContent( MessageBox.ToSharedRef() );
 	}
diff --git a/Engine/Source/Editor/UnrealEd/Public/Dialogs/Dialogs.h b/Engine/Source/Editor/UnrealEd/Public/Dialogs/Dialogs.h
--- a/Engine/Source/Editor/UnrealEd/Public/Dialogs/Dialogs.h
+++ b/Engine/Source/Editor/UnrealEd/Public/Dialogs/Dialogs.h
@@ -102,6 +102,9 @@ public:
 		/** If true suppression will not persist for future editor sessions */
 		bool bDontPersistSuppressionAcrossSessions;
 
+		/** Wrap message at specified length, zero or negative number will disable the wrapping */
+		float WrapMessageAt;
+
 		/** Text used on the button which will return FSuppressableWarningDialog::Confirm */
 		FText ConfirmText;
 
@@ -129,6 +132,7 @@ public:
 			, IniSettingFileName(InIniSettingFileName)
 			, bDefaultToSuppressInTheFuture(false)
 			, bDontPersistSuppressionAcrossSessions(false)
+			, WrapMessageAt(512.0f)
 			, ConfirmText()
 			, CancelText()
 			, CheckBoxText(NSLOCTEXT("ModalDialogs", "DefaultCheckBoxMessage", "Don't show this again"))
