Check Parameters All ON / OFF with context menu

diff --git a/Engine/Source/Editor/MaterialEditor/Private/MaterialEditorInstanceDetailCustomization.cpp b/Engine/Source/Editor/MaterialEditor/Private/MaterialEditorInstanceDetailCustomization.cpp
--- a/Engine/Source/Editor/MaterialEditor/Private/MaterialEditorInstanceDetailCustomization.cpp
+++ b/Engine/Source/Editor/MaterialEditor/Private/MaterialEditorInstanceDetailCustomization.cpp
@@ -1,3 +1,4 @@
+
 // Copyright Epic Games, Inc. All Rights Reserved.
 
 #include "MaterialEditorInstanceDetailCustomization.h"
@@ -309,7 +310,7 @@ void FMaterialInstanceParameterDetails::CreateGroupsWidget(TSharedRef<IPropertyH
 				FUIAction PasteAction(
 					FExecuteAction::CreateSP(this, &FMaterialInstanceParameterDetails::OnPasteParameterValues, GroupIdx),
 					FCanExecuteAction::CreateSP(this, &FMaterialInstanceParameterDetails::CanPasteParameterValues, GroupIdx));
-				DetailGroup.HeaderRow()
+				FDetailWidgetRow& HeaderRow = DetailGroup.HeaderRow()
 					.CopyAction(CopyAction)
 					.PasteAction(PasteAction)
 					.NameContent()
@@ -319,6 +320,24 @@ void FMaterialInstanceParameterDetails::CreateGroupsWidget(TSharedRef<IPropertyH
 					];
 
 				CreateSingleGroupWidget(ParameterGroup, ParameterGroupsProperty->GetChildHandle(GroupIdx), DetailGroup);
+
+				HeaderRow.AddCustomContextMenuAction(FUIAction(
+						FExecuteAction::CreateLambda([&]()mutable 
+						{
+							EnableGroupParameters(ParameterGroup, true);
+						})),
+							LOCTEXT("ToggleParametersEnable", "Enable All Parameters"),
+							LOCTEXT("ToggleParametersEnableTooltip", "Enable All Parameters in group"),
+							FSlateIcon());
+
+				HeaderRow.AddCustomContextMenuAction(FUIAction(
+						FExecuteAction::CreateLambda([&]()mutable 
+						{
+							EnableGroupParameters(ParameterGroup, false);
+						})),
+							LOCTEXT("ToggleParametersDisable", "Disable All Parameters"),
+							LOCTEXT("ToggleParametersDisableTooltip", "Disable All Parameters in group"),
+							FSlateIcon());
 			}
 		}
 	}
@@ -374,7 +393,15 @@ void FMaterialInstanceParameterDetails::CreateGroupsWidget(TSharedRef<IPropertyH
 	}
 }
 
-
+void FMaterialInstanceParameterDetails::EnableGroupParameters(FEditorParameterGroup& ParameterGroup, bool ShouldEnable)
+{
+	// loop through each parameter in the group and toggle to enable/disable them all
+	for (int32 ParamIdx = 0; ParamIdx < ParameterGroup.Parameters.Num(); ++ParamIdx)
+	{
+		UDEditorParameterValue* Parameter = ParameterGroup.Parameters[ParamIdx];
+		Parameter->bOverride = ShouldEnable;
+	}
+}
 
 void FMaterialInstanceParameterDetails::CreateSingleGroupWidget(FEditorParameterGroup& ParameterGroup, TSharedPtr<IPropertyHandle> ParameterGroupProperty, IDetailGroup& DetailGroup )
 {
diff --git a/Engine/Source/Editor/MaterialEditor/Private/MaterialEditorInstanceDetailCustomization.h b/Engine/Source/Editor/MaterialEditor/Private/MaterialEditorInstanceDetailCustomization.h
--- a/Engine/Source/Editor/MaterialEditor/Private/MaterialEditorInstanceDetailCustomization.h
+++ b/Engine/Source/Editor/MaterialEditor/Private/MaterialEditorInstanceDetailCustomization.h
@@ -47,6 +47,9 @@ private:
 	/** Builds the widget for an individual parameter group */
 	void CreateSingleGroupWidget(struct FEditorParameterGroup& ParameterGroup, TSharedPtr<IPropertyHandle> ParameterGroupProperty, class IDetailGroup& DetailGroup);
 
+	/** Enable/Disable all parameter properties in a group */
+	static void EnableGroupParameters(struct FEditorParameterGroup& ParameterGroup, bool ShouldEnable);
+
 	/** These methods generate the custom widgets for the various parameter types */
 	void CreateParameterValueWidget(class UDEditorParameterValue* Parameter, TSharedPtr<IPropertyHandle> ParameterProperty, IDetailGroup& DetailGroup);
 	void CreateMaskParameterValueWidget(class UDEditorParameterValue* Parameter, TSharedPtr<IPropertyHandle> ParameterProperty, IDetailGroup& DetailGroup);
