Implements UMaterialInstance::bHasPixelAnimation

#tests win64

diff --git a/Engine/Plugins/Editor/MaterialAnalyzer/Source/Private/SMaterialAnalyzer.cpp b/Engine/Plugins/Editor/MaterialAnalyzer/Source/Private/SMaterialAnalyzer.cpp
--- a/Engine/Plugins/Editor/MaterialAnalyzer/Source/Private/SMaterialAnalyzer.cpp
+++ b/Engine/Plugins/Editor/MaterialAnalyzer/Source/Private/SMaterialAnalyzer.cpp
@@ -38,6 +38,7 @@ SMaterialAnalyzer::SMaterialAnalyzer()
 	BasePropertyOverrideNames.Add(TEXT("bOverride_TwoSided"), TEXT("TwoSidedOverride"));
 	BasePropertyOverrideNames.Add(TEXT("bOverride_bIsThinSurface"), TEXT("bIsThinSurfaceOverride"));
 	BasePropertyOverrideNames.Add(TEXT("bOverride_OutputTranslucentVelocity"), TEXT("bOutputTranslucentVelocity"));
+	BasePropertyOverrideNames.Add(TEXT("bOverride_bHasPixelAnimation"), TEXT("bHasPixelAnimation"));
 }
 
 SMaterialAnalyzer::~SMaterialAnalyzer()
diff --git a/Engine/Source/Editor/MaterialEditor/Private/MaterialEditorInstanceDetailCustomization.cpp b/Engine/Source/Editor/MaterialEditor/Private/MaterialEditorInstanceDetailCustomization.cpp
--- a/Engine/Source/Editor/MaterialEditor/Private/MaterialEditorInstanceDetailCustomization.cpp
+++ b/Engine/Source/Editor/MaterialEditor/Private/MaterialEditorInstanceDetailCustomization.cpp
@@ -1259,7 +1259,8 @@ void FMaterialInstanceParameterDetails::CreateBasePropertyOverrideWidgets(IDetai
 	TAttribute<bool> IsOverrideIsThinSurfaceEnabled = TAttribute<bool>::Create(TAttribute<bool>::FGetter::CreateSP(this, &FMaterialInstanceParameterDetails::OverrideIsThinSurfaceEnabled));
 	TAttribute<bool> IsOverrideDitheredLODTransitionEnabled = TAttribute<bool>::Create(TAttribute<bool>::FGetter::CreateSP(this, &FMaterialInstanceParameterDetails::OverrideDitheredLODTransitionEnabled));
 	TAttribute<bool> IsOverrideOutputTranslucentVelocityEnabled = TAttribute<bool>::Create(TAttribute<bool>::FGetter::CreateSP(this, &FMaterialInstanceParameterDetails::OverrideOutputTranslucentVelocityEnabled));
-	TAttribute<bool> IsOverrideDisplacementScalingEnabled = TAttribute<bool>::Create(TAttribute<bool>::FGetter::CreateSP(this, &FMaterialInstanceParameterDetails::OverrideDisplacementScalingEnabled)); 
+	TAttribute<bool> IsOverrideHasPixelAnimationEnabled = TAttribute<bool>::Create(TAttribute<bool>::FGetter::CreateSP(this, &FMaterialInstanceParameterDetails::OverrideHasPixelAnimationEnabled));
+	TAttribute<bool> IsOverrideDisplacementScalingEnabled = TAttribute<bool>::Create(TAttribute<bool>::FGetter::CreateSP(this, &FMaterialInstanceParameterDetails::OverrideDisplacementScalingEnabled));
 	TAttribute<bool> IsOverrideMaxWorldPositionOffsetDisplacementEnabled = TAttribute<bool>::Create(TAttribute<bool>::FGetter::CreateSP(this, &FMaterialInstanceParameterDetails::OverrideMaxWorldPositionOffsetDisplacementEnabled));
 
 	TSharedRef<IPropertyHandle> BasePropertyOverridePropery = DetailLayout.GetProperty("BasePropertyOverrides");
@@ -1270,6 +1271,7 @@ void FMaterialInstanceParameterDetails::CreateBasePropertyOverrideWidgets(IDetai
 	TSharedPtr<IPropertyHandle> IsThinSurfaceProperty = BasePropertyOverridePropery->GetChildHandle("IsThinSurface");
 	TSharedPtr<IPropertyHandle> DitheredLODTransitionProperty = BasePropertyOverridePropery->GetChildHandle("DitheredLODTransition");
 	TSharedPtr<IPropertyHandle> OutputTranslucentVelocityProperty = BasePropertyOverridePropery->GetChildHandle("bOutputTranslucentVelocity");
+	TSharedPtr<IPropertyHandle> HasPixelAnimationProperty = BasePropertyOverridePropery->GetChildHandle("bHasPixelAnimation");
 	TSharedPtr<IPropertyHandle> DisplacementScalingProperty = BasePropertyOverridePropery->GetChildHandle("DisplacementScaling");
 	TSharedPtr<IPropertyHandle> MaxWorldPositionOffsetDisplacementProperty = BasePropertyOverridePropery->GetChildHandle("MaxWorldPositionOffsetDisplacement");
 
@@ -1437,6 +1439,25 @@ void FMaterialInstanceParameterDetails::CreateBasePropertyOverrideWidgets(IDetai
 			.Visibility(TAttribute<EVisibility>::Create(TAttribute<EVisibility>::FGetter::CreateSP(this, &FMaterialInstanceParameterDetails::IsOverriddenAndVisible, IsOverrideOutputTranslucentVelocityEnabled)))
 			.OverrideResetToDefault(ResetOutputTranslucentVelocityPropertyOverride);
 	}
+	{
+		FIsResetToDefaultVisible IsHasPixelAnimationPropertyResetVisible = FIsResetToDefaultVisible::CreateLambda([this](TSharedPtr<IPropertyHandle> InHandle) {
+			return MaterialEditorInstance->Parent != nullptr ? MaterialEditorInstance->BasePropertyOverrides.bHasPixelAnimation != MaterialEditorInstance->Parent->HasPixelAnimation() : false;
+			});
+		FResetToDefaultHandler ResetHasPixelAnimationPropertyHandler = FResetToDefaultHandler::CreateLambda([this](TSharedPtr<IPropertyHandle> InHandle) {
+			if (MaterialEditorInstance->Parent != nullptr)
+			{
+				MaterialEditorInstance->BasePropertyOverrides.bHasPixelAnimation = MaterialEditorInstance->Parent->HasPixelAnimation();
+			}
+			});
+		FResetToDefaultOverride ResetHasPixelAnimationPropertyOverride = FResetToDefaultOverride::Create(IsHasPixelAnimationPropertyResetVisible, ResetHasPixelAnimationPropertyHandler);
+		IDetailPropertyRow& HasPixelAnimationPropertyRow = BasePropertyOverrideGroup.AddPropertyRow(HasPixelAnimationProperty.ToSharedRef());
+		HasPixelAnimationPropertyRow
+			.DisplayName(HasPixelAnimationProperty->GetPropertyDisplayName())
+			.ToolTip(bStaticParametersOverrideDisabled ? ParameterDisabledToolTipString : HasPixelAnimationProperty->GetToolTipText())
+			.EditCondition(IsOverrideHasPixelAnimationEnabled, FOnBooleanValueChanged::CreateSP(this, &FMaterialInstanceParameterDetails::OnOverrideHasPixelAnimationChanged))
+			.Visibility(TAttribute<EVisibility>::Create(TAttribute<EVisibility>::FGetter::CreateSP(this, &FMaterialInstanceParameterDetails::IsOverriddenAndVisible, IsOverrideHasPixelAnimationEnabled)))
+			.OverrideResetToDefault(ResetHasPixelAnimationPropertyOverride);
+	}
 	{
 		FIsResetToDefaultVisible IsDisplacementScalingPropertyResetVisible = FIsResetToDefaultVisible::CreateLambda([this](TSharedPtr<IPropertyHandle> InHandle)
 		{
@@ -1525,6 +1546,11 @@ bool FMaterialInstanceParameterDetails::OverrideOutputTranslucentVelocityEnabled
 	return MaterialEditorInstance->BasePropertyOverrides.bOverride_OutputTranslucentVelocity;
 }
 
+bool FMaterialInstanceParameterDetails::OverrideHasPixelAnimationEnabled() const
+{
+	return MaterialEditorInstance->BasePropertyOverrides.bOverride_bHasPixelAnimation;
+}
+
 bool FMaterialInstanceParameterDetails::OverrideDisplacementScalingEnabled() const
 {
 	return MaterialEditorInstance->BasePropertyOverrides.bOverride_DisplacementScaling;
@@ -1624,6 +1650,17 @@ void FMaterialInstanceParameterDetails::OnOverrideOutputTranslucentVelocityChang
 	FEditorSupportDelegates::RedrawAllViewports.Broadcast();
 }
 
+void FMaterialInstanceParameterDetails::OnOverrideHasPixelAnimationChanged(bool NewValue)
+{
+	if (DoesSourceMaterialInstanceDisallowStaticParameterPermutation(MaterialEditorInstance, NewValue))
+	{
+		return;
+	}
+	MaterialEditorInstance->BasePropertyOverrides.bOverride_bHasPixelAnimation = NewValue;
+	MaterialEditorInstance->PostEditChange();
+	FEditorSupportDelegates::RedrawAllViewports.Broadcast();
+}
+
 void FMaterialInstanceParameterDetails::OnOverrideDisplacementScalingChanged(bool NewValue)
 {
 	MaterialEditorInstance->BasePropertyOverrides.bOverride_DisplacementScaling = NewValue;
diff --git a/Engine/Source/Editor/MaterialEditor/Private/MaterialEditorInstanceDetailCustomization.h b/Engine/Source/Editor/MaterialEditor/Private/MaterialEditorInstanceDetailCustomization.h
--- a/Engine/Source/Editor/MaterialEditor/Private/MaterialEditorInstanceDetailCustomization.h
+++ b/Engine/Source/Editor/MaterialEditor/Private/MaterialEditorInstanceDetailCustomization.h
@@ -117,6 +117,7 @@ private:
 	bool OverrideIsThinSurfaceEnabled() const;
 	bool OverrideDitheredLODTransitionEnabled() const;
 	bool OverrideOutputTranslucentVelocityEnabled() const;
+	bool OverrideHasPixelAnimationEnabled() const;
 	bool OverrideDisplacementScalingEnabled() const;
 	bool OverrideMaxWorldPositionOffsetDisplacementEnabled() const;
 	void OnOverrideOpacityClipMaskValueChanged(bool NewValue);
@@ -126,6 +127,7 @@ private:
 	void OnOverrideIsThinSurfaceChanged(bool NewValue);
 	void OnOverrideDitheredLODTransitionChanged(bool NewValue);
 	void OnOverrideOutputTranslucentVelocityChanged(bool NewValue);
+	void OnOverrideHasPixelAnimationChanged(bool NewValue);
 	void OnOverrideDisplacementScalingChanged(bool NewValue);
 	void OnOverrideMaxWorldPositionOffsetDisplacementChanged(bool NewValue);
 
diff --git a/Engine/Source/Editor/UnrealEd/Private/PreviewMaterial.cpp b/Engine/Source/Editor/UnrealEd/Private/PreviewMaterial.cpp
--- a/Engine/Source/Editor/UnrealEd/Private/PreviewMaterial.cpp
+++ b/Engine/Source/Editor/UnrealEd/Private/PreviewMaterial.cpp
@@ -941,6 +941,10 @@ void UMaterialEditorInstanceConstant::CopyBasePropertiesFromParent()
 	{
 		BasePropertyOverrides.bOutputTranslucentVelocity = SourceInstance->IsTranslucencyWritingVelocity();
 	}
+	if (!BasePropertyOverrides.bOverride_bHasPixelAnimation)
+	{
+		BasePropertyOverrides.bHasPixelAnimation = SourceInstance->HasPixelAnimation();
+	}
 	if (!BasePropertyOverrides.DitheredLODTransition)
 	{
 		BasePropertyOverrides.DitheredLODTransition = SourceInstance->IsDitheredLODTransition();
diff --git a/Engine/Source/Runtime/Engine/Classes/Materials/MaterialInstance.h b/Engine/Source/Runtime/Engine/Classes/Materials/MaterialInstance.h
--- a/Engine/Source/Runtime/Engine/Classes/Materials/MaterialInstance.h
+++ b/Engine/Source/Runtime/Engine/Classes/Materials/MaterialInstance.h
@@ -581,6 +581,7 @@ class UMaterialInstance : public UMaterialInterface
 	uint8 bCastDynamicShadowAsMasked : 1;
 	uint8 bOutputTranslucentVelocity : 1;
 	uint8 bIsShadingModelFromMaterialExpression : 1;
+	uint8 bHasPixelAnimation : 1;
 protected:
 	uint8 bLoadedCachedData : 1;
 public:
diff --git a/Engine/Source/Runtime/Engine/Classes/Materials/MaterialInstanceBasePropertyOverrides.h b/Engine/Source/Runtime/Engine/Classes/Materials/MaterialInstanceBasePropertyOverrides.h
--- a/Engine/Source/Runtime/Engine/Classes/Materials/MaterialInstanceBasePropertyOverrides.h
+++ b/Engine/Source/Runtime/Engine/Classes/Materials/MaterialInstanceBasePropertyOverrides.h
@@ -44,6 +44,10 @@ struct FMaterialInstanceBasePropertyOverrides
 	/** Enables override of the output velocity property. */
 	UPROPERTY(EditAnywhere, Category = Material)
 	uint8 bOverride_OutputTranslucentVelocity : 1;
+	
+	/** Enables override of the has pixel animation property. */
+	UPROPERTY(EditAnywhere, Category = Material)
+	uint8 bOverride_bHasPixelAnimation : 1;
 
 	/** Enables override of the displacement magnitude and center property. */
 	UPROPERTY(EditAnywhere, Category = Material)
@@ -72,6 +76,12 @@ struct FMaterialInstanceBasePropertyOverrides
 	/** Whether the material should output velocity even though it has a translucent blend mode. */
 	UPROPERTY(EditAnywhere, Category = Material, meta = (editcondition = "bOverride_OutputTranslucentVelocity"))
 	uint8 bOutputTranslucentVelocity : 1;
+	
+	/** Whether the opaque material has any pixel animations happening, that isn't included in the geometric velocities.
+	 * This allows to disable renderer's heuristics that assumes animation is fully described with motion vector, such as TSR's anti-flickering heuristic.
+	 */
+	UPROPERTY(EditAnywhere, Category = Material, meta = (editcondition = "bOverride_bHasPixelAnimation"))
+	uint8 bHasPixelAnimation : 1;
 
 	/** The blend mode */
 	UPROPERTY(EditAnywhere, Category = Material, meta = (editcondition = "bOverride_BlendMode"))
diff --git a/Engine/Source/Runtime/Engine/Private/Materials/MaterialInstance.cpp b/Engine/Source/Runtime/Engine/Private/Materials/MaterialInstance.cpp
--- a/Engine/Source/Runtime/Engine/Private/Materials/MaterialInstance.cpp
+++ b/Engine/Source/Runtime/Engine/Private/Materials/MaterialInstance.cpp
@@ -2084,6 +2084,7 @@ void UMaterialInstance::UpdateOverridableBaseProperties()
 		DitheredLODTransition = 0;
 		bIsShadingModelFromMaterialExpression = 0;
 		bOutputTranslucentVelocity = false;
+		bHasPixelAnimation = false;
 		DisplacementScaling = FDisplacementScaling();
 		MaxWorldPositionOffsetDisplacement = 0.0f;
 		return;
@@ -2119,6 +2120,16 @@ void UMaterialInstance::UpdateOverridableBaseProperties()
 		BasePropertyOverrides.bOutputTranslucentVelocity = bOutputTranslucentVelocity;
 	}
 
+	if (BasePropertyOverrides.bOverride_bHasPixelAnimation)
+	{
+		bHasPixelAnimation = BasePropertyOverrides.bHasPixelAnimation;
+	}
+	else
+	{
+		bHasPixelAnimation = Parent->HasPixelAnimation();
+		BasePropertyOverrides.bHasPixelAnimation = bHasPixelAnimation;
+	}
+
 	if (BasePropertyOverrides.bOverride_ShadingModel)
 	{
 		if (BasePropertyOverrides.ShadingModel == MSM_FromMaterialExpression)
@@ -4503,6 +4514,7 @@ bool UMaterialInstance::HasOverridenBaseProperties()const
 		(IsDitheredLODTransition() != Parent->IsDitheredLODTransition()) ||
 		(GetCastDynamicShadowAsMasked() != Parent->GetCastDynamicShadowAsMasked()) ||
 		(IsTranslucencyWritingVelocity() != Parent->IsTranslucencyWritingVelocity()) ||
+		(HasPixelAnimation() != Parent->HasPixelAnimation()) ||
 		(GetDisplacementScaling() != Parent->GetDisplacementScaling()) ||
 		(GetMaxWorldPositionOffsetDisplacement() != Parent->GetMaxWorldPositionOffsetDisplacement())
 		))
@@ -4527,6 +4539,7 @@ FString UMaterialInstance::GetBasePropertyOverrideString() const
 		BasePropString += FString::Printf(TEXT("bOverride_DitheredLODTransition_%d, "), (IsDitheredLODTransition() != Parent->IsDitheredLODTransition()));
 		BasePropString += FString::Printf(TEXT("bOverride_CastDynamicShadowAsMasked_%d, "), (GetCastDynamicShadowAsMasked() != Parent->GetCastDynamicShadowAsMasked()));
 		BasePropString += FString::Printf(TEXT("bOverride_OutputTranslucentVelocity_%d "), (IsTranslucencyWritingVelocity() != Parent->IsTranslucencyWritingVelocity()));
+		BasePropString += FString::Printf(TEXT("bOverride_bHasPixelAnimation_%d "), (HasPixelAnimation() != Parent->HasPixelAnimation()));
 		BasePropString += FString::Printf(TEXT("bOverride_DisplacementScaling_%d "), (GetDisplacementScaling() != Parent->GetDisplacementScaling()));
 		BasePropString += FString::Printf(TEXT("bOverride_MaxWorldPositionOffsetDisplacement_%d "), (GetMaxWorldPositionOffsetDisplacement() != Parent->GetMaxWorldPositionOffsetDisplacement()));
 	}
@@ -4596,7 +4609,8 @@ bool UMaterialInstance::ShouldAlwaysEvaluateWorldPositionOffset() const
 
 bool UMaterialInstance::HasPixelAnimation() const
 {
-	return Parent ? Parent->HasPixelAnimation() : false;
+	const UMaterial* Material = GetMaterial();
+	return Material ? bHasPixelAnimation && GetMaterial()->MaterialDomain == MD_Surface : false;
 }
 
 bool UMaterialInstance::IsMasked() const
@@ -4806,7 +4820,8 @@ bool UMaterialInstance::IsRedundant() const
 	|| BasePropertyOverrides.bOverride_DitheredLODTransition 
 	|| BasePropertyOverrides.bOverride_CastDynamicShadowAsMasked
 	|| BasePropertyOverrides.bOverride_TwoSided 
-	|| BasePropertyOverrides.bOutputTranslucentVelocity)
+	|| BasePropertyOverrides.bOutputTranslucentVelocity
+	|| BasePropertyOverrides.bHasPixelAnimation)
 	{
 		return false;
 	}
diff --git a/Engine/Source/Runtime/Engine/Private/Materials/MaterialShared.cpp b/Engine/Source/Runtime/Engine/Private/Materials/MaterialShared.cpp
--- a/Engine/Source/Runtime/Engine/Private/Materials/MaterialShared.cpp
+++ b/Engine/Source/Runtime/Engine/Private/Materials/MaterialShared.cpp
@@ -1670,7 +1670,7 @@ bool FMaterialResource::ShouldWriteOnlyAlpha() const { return Material->bWriteOn
 bool FMaterialResource::ShouldEnableResponsiveAA() const { return Material->bEnableResponsiveAA; }
 bool FMaterialResource::ShouldDoSSR() const { return Material->bScreenSpaceReflections; }
 bool FMaterialResource::ShouldDoContactShadows() const { return Material->bContactShadows; }
-bool FMaterialResource::HasPixelAnimation() const { return Material->HasPixelAnimation() && GetMaterialDomain() == MD_Surface && IsOpaqueOrMaskedBlendMode(GetBlendMode()); }
+bool FMaterialResource::HasPixelAnimation() const { return (MaterialInstance ? MaterialInstance->HasPixelAnimation() : Material->HasPixelAnimation()) && GetMaterialDomain() == MD_Surface && IsOpaqueOrMaskedBlendMode(GetBlendMode()); }
 bool FMaterialResource::IsWireframe() const { return Material->Wireframe; }
 bool FMaterialResource::IsUIMaterial() const { return Material->MaterialDomain == MD_UI; }
 bool FMaterialResource::IsPostProcessMaterial() const { return Material->MaterialDomain == MD_PostProcess; }
@@ -4813,6 +4813,7 @@ FMaterialInstanceBasePropertyOverrides::FMaterialInstanceBasePropertyOverrides()
 	,bOverride_TwoSided(false)
 	,bOverride_bIsThinSurface(false)
 	,bOverride_OutputTranslucentVelocity(false)
+	,bOverride_bHasPixelAnimation(false)
 	,bOverride_DisplacementScaling(false)
 	,bOverride_MaxWorldPositionOffsetDisplacement(false)
 	,TwoSided(0)
@@ -4820,6 +4821,7 @@ FMaterialInstanceBasePropertyOverrides::FMaterialInstanceBasePropertyOverrides()
 	,DitheredLODTransition(0)
 	,bCastDynamicShadowAsMasked(false)
 	,bOutputTranslucentVelocity(false)
+	,bHasPixelAnimation(false)
 	,BlendMode(BLEND_Opaque)
 	,ShadingModel(MSM_DefaultLit)
 	,OpacityMaskClipValue(.333333f)
@@ -4837,6 +4839,7 @@ bool FMaterialInstanceBasePropertyOverrides::operator==(const FMaterialInstanceB
 		bOverride_bIsThinSurface == Other.bOverride_bIsThinSurface &&
 		bOverride_DitheredLODTransition == Other.bOverride_DitheredLODTransition &&
 		bOverride_OutputTranslucentVelocity == Other.bOverride_OutputTranslucentVelocity &&
+		bOverride_bHasPixelAnimation == Other.bOverride_bHasPixelAnimation &&
 		bOverride_DisplacementScaling == Other.bOverride_DisplacementScaling &&
 		bOverride_MaxWorldPositionOffsetDisplacement == Other.bOverride_MaxWorldPositionOffsetDisplacement &&
 		OpacityMaskClipValue == Other.OpacityMaskClipValue &&
@@ -4846,6 +4849,8 @@ bool FMaterialInstanceBasePropertyOverrides::operator==(const FMaterialInstanceB
 		bIsThinSurface == Other.bIsThinSurface &&
 		DitheredLODTransition == Other.DitheredLODTransition &&
 		bCastDynamicShadowAsMasked == Other.bCastDynamicShadowAsMasked &&
+		bOutputTranslucentVelocity == Other.bOutputTranslucentVelocity &&
+		bHasPixelAnimation == Other.bHasPixelAnimation &&
 		DisplacementScaling == Other.DisplacementScaling &&
 		MaxWorldPositionOffsetDisplacement == Other.MaxWorldPositionOffsetDisplacement;
 }
