Open/Close state of the Material Editors Detail Panel Parameter groups are not preserved through static-switch param changes or after Undo/Redo

diff --git a/Engine/Source/Editor/MaterialEditor/Private/MaterialEditorInstanceDetailCustomization.cpp b/Engine/Source/Editor/MaterialEditor/Private/MaterialEditorInstanceDetailCustomization.cpp
--- a/Engine/Source/Editor/MaterialEditor/Private/MaterialEditorInstanceDetailCustomization.cpp
+++ b/Engine/Source/Editor/MaterialEditor/Private/MaterialEditorInstanceDetailCustomization.cpp
@@ -301,8 +301,7 @@ void FMaterialInstanceParameterDetails::CreateGroupsWidget(TSharedRef<IPropertyH
 			
 			if (bCreateGroup)
 			{
-				IDetailGroup& DetailGroup = GroupsCategory.AddGroup(ParameterGroup.GroupName, FText::FromName(ParameterGroup.GroupName), false, true);
-
+				IDetailGroup& DetailGroup = GroupsCategory.AddGroup(ParameterGroup.GroupName, FText::FromName(ParameterGroup.GroupName), false, false);
 				FUIAction CopyAction(
 					FExecuteAction::CreateSP(this, &FMaterialInstanceParameterDetails::OnCopyParameterValues, GroupIdx),
 					FCanExecuteAction::CreateSP(this, &FMaterialInstanceParameterDetails::CanCopyParameterValues, GroupIdx));
