Fix "Promote to Variable" in Root node pins only creating Vector nodes.

diff --git a/Engine/Source/Editor/MaterialEditor/Private/MaterialEditor.cpp b/Engine/Source/Editor/MaterialEditor/Private/MaterialEditor.cpp
--- a/Engine/Source/Editor/MaterialEditor/Private/MaterialEditor.cpp
+++ b/Engine/Source/Editor/MaterialEditor/Private/MaterialEditor.cpp
@@ -4784,9 +4784,11 @@ UClass* FMaterialEditor::GetOnPromoteToParameterClass(const UEdGraphPin* TargetP
 
 	if (RootPinNode != nullptr)
 	{
-		EMaterialProperty propertyId = (EMaterialProperty)FCString::Atoi(*TargetPin->PinType.PinSubCategory.ToString());
-
-		switch (propertyId)
+		const UMaterialGraph* MaterialGraph = CastChecked<UMaterialGraph>(RootPinNode->GetGraph());
+		const FMaterialInputInfo& MaterialInput = MaterialGraph->MaterialInputs[TargetPin->SourceIndex];
+		EMaterialProperty PropertyId = MaterialInput.GetProperty();
+		
+		switch (PropertyId)
 		{
 			case MP_Opacity:
 			case MP_Metallic:
