Report issues in base material for shader platforms not covered by current RHI feature levels.

diff --git a/Engine/Source/Editor/MaterialEditor/Private/MaterialEditor.cpp b/Engine/Source/Editor/MaterialEditor/Private/MaterialEditor.cpp
--- a/Engine/Source/Editor/MaterialEditor/Private/MaterialEditor.cpp
+++ b/Engine/Source/Editor/MaterialEditor/Private/MaterialEditor.cpp
@@ -2632,8 +2632,7 @@ bool FMaterialEditor::UpdateOriginalMaterial()
 			for (int32 QualityLevel = 0; QualityLevel < EMaterialQualityLevel::Num; ++QualityLevel)
 			{
 				const auto& PlatformData = PlatformPtr->GetPlatformData((EMaterialQualityLevel::Type)QualityLevel);
-				// base material is covered by previous check
-				for (int32 InstanceIndex = 1; InstanceIndex < PlatformData.Instances.Num(); ++InstanceIndex)
+				for (int32 InstanceIndex = 0; InstanceIndex < PlatformData.Instances.Num(); ++InstanceIndex)
 				{
 					const FMaterialResource* CurrentResource = PlatformData.Instances[InstanceIndex].MaterialResourcesStats;
 					if (CurrentResource && CurrentResource->GetCompileErrors().Num() > 0)
@@ -2642,14 +2641,20 @@ bool FMaterialEditor::UpdateOriginalMaterial()
 						const auto& AssetName = MaterialStatsManager->GetMaterialName(InstanceIndex);
 						const FString QualityName = FMaterialStatsUtils::MaterialQualityToShortString((EMaterialQualityLevel::Type)QualityLevel);
 
-						Errors.Push(FText::Format(NSLOCTEXT("UnrealEd", "Warning_CompileErrorsInMaterial_ListEntryDerivedMaterial", "- Material instance {0} for platform {1} at quality level {2}."), FText::FromString(*AssetName), FText::FromName(PlatformName), FText::FromString(QualityName)));
+						if (InstanceIndex == 0)
+						{
+							Errors.Push(FText::Format(NSLOCTEXT("UnrealEd", "Warning_CompileErrorsInMaterial_ListEntryMaterial", "- For platform {0} at quality level {1}."), FText::FromName(PlatformName), FText::FromString(QualityName)));
+						}
+						else
+						{
+							Errors.Push(FText::Format(NSLOCTEXT("UnrealEd", "Warning_CompileErrorsInMaterial_ListEntryDerivedMaterial", "- Material instance {0} for platform {1} at quality level {2}."), FText::FromString(*AssetName), FText::FromName(PlatformName), FText::FromString(QualityName)));
+						}
 					}
 				}
 			}
 		}
 	}
 
-
 	if (Errors.Num() > 0)
 	{
 		const FText JoinedErrors = FText::Join(FText::FromString(TEXT("\n")), Errors);
