Material Root node inline control updates

diff --git a/Engine/Source/Editor/GraphEditor/Public/KismetPins/SGraphPinNumSlider.h b/Engine/Source/Editor/GraphEditor/Public/KismetPins/SGraphPinNumSlider.h
--- a/Engine/Source/Editor/GraphEditor/Public/KismetPins/SGraphPinNumSlider.h
+++ b/Engine/Source/Editor/GraphEditor/Public/KismetPins/SGraphPinNumSlider.h
@@ -46,6 +46,7 @@ protected:
 			.MaxDesiredWidth(400)
 			[
 				SNew(SNumericEntryBox<NumericType>)
+				.MinDesiredValueWidth(100.0f)
 				.EditableTextBoxStyle(FAppStyle::Get(), "Graph.EditableTextBox")
 				.BorderForegroundColor(FSlateColor::UseForeground())
 				.Visibility(this, &SGraphPinNumSlider::GetDefaultValueVisibility)
@@ -84,7 +85,7 @@ protected:
 		{
 			const FScopedTransaction Transaction(NSLOCTEXT("GraphEditor", "ChangeNumberPinValue", "Change Number Pin Value"));
 			GraphPinObj->Modify();
-			GraphPinObj->GetSchema()->TrySetDefaultValue(*GraphPinObj, *LexToString(NewValue));
+			GraphPinObj->GetSchema()->TrySetDefaultValue(*GraphPinObj, *LexToSanitizedString(NewValue));
 			LastSliderCommittedValue = NewValue;
 		}
 	}
@@ -113,20 +114,6 @@ protected:
 		return bIsUsingSlider ? SliderValue : Num;
 	}
 
-public:
-	// virtual EVisibility GetDefaultValueVisibility() const override
-	// {
-	// 	UEdGraphPin* GraphPin = GetPinObj();
-	// 	if (GraphPin->Direction == EGPD_Output)
-	// 	{
-	// 		return EVisibility::Collapsed;
-	// 	}
-	// 	else
-	// 	{
-	// 		return EVisibility::Visible;
-	// 	}
-	// }
-
 private:
 	FProperty* PinProperty;
 	NumericType LastSliderCommittedValue;
diff --git a/Engine/Source/Editor/GraphEditor/Public/KismetPins/SVector2DSlider.h b/Engine/Source/Editor/GraphEditor/Public/KismetPins/SVector2DSlider.h
--- a/Engine/Source/Editor/GraphEditor/Public/KismetPins/SVector2DSlider.h
+++ b/Engine/Source/Editor/GraphEditor/Public/KismetPins/SVector2DSlider.h
@@ -61,6 +61,8 @@ public:
 					// Create Text box 0 
 					SNew( SNumericEntryBox<NumericType> )
 					.LabelVAlign(VAlign_Center)
+					.MaxFractionalDigits(3)
+					.MinDesiredValueWidth(45.0f)
 					.Label()
 					[
 						SNew( STextBlock )
@@ -96,6 +98,8 @@ public:
 					// Create Text box 1
 					SNew( SNumericEntryBox<NumericType> )
 					.LabelVAlign(VAlign_Center)
+					.MaxFractionalDigits(3)
+					.MinDesiredValueWidth(45.0f)
 					.Label()
 					[
 						SNew( STextBlock )
diff --git a/Engine/Source/Editor/GraphEditor/Public/KismetPins/SVector4Slider.h b/Engine/Source/Editor/GraphEditor/Public/KismetPins/SVector4Slider.h
--- a/Engine/Source/Editor/GraphEditor/Public/KismetPins/SVector4Slider.h
+++ b/Engine/Source/Editor/GraphEditor/Public/KismetPins/SVector4Slider.h
@@ -57,10 +57,9 @@ public:
 
 		this->ChildSlot
 			[
-				SNew(SVerticalBox)
-				+ SVerticalBox::Slot()
-				.AutoHeight()
-				.Padding(0)
+				SNew(SBox)
+				.MinDesiredWidth(275)
+				.MaxDesiredWidth(400)
 				[
 					SNew(SHorizontalBox)
 					+ SHorizontalBox::Slot()
@@ -69,6 +68,8 @@ public:
 						// Create Text box 0 
 						SNew(SNumericEntryBox<NumericType>)
 						.LabelVAlign(VAlign_Center)
+						.MaxFractionalDigits(3)
+						.MinDesiredValueWidth(45.0f)
 						.Label()
 						[
 							SNew(STextBlock)
@@ -104,6 +105,8 @@ public:
 						// Create Text box 1
 						SNew(SNumericEntryBox<NumericType>)
 						.LabelVAlign(VAlign_Center)
+						.MaxFractionalDigits(3)
+						.MinDesiredValueWidth(45.0f)
 						.Label()
 						[
 							SNew(STextBlock)
@@ -139,6 +142,8 @@ public:
 						// Create Text box 2
 						SNew(SNumericEntryBox<NumericType>)
 						.LabelVAlign(VAlign_Center)
+						.MaxFractionalDigits(3)
+						.MinDesiredValueWidth(45.0f)
 						.Label()
 						[
 							SNew(STextBlock)
@@ -174,6 +179,8 @@ public:
 						// Create Text box 3
 						SNew(SNumericEntryBox<NumericType>)
 						.LabelVAlign(VAlign_Center)
+						.MaxFractionalDigits(3)
+						.MinDesiredValueWidth(45.0f)
 						.Label()
 						[
 							SNew(STextBlock)
diff --git a/Engine/Source/Editor/GraphEditor/Public/KismetPins/SVectorSlider.h b/Engine/Source/Editor/GraphEditor/Public/KismetPins/SVectorSlider.h
--- a/Engine/Source/Editor/GraphEditor/Public/KismetPins/SVectorSlider.h
+++ b/Engine/Source/Editor/GraphEditor/Public/KismetPins/SVectorSlider.h
@@ -68,6 +68,8 @@ public:
 					// Create Text box 0 
 					SNew( SNumericEntryBox<NumericType> )
 					.LabelVAlign(VAlign_Center)
+					.MaxFractionalDigits(3)
+					.MinDesiredValueWidth(45.0f)
 					.Label()
 					[
 						SNew( STextBlock )
@@ -103,6 +105,8 @@ public:
 					// Create Text box 1
 					SNew( SNumericEntryBox<NumericType> )
 					.LabelVAlign(VAlign_Center)
+					.MaxFractionalDigits(3)
+					.MinDesiredValueWidth(45.0f)
 					.Label()
 					[
 						SNew( STextBlock )
@@ -134,11 +138,13 @@ public:
 					.BorderBackgroundColor( FLinearColor::White )
 				]
 				+ SHorizontalBox::Slot()
-					.AutoWidth().Padding(2).HAlign(HAlign_Fill)
-					[
-						// Create Text box 1
-						SNew(SNumericEntryBox<NumericType>)
-						.LabelVAlign(VAlign_Center)
+				.AutoWidth().Padding(2).HAlign(HAlign_Fill)
+				[
+					// Create Text box 2
+					SNew(SNumericEntryBox<NumericType>)
+					.MaxFractionalDigits(3)
+					.MinDesiredValueWidth(45.0f)
+					.LabelVAlign(VAlign_Center)
 					.Label()
 					[
 						SNew(STextBlock)
@@ -168,7 +174,7 @@ public:
 					.EditableTextBoxStyle(&FAppStyle::GetWidgetStyle<FEditableTextBoxStyle>("Graph.VectorEditableTextBox"))
 					.BorderForegroundColor(FLinearColor::White)
 					.BorderBackgroundColor(FLinearColor::White)
-					]
+				]
 			]
 		];
 	}
diff --git a/Engine/Source/Editor/MaterialEditor/Private/MaterialEditor.cpp b/Engine/Source/Editor/MaterialEditor/Private/MaterialEditor.cpp
--- a/Engine/Source/Editor/MaterialEditor/Private/MaterialEditor.cpp
+++ b/Engine/Source/Editor/MaterialEditor/Private/MaterialEditor.cpp
@@ -2130,6 +2130,22 @@ void FMaterialEditor::AddGraphEditorPinActionsToContextMenu(FToolMenuSection& In
 		);
 	}
 
+	// Reset to Default Value
+	{
+		FToolUIAction ResetToDefaultAction;
+		ResetToDefaultAction.ExecuteAction = FToolMenuExecuteAction::CreateSP(this, &FMaterialEditor::OnResetToDefault);
+		ResetToDefaultAction.IsActionVisibleDelegate = FToolMenuIsActionButtonVisible::CreateSP(this, &FMaterialEditor::OnCanResetToDefault);
+
+		TSharedPtr<FUICommandInfo> ResetToDefaultCommand = FMaterialEditorCommands::Get().ResetToDefault;
+		InSection.AddMenuEntry(
+			ResetToDefaultCommand->GetCommandName(),
+			ResetToDefaultCommand->GetLabel(),
+			ResetToDefaultCommand->GetDescription(),
+			ResetToDefaultCommand->GetIcon(),
+			ResetToDefaultAction
+		);
+	}
+
 	{
 		auto AddSubstrateContextualMenu = [&](TSharedPtr<FUICommandInfo> CreateNodeCommand, ESubstrateNodeForPin SubstrateNodeForPin)
 		{
@@ -4902,6 +4918,165 @@ void FMaterialEditor::OnPromoteToParameter(const FToolMenuContext& InMenuContext
 	}
 }
 
+void FMaterialEditor::OnResetToDefault(const FToolMenuContext& InMenuContext) const
+{
+	UGraphNodeContextMenuContext* NodeContext = InMenuContext.FindContext<UGraphNodeContextMenuContext>();
+	const int32 PinIndex = NodeContext->Pin->SourceIndex;
+	const UMaterialGraphNode_Root* RootPinNode = Cast<UMaterialGraphNode_Root>(NodeContext->Pin->GetOwningNode());
+	UEdGraphPin* TargetPin = RootPinNode->GetPinAt(PinIndex);
+
+	if (RootPinNode != nullptr)
+	{
+		const FScopedTransaction Transaction( NSLOCTEXT("GraphEditor", "ResetPinToDefault", "Reset Pin Value to its default" ) );
+		TargetPin->Modify();
+		
+		const UMaterialGraph* MaterialGraph = CastChecked<UMaterialGraph>(RootPinNode->GetGraph());
+		const FMaterialInputInfo& MaterialInput = MaterialGraph->MaterialInputs[TargetPin->SourceIndex];
+		const EMaterialProperty PropertyId = MaterialInput.GetProperty();
+		UMaterialEditorOnlyData* EditorOnlyData = Material->GetEditorOnlyData();
+		switch (PropertyId)
+		{
+			case MP_BaseColor:
+				EditorOnlyData->BaseColor.Constant = FColor(128, 128, 128);
+				TargetPin->GetSchema()->TrySetDefaultValue(*TargetPin, EditorOnlyData->BaseColor.GetDefaultValue());
+				break;
+
+			case MP_Opacity:
+				EditorOnlyData->Opacity.Constant = FMaterialAttributeDefinitionMap::GetDefaultValue(MP_Opacity).X;
+				TargetPin->GetSchema()->TrySetDefaultValue(*TargetPin, EditorOnlyData->Opacity.GetDefaultValue());
+				break;
+			
+			case MP_Metallic:
+				EditorOnlyData->Metallic.Constant = FMaterialAttributeDefinitionMap::GetDefaultValue(MP_Metallic).X;
+				TargetPin->GetSchema()->TrySetDefaultValue(*TargetPin, EditorOnlyData->Metallic.GetDefaultValue());
+				break;
+			
+			case MP_Specular:
+				EditorOnlyData->Specular.Constant = FMaterialAttributeDefinitionMap::GetDefaultValue(MP_Specular).X;
+				TargetPin->GetSchema()->TrySetDefaultValue(*TargetPin, EditorOnlyData->Specular.GetDefaultValue());
+				break;
+			
+			case MP_Roughness:
+				EditorOnlyData->Roughness.Constant = FMaterialAttributeDefinitionMap::GetDefaultValue(MP_Roughness).X;
+				TargetPin->GetSchema()->TrySetDefaultValue(*TargetPin, EditorOnlyData->Roughness.GetDefaultValue());
+				break;
+			
+			case MP_Anisotropy:
+				EditorOnlyData->Anisotropy.Constant = FMaterialAttributeDefinitionMap::GetDefaultValue(MP_Anisotropy).X;
+				TargetPin->GetSchema()->TrySetDefaultValue(*TargetPin, EditorOnlyData->Anisotropy.GetDefaultValue());
+				break;
+			
+			case MP_CustomData0:
+				EditorOnlyData->ClearCoat.Constant = FMaterialAttributeDefinitionMap::GetDefaultValue(MP_CustomData0).X;
+				TargetPin->GetSchema()->TrySetDefaultValue(*TargetPin, EditorOnlyData->ClearCoat.GetDefaultValue());
+				break;
+			
+			case MP_CustomData1:
+				EditorOnlyData->ClearCoatRoughness.Constant = FMaterialAttributeDefinitionMap::GetDefaultValue(MP_CustomData1).X;
+				TargetPin->GetSchema()->TrySetDefaultValue(*TargetPin, EditorOnlyData->ClearCoatRoughness.GetDefaultValue());
+				break;
+			
+			case MP_AmbientOcclusion:
+				EditorOnlyData->AmbientOcclusion.Constant = FMaterialAttributeDefinitionMap::GetDefaultValue(MP_AmbientOcclusion).X;
+				TargetPin->GetSchema()->TrySetDefaultValue(*TargetPin, EditorOnlyData->AmbientOcclusion.GetDefaultValue());
+				break;
+			
+			case MP_Refraction:
+				EditorOnlyData->Refraction.Constant = FMaterialAttributeDefinitionMap::GetDefaultValue(MP_Refraction).X;
+				TargetPin->GetSchema()->TrySetDefaultValue(*TargetPin, EditorOnlyData->Refraction.GetDefaultValue());
+				break;
+			
+			case MP_OpacityMask:
+				EditorOnlyData->OpacityMask.Constant = FMaterialAttributeDefinitionMap::GetDefaultValue(MP_OpacityMask).X;
+				TargetPin->GetSchema()->TrySetDefaultValue(*TargetPin, EditorOnlyData->OpacityMask.GetDefaultValue());
+				break;
+			
+			case MP_SurfaceThickness:
+				EditorOnlyData->SurfaceThickness.Constant = FMaterialAttributeDefinitionMap::GetDefaultValue(MP_SurfaceThickness).X;
+				TargetPin->GetSchema()->TrySetDefaultValue(*TargetPin, EditorOnlyData->SurfaceThickness.GetDefaultValue());
+				break;
+			
+			case MP_Displacement:
+				EditorOnlyData->Displacement.Constant = FMaterialAttributeDefinitionMap::GetDefaultValue(MP_Displacement).X;
+				TargetPin->GetSchema()->TrySetDefaultValue(*TargetPin, EditorOnlyData->Displacement.GetDefaultValue());
+				break;
+			
+			case MP_WorldPositionOffset:
+				EditorOnlyData->WorldPositionOffset.Constant = FMaterialAttributeDefinitionMap::GetDefaultValue(MP_WorldPositionOffset);
+				TargetPin->GetSchema()->TrySetDefaultValue(*TargetPin, EditorOnlyData->WorldPositionOffset.GetDefaultValue());
+				break;
+			
+			case MP_EmissiveColor:
+				EditorOnlyData->EmissiveColor.Constant = FLinearColor(FMaterialAttributeDefinitionMap::GetDefaultValue(MP_EmissiveColor)).ToFColorSRGB();
+				TargetPin->GetSchema()->TrySetDefaultValue(*TargetPin, EditorOnlyData->EmissiveColor.GetDefaultValue());
+				break;
+			
+			case MP_SubsurfaceColor:
+				EditorOnlyData->SubsurfaceColor.Constant = FLinearColor(FMaterialAttributeDefinitionMap::GetDefaultValue(MP_SubsurfaceColor)).ToFColorSRGB(); 
+				TargetPin->GetSchema()->TrySetDefaultValue(*TargetPin, EditorOnlyData->SubsurfaceColor.GetDefaultValue());
+				break;
+			
+			case MP_Normal:
+				EditorOnlyData->Normal.Constant = FMaterialAttributeDefinitionMap::GetDefaultValue(MP_Normal); 
+				TargetPin->GetSchema()->TrySetDefaultValue(*TargetPin, EditorOnlyData->Normal.GetDefaultValue());
+				break;
+			
+			case MP_Tangent:
+				EditorOnlyData->Tangent.Constant = FMaterialAttributeDefinitionMap::GetDefaultValue(MP_Tangent);
+				TargetPin->GetSchema()->TrySetDefaultValue(*TargetPin, EditorOnlyData->Tangent.GetDefaultValue());
+				break;
+			
+			default:
+				break;
+		}
+		FMaterialEditorUtilities::UpdateMaterialAfterGraphChange(MaterialGraph);
+	}
+}
+
+bool FMaterialEditor::OnCanResetToDefault(const FToolMenuContext& InMenuContext) const
+{
+	UGraphNodeContextMenuContext* NodeContext = InMenuContext.FindContext<UGraphNodeContextMenuContext>();
+	if (!NodeContext)
+	{
+		return false;
+	}
+	const UEdGraphPin* TargetPin = NodeContext->Pin;
+	const UMaterialGraphNode_Root* RootPinNode = Cast<UMaterialGraphNode_Root>(TargetPin->GetOwningNode());
+	if (RootPinNode != nullptr)
+	{
+		const UMaterialGraph* MaterialGraph = CastChecked<UMaterialGraph>(RootPinNode->GetGraph());
+		const FMaterialInputInfo& MaterialInput = MaterialGraph->MaterialInputs[TargetPin->SourceIndex];
+		const EMaterialProperty PropertyId = MaterialInput.GetProperty();
+		UMaterialEditorOnlyData* EditorOnlyData = Material->GetEditorOnlyData();
+		switch (PropertyId)
+		{
+			case MP_BaseColor:				
+			case MP_Opacity:				
+			case MP_Metallic:				
+			case MP_Specular:				
+			case MP_Roughness:				
+			case MP_Anisotropy:				
+			case MP_CustomData0:			
+			case MP_CustomData1:			
+			case MP_AmbientOcclusion:		
+			case MP_Refraction:				
+			case MP_OpacityMask:			
+			case MP_SurfaceThickness:		
+			case MP_Displacement:			
+			case MP_WorldPositionOffset:	
+			case MP_EmissiveColor:			
+			case MP_SubsurfaceColor:		
+			case MP_Normal:					
+			case MP_Tangent:				
+				return true;
+				
+			default:
+				return false;
+		}
+	}
+	return false;
+}
+
 bool FMaterialEditor::OnCanPromoteToParameter(const FToolMenuContext& InMenuContext) const
 {
 	UGraphNodeContextMenuContext* NodeContext = InMenuContext.FindContext<UGraphNodeContextMenuContext>();
diff --git a/Engine/Source/Editor/MaterialEditor/Private/MaterialEditor.h b/Engine/Source/Editor/MaterialEditor/Private/MaterialEditor.h
--- a/Engine/Source/Editor/MaterialEditor/Private/MaterialEditor.h
+++ b/Engine/Source/Editor/MaterialEditor/Private/MaterialEditor.h
@@ -809,13 +809,19 @@ private:
 
 	/** Will promote selected pin to a parameter of the pin type */
 	void OnPromoteToParameter(const FToolMenuContext& InMenuContext) const;
-
+	
 	/** Used to know if we can promote selected pin to a parameter of the pin type */
 	bool OnCanPromoteToParameter(const FToolMenuContext& InMenuContext) const;
 
 	/** Will  return the UClass to create from the Pin Type */
 	UClass* GetOnPromoteToParameterClass(const UEdGraphPin* TargetPin) const;
 
+	/** Used to know if we can reset the selected pin to it's default value */
+	bool OnCanResetToDefault(const FToolMenuContext& InMenuContext) const;
+
+	/** Will reset selected pin to it's default value */
+	void OnResetToDefault(const FToolMenuContext& InMenuContext) const;
+	
 	enum class ESubstrateNodeForPin : uint8
 	{
 		Slab,
diff --git a/Engine/Source/Editor/MaterialEditor/Private/MaterialEditorActions.cpp b/Engine/Source/Editor/MaterialEditor/Private/MaterialEditorActions.cpp
--- a/Engine/Source/Editor/MaterialEditor/Private/MaterialEditorActions.cpp
+++ b/Engine/Source/Editor/MaterialEditor/Private/MaterialEditorActions.cpp
@@ -81,6 +81,7 @@ void FMaterialEditorCommands::RegisterCommands()
 	UI_COMMAND( CreateHorizontalMixNode, "Create Horizontal Mix Node", "Create a Horizontal Mix Node linked to the pin", EUserInterfaceActionType::Button, FInputChord());
 	UI_COMMAND( CreateVerticalLayerNode, "Create Vertical Layer Node", "Create a Vertical Layer Node linked to the pin", EUserInterfaceActionType::Button, FInputChord());
 	UI_COMMAND( CreateWeightNode, "Create Weight Node", "Create a Weight Node linked to the pin", EUserInterfaceActionType::Button, FInputChord());
+	UI_COMMAND( ResetToDefault, "Reset Pin to Default Value", "Reset pin value to default for that type", EUserInterfaceActionType::Button, FInputChord());
 
 	UI_COMMAND(QualityLevel_All, "All", "Sets node preview to show all quality levels.)", EUserInterfaceActionType::RadioButton, FInputChord());
 	UI_COMMAND(QualityLevel_Epic, "Epic", "Sets node preview to Epic quality.", EUserInterfaceActionType::RadioButton, FInputChord());
diff --git a/Engine/Source/Editor/MaterialEditor/Public/MaterialEditorActions.h b/Engine/Source/Editor/MaterialEditor/Public/MaterialEditorActions.h
--- a/Engine/Source/Editor/MaterialEditor/Public/MaterialEditorActions.h
+++ b/Engine/Source/Editor/MaterialEditor/Public/MaterialEditorActions.h
@@ -190,13 +190,22 @@ public:
 	/** Create component mask node */
 	TSharedPtr< FUICommandInfo > CreateComponentMaskNode;
 
-	/** Create component mask node */
+	/** Promote pin to parameter */
 	TSharedPtr< FUICommandInfo > PromoteToParameter;
+	
+	/** Reset pin to default value */
+	TSharedPtr< FUICommandInfo > ResetToDefault;
 
-	/** Create component mask node */
+	/** Create slab node */
 	TSharedPtr< FUICommandInfo > CreateSlabNode;
+
+	/** Create horizontal mix node */
 	TSharedPtr< FUICommandInfo > CreateHorizontalMixNode;
+
+	/** Create vertical layer node */
 	TSharedPtr< FUICommandInfo > CreateVerticalLayerNode;
+
+	/** Create weight node */
 	TSharedPtr< FUICommandInfo > CreateWeightNode;
 
 	TSharedPtr< FUICommandInfo > QualityLevel_All;
