Material Editor: Disable inline sliders for now (pending UX revision) and revert to text input

diff --git a/Engine/Source/Editor/MaterialEditor/Private/MaterialEditorGraphPanelPinFactory.cpp b/Engine/Source/Editor/MaterialEditor/Private/MaterialEditorGraphPanelPinFactory.cpp
--- a/Engine/Source/Editor/MaterialEditor/Private/MaterialEditorGraphPanelPinFactory.cpp
+++ b/Engine/Source/Editor/MaterialEditor/Private/MaterialEditorGraphPanelPinFactory.cpp
@@ -23,11 +23,11 @@
 #include "MaterialGraph/MaterialGraphNode.h"
 #include "Materials/MaterialExpression.h"
 #include "Materials/MaterialExpressionScalarParameter.h"
-#include "KismetPins/SGraphPinNumSlider.h"
-#include "KismetPins/SGraphPinVector2DSlider.h"
-#include "KismetPins/SGraphPinVectorSlider.h"
-#include "KismetPins/SGraphPinVector4Slider.h"
-#include "KismetPins/SGraphPinIntegerSlider.h"
+#include "KismetPins/SGraphPinNum.h"
+#include "KismetPins/SGraphPinVector2D.h"
+#include "KismetPins/SGraphPinVector.h"
+#include "KismetPins/SGraphPinVector4.h"
+#include "KismetPins/SGraphPinInteger.h"
 
 TSharedPtr<class SGraphPin> FMaterialEditorGraphPanelPinFactory::CreatePin(class UEdGraphPin* InPin) const
 {
@@ -51,15 +51,15 @@ TSharedPtr<class SGraphPin> FMaterialEditorGraphPanelPinFactory::CreatePin(class
 
 			if (InPin->PinType.PinSubCategory == MaterialGraphSchema->PSC_Red || InPin->PinType.PinSubCategory == MaterialGraphSchema->PSC_Float)
 			{
-				return SNew(SGraphPinNumSlider<double>, InPin, InProperty);	
+				return SNew(SGraphPinNum<double>, InPin);	
 			}
 			else if (InPin->PinType.PinSubCategory == MaterialGraphSchema->PSC_RG)
 			{
-				return SNew(SGraphPinVector2DSlider<float>, InPin, InProperty);
+				return SNew(SGraphPinVector2D<float>, InPin);
 			}
 			else if (InPin->PinType.PinSubCategory == MaterialGraphSchema->PSC_RGB)
 			{
-				return SNew(SGraphPinVectorSlider<float>, InPin, InProperty);
+				return SNew(SGraphPinVector<float>, InPin);
 			}
 			else if (InPin->PinType.PinSubCategory == MaterialGraphSchema->PSC_RGBA)
 			{
@@ -67,11 +67,11 @@ TSharedPtr<class SGraphPin> FMaterialEditorGraphPanelPinFactory::CreatePin(class
 			}
 			else if (InPin->PinType.PinSubCategory == MaterialGraphSchema->PSC_Vector4)
 			{
-				return SNew(SGraphPinVector4Slider<float>, InPin, InProperty);
+				return SNew(SGraphPinVector4<float>, InPin);
 			}
 			else if (InPin->PinType.PinSubCategory == MaterialGraphSchema->PSC_Int)
 			{
-				return SNew(SGraphPinIntegerSlider, InPin, InProperty);
+				return SNew(SGraphPinInteger, InPin);
 			}
 			else if (InPin->PinType.PinSubCategory == MaterialGraphSchema->PSC_Byte)
 			{
@@ -82,7 +82,7 @@ TSharedPtr<class SGraphPin> FMaterialEditorGraphPanelPinFactory::CreatePin(class
 				}
 				else
 				{
-					return SNew(SGraphPinIntegerSlider, InPin, InProperty);
+					return SNew(SGraphPinInteger, InPin);
 				}
 			}
 			else if (InPin->PinType.PinSubCategory == MaterialGraphSchema->PSC_Bool)
