Material Root node inline control updates

diff --git a/Engine/Source/Editor/MaterialEditor/Private/MaterialEditorGraphPanelPinFactory.cpp b/Engine/Source/Editor/MaterialEditor/Private/MaterialEditorGraphPanelPinFactory.cpp
--- a/Engine/Source/Editor/MaterialEditor/Private/MaterialEditorGraphPanelPinFactory.cpp
+++ b/Engine/Source/Editor/MaterialEditor/Private/MaterialEditorGraphPanelPinFactory.cpp
@@ -28,15 +28,21 @@
 #include "KismetPins/SGraphPinVector.h"
 #include "KismetPins/SGraphPinVector4.h"
 #include "KismetPins/SGraphPinInteger.h"
+#include "MaterialValueType.h"
 
 TSharedPtr<class SGraphPin> FMaterialEditorGraphPanelPinFactory::CreatePin(class UEdGraphPin* InPin) const
 {
 	if (const UMaterialGraphSchema* MaterialGraphSchema = Cast<const UMaterialGraphSchema>(InPin->GetSchema()))
 	{
+		UMaterialGraph* MaterialGraph = CastChecked<UMaterialGraph>(InPin->GetOwningNode()->GetGraph());
 		if (InPin->PinType.PinCategory == MaterialGraphSchema->PC_Exec)
 		{
 			return SNew(SGraphPinExec, InPin);
 		}
+		else if (!MaterialGraph->IsInputActive(InPin) && InPin->PinType.PinCategory == MaterialGraphSchema->PC_MaterialInput)
+		{
+			return SNew(SGraphPinMaterialInput, InPin);
+		}
 		else
 		{
 			// get metadata and pass it along
