Fix warning V547: Expression 'CompilationOutput.RootOperatorIndex >= 0' is always true. Unsigned type value is always >= 0.

diff --git a/Engine/Source/Editor/MaterialEditor/Private/SMaterialEditorSubstrateWidget.cpp b/Engine/Source/Editor/MaterialEditor/Private/SMaterialEditorSubstrateWidget.cpp
--- a/Engine/Source/Editor/MaterialEditor/Private/SMaterialEditorSubstrateWidget.cpp
+++ b/Engine/Source/Editor/MaterialEditor/Private/SMaterialEditorSubstrateWidget.cpp
@@ -503,22 +503,7 @@ void SMaterialEditorSubstrateWidget::Tick(const FGeometry& AllottedGeometry, con
 				MaterialDescription += CompilationOutput.SubstrateMaterialDescription;
 
 				// Now generate a visual representation of the material from the topology tree of operators.
-				{
-					if (CompilationOutput.RootOperatorIndex >= 0)
-					{										
-						MaterialBox->SetContent(FSubstrateWidget::ProcessOperator(CompilationOutput));
-					}
-					else
-					{
-						// The tree does not looks sane so generate a visual error without crashing.
-						auto TreeError = SNew(SErrorText)
-							.ErrorText(LOCTEXT("TreeError", "Tree Error"))
-							.BackgroundColor(FSlateColor(EStyleColor::AccentRed));
-
-						const TSharedRef<SWidget>& TreeErrorAsShared = TreeError->AsShared();
-						MaterialBox->SetContent(TreeErrorAsShared);
-					}
-				}
+				MaterialBox->SetContent(FSubstrateWidget::ProcessOperator(CompilationOutput));
 			}
 			else
 			{
