Substrate - re-instated remove dead code (from 29421605) to be safe.

diff --git a/Engine/Shaders/Public/ShaderVersion.ush b/Engine/Shaders/Public/ShaderVersion.ush
--- a/Engine/Shaders/Public/ShaderVersion.ush
+++ b/Engine/Shaders/Public/ShaderVersion.ush
@@ -2,4 +2,4 @@
 // Each time the console command is executed it generates a new GUID. As this file is included
 // in Platform.ush (which should be included in any shader) it allows to invalidate the shader DDC.
 // 
-#pragma message("UESHADERMETADATA_VERSION 78C3F233-CBA3-4DB9-8BF5-F1FE3E490DD3")
+#pragma message("UESHADERMETADATA_VERSION CF29F038-BD76-4128-AE71-6D133200D942")
diff --git a/Engine/Source/Editor/MaterialEditor/Private/SMaterialEditorSubstrateWidget.cpp b/Engine/Source/Editor/MaterialEditor/Private/SMaterialEditorSubstrateWidget.cpp
--- a/Engine/Source/Editor/MaterialEditor/Private/SMaterialEditorSubstrateWidget.cpp
+++ b/Engine/Source/Editor/MaterialEditor/Private/SMaterialEditorSubstrateWidget.cpp
@@ -503,7 +503,19 @@ void SMaterialEditorSubstrateWidget::Tick(const FGeometry& AllottedGeometry, con
 				MaterialDescription += CompilationOutput.SubstrateMaterialDescription;
 
 				// Now generate a visual representation of the material from the topology tree of operators.
-				MaterialBox->SetContent(FSubstrateWidget::ProcessOperator(CompilationOutput));
+				if (CompilationOutput.RootOperatorIndex >= 0)
+				{
+					MaterialBox->SetContent(FSubstrateWidget::ProcessOperator(CompilationOutput));
+				}
+				else
+				{
+					// The tree does not looks sane so generate a visual error without crashing.
+					auto TreeError = SNew(SErrorText)
+						.ErrorText(LOCTEXT("TreeError", "Tree Error"))
+						.BackgroundColor(FSlateColor(EStyleColor::AccentRed));
+					const TSharedRef<SWidget>& TreeErrorAsShared = TreeError->AsShared();
+					MaterialBox->SetContent(TreeErrorAsShared);
+				}
 			}
 			else
 			{
diff --git a/Engine/Source/Runtime/Engine/Public/Rendering/SubstrateMaterialShared.h b/Engine/Source/Runtime/Engine/Public/Rendering/SubstrateMaterialShared.h
--- a/Engine/Source/Runtime/Engine/Public/Rendering/SubstrateMaterialShared.h
+++ b/Engine/Source/Runtime/Engine/Public/Rendering/SubstrateMaterialShared.h
@@ -142,7 +142,7 @@ public:
 	/** The byte per pixel count supported by the platform the material has been compiled against */
 	LAYOUT_FIELD_EDITORONLY(uint8, bMaterialOutOfBudgetHasBeenSimplified, 1);
 
-	LAYOUT_FIELD_EDITORONLY(uint8, RootOperatorIndex);
+	LAYOUT_FIELD_EDITORONLY(int8, RootOperatorIndex);
 	LAYOUT_ARRAY_EDITORONLY(FSubstrateOperator, Operators, SUBSTRATE_COMPILATION_OUTPUT_MAX_OPERATOR);
 };
 
