Expose bOverride_CastDynamicShadowAsMasked in MaterialInstance Editor.

diff --git a/Engine/Source/Editor/MaterialEditor/Private/MaterialEditorInstanceDetailCustomization.cpp b/Engine/Source/Editor/MaterialEditor/Private/MaterialEditorInstanceDetailCustomization.cpp
--- a/Engine/Source/Editor/MaterialEditor/Private/MaterialEditorInstanceDetailCustomization.cpp
+++ b/Engine/Source/Editor/MaterialEditor/Private/MaterialEditorInstanceDetailCustomization.cpp
@@ -1261,6 +1261,7 @@ void FMaterialInstanceParameterDetails::CreateBasePropertyOverrideWidgets(IDetai
 	TAttribute<bool> IsOverrideHasPixelAnimationEnabled = TAttribute<bool>::Create(TAttribute<bool>::FGetter::CreateSP(this, &FMaterialInstanceParameterDetails::OverrideHasPixelAnimationEnabled));
 	TAttribute<bool> IsOverrideDisplacementScalingEnabled = TAttribute<bool>::Create(TAttribute<bool>::FGetter::CreateSP(this, &FMaterialInstanceParameterDetails::OverrideDisplacementScalingEnabled));
 	TAttribute<bool> IsOverrideMaxWorldPositionOffsetDisplacementEnabled = TAttribute<bool>::Create(TAttribute<bool>::FGetter::CreateSP(this, &FMaterialInstanceParameterDetails::OverrideMaxWorldPositionOffsetDisplacementEnabled));
+	TAttribute<bool> IsOverrideCastDynamicShadowAsMaskedEnabled = TAttribute<bool>::Create(TAttribute<bool>::FGetter::CreateSP(this, &FMaterialInstanceParameterDetails::OverrideCastDynamicShadowAsMaskedEnabled));
 
 	TSharedRef<IPropertyHandle> BasePropertyOverridePropery = DetailLayout.GetProperty("BasePropertyOverrides");
 	TSharedPtr<IPropertyHandle> OpacityClipMaskValueProperty = BasePropertyOverridePropery->GetChildHandle("OpacityMaskClipValue");
@@ -1273,6 +1274,7 @@ void FMaterialInstanceParameterDetails::CreateBasePropertyOverrideWidgets(IDetai
 	TSharedPtr<IPropertyHandle> HasPixelAnimationProperty = BasePropertyOverridePropery->GetChildHandle("bHasPixelAnimation");
 	TSharedPtr<IPropertyHandle> DisplacementScalingProperty = BasePropertyOverridePropery->GetChildHandle("DisplacementScaling");
 	TSharedPtr<IPropertyHandle> MaxWorldPositionOffsetDisplacementProperty = BasePropertyOverridePropery->GetChildHandle("MaxWorldPositionOffsetDisplacement");
+	TSharedPtr<IPropertyHandle> CastDynamicShadowAsMaskedProperty = BasePropertyOverridePropery->GetChildHandle("bCastDynamicShadowAsMasked");
 
 	const FText ParameterDisabledToolTipString = FText::FromString(TEXT("This material instance parent restricts the creation of new shader permutations. Overriding this parameter would result in the generation of additional shader permutations."));
 	const bool bStaticParametersOverrideDisabled = MaterialEditorInstance->SourceInstance->bDisallowStaticParameterPermutations;
@@ -1498,6 +1500,25 @@ void FMaterialInstanceParameterDetails::CreateBasePropertyOverrideWidgets(IDetai
 			.Visibility(TAttribute<EVisibility>::Create(TAttribute<EVisibility>::FGetter::CreateSP(this, &FMaterialInstanceParameterDetails::IsOverriddenAndVisible, IsOverrideMaxWorldPositionOffsetDisplacementEnabled)))
 			.OverrideResetToDefault(ResetMaxWorldPositionOffsetDisplacementPropertyOverride);
 	}
+	{
+		FIsResetToDefaultVisible IsCastDynamicShadowAsMaskedPropertyResetVisible = FIsResetToDefaultVisible::CreateLambda([this](TSharedPtr<IPropertyHandle> InHandle) {
+			return MaterialEditorInstance->Parent != nullptr ? MaterialEditorInstance->BasePropertyOverrides.bCastDynamicShadowAsMasked != MaterialEditorInstance->Parent->GetCastDynamicShadowAsMasked() : false;
+			});
+		FResetToDefaultHandler ResetCastDynamicShadowAsMaskedPropertyHandler = FResetToDefaultHandler::CreateLambda([this](TSharedPtr<IPropertyHandle> InHandle) {
+			if (MaterialEditorInstance->Parent != nullptr)
+			{
+				MaterialEditorInstance->BasePropertyOverrides.bCastDynamicShadowAsMasked = MaterialEditorInstance->Parent->GetCastDynamicShadowAsMasked();
+			}
+			});
+		FResetToDefaultOverride ResetCastDynamicShadowAsMaskedPropertyOverride = FResetToDefaultOverride::Create(IsCastDynamicShadowAsMaskedPropertyResetVisible, ResetCastDynamicShadowAsMaskedPropertyHandler);
+		IDetailPropertyRow& CastDynamicShadowAsMaskedPropertyRow = BasePropertyOverrideGroup.AddPropertyRow(CastDynamicShadowAsMaskedProperty.ToSharedRef());
+		CastDynamicShadowAsMaskedPropertyRow
+			.DisplayName(CastDynamicShadowAsMaskedProperty->GetPropertyDisplayName())
+			.ToolTip(bStaticParametersOverrideDisabled ? ParameterDisabledToolTipString : CastDynamicShadowAsMaskedProperty->GetToolTipText())
+			.EditCondition(IsOverrideCastDynamicShadowAsMaskedEnabled, FOnBooleanValueChanged::CreateSP(this, &FMaterialInstanceParameterDetails::OnOverrideCastDynamicShadowAsMaskedChanged))
+			.Visibility(TAttribute<EVisibility>::Create(TAttribute<EVisibility>::FGetter::CreateSP(this, &FMaterialInstanceParameterDetails::IsOverriddenAndVisible, IsOverrideCastDynamicShadowAsMaskedEnabled)))
+			.OverrideResetToDefault(ResetCastDynamicShadowAsMaskedPropertyOverride);
+	}
 }
 
 EVisibility FMaterialInstanceParameterDetails::IsOverriddenAndVisible(TAttribute<bool> IsOverridden) const
@@ -1560,6 +1581,11 @@ bool FMaterialInstanceParameterDetails::OverrideMaxWorldPositionOffsetDisplaceme
 	return MaterialEditorInstance->BasePropertyOverrides.bOverride_MaxWorldPositionOffsetDisplacement;
 }
 
+bool FMaterialInstanceParameterDetails::OverrideCastDynamicShadowAsMaskedEnabled() const
+{
+	return MaterialEditorInstance->BasePropertyOverrides.bOverride_CastDynamicShadowAsMasked;
+}
+
 /** Helper function used by some parameters to verify that they are allowed to be overridden. This
   * must be prevented if the source material instance disallows the creation of new static parameter
   * permutations as that would trigger a new shader creation. */
@@ -1572,6 +1598,17 @@ static bool DoesSourceMaterialInstanceDisallowStaticParameterPermutation(const U
 	return false;
 }
 
+void FMaterialInstanceParameterDetails::OnOverrideCastDynamicShadowAsMaskedChanged(bool NewValue)
+{
+	if (DoesSourceMaterialInstanceDisallowStaticParameterPermutation(MaterialEditorInstance, NewValue))
+	{
+		return;
+	}
+	MaterialEditorInstance->BasePropertyOverrides.bOverride_CastDynamicShadowAsMasked = NewValue;
+	MaterialEditorInstance->PostEditChange();
+	FEditorSupportDelegates::RedrawAllViewports.Broadcast();
+}
+
 void FMaterialInstanceParameterDetails::OnOverrideOpacityClipMaskValueChanged(bool NewValue)
 {
 	if (DoesSourceMaterialInstanceDisallowStaticParameterPermutation(MaterialEditorInstance, NewValue))
diff --git a/Engine/Source/Editor/MaterialEditor/Private/MaterialEditorInstanceDetailCustomization.h b/Engine/Source/Editor/MaterialEditor/Private/MaterialEditorInstanceDetailCustomization.h
--- a/Engine/Source/Editor/MaterialEditor/Private/MaterialEditorInstanceDetailCustomization.h
+++ b/Engine/Source/Editor/MaterialEditor/Private/MaterialEditorInstanceDetailCustomization.h
@@ -130,6 +130,8 @@ private:
 	void OnOverrideHasPixelAnimationChanged(bool NewValue);
 	void OnOverrideDisplacementScalingChanged(bool NewValue);
 	void OnOverrideMaxWorldPositionOffsetDisplacementChanged(bool NewValue);
+	bool OverrideCastDynamicShadowAsMaskedEnabled() const;
+	void OnOverrideCastDynamicShadowAsMaskedChanged(bool NewValue);
 
 private:
 	/** Object that stores all of the possible parameters we can edit */
diff --git a/Engine/Source/Editor/UnrealEd/Private/PreviewMaterial.cpp b/Engine/Source/Editor/UnrealEd/Private/PreviewMaterial.cpp
--- a/Engine/Source/Editor/UnrealEd/Private/PreviewMaterial.cpp
+++ b/Engine/Source/Editor/UnrealEd/Private/PreviewMaterial.cpp
@@ -964,6 +964,10 @@ void UMaterialEditorInstanceConstant::CopyBasePropertiesFromParent()
 	{
 		BasePropertyOverrides.MaxWorldPositionOffsetDisplacement = SourceInstance->GetMaxWorldPositionOffsetDisplacement();
 	}
+	if (!BasePropertyOverrides.bOverride_CastDynamicShadowAsMasked)
+	{
+		BasePropertyOverrides.bCastDynamicShadowAsMasked = SourceInstance->GetCastDynamicShadowAsMasked();
+	}
 
 	// Copy the Lightmass settings...
 	// The lightmass functions (GetCastShadowAsMasked, etc.) check if the value is overridden and returns the current value if so, otherwise returns the parent value
diff --git a/Engine/Source/Runtime/Engine/Classes/Materials/MaterialInstanceBasePropertyOverrides.h b/Engine/Source/Runtime/Engine/Classes/Materials/MaterialInstanceBasePropertyOverrides.h
--- a/Engine/Source/Runtime/Engine/Classes/Materials/MaterialInstanceBasePropertyOverrides.h
+++ b/Engine/Source/Runtime/Engine/Classes/Materials/MaterialInstanceBasePropertyOverrides.h
@@ -70,7 +70,7 @@ struct FMaterialInstanceBasePropertyOverrides
 	uint8 DitheredLODTransition : 1;
 
 	/** Whether the material should cast shadows as masked even though it has a translucent blend mode. */
-	UPROPERTY(EditAnywhere, Category = Material, meta = (editcondition = "bOverride_CastShadowAsMasked", NoSpinbox = true))
+	UPROPERTY(EditAnywhere, Category = Material, meta = (editcondition = "bOverride_CastDynamicShadowAsMasked", NoSpinbox = true))
 	uint8 bCastDynamicShadowAsMasked:1;
 
 	/** Whether the material should output velocity even though it has a translucent blend mode. */
