MaterialEditor: Soft restriction on which blend modes are meant to be used with the alpha holdout blend mode.

diff --git a/Engine/Source/Editor/MaterialEditor/Private/MaterialEditorDetailCustomization.cpp b/Engine/Source/Editor/MaterialEditor/Private/MaterialEditorDetailCustomization.cpp
--- a/Engine/Source/Editor/MaterialEditor/Private/MaterialEditorDetailCustomization.cpp
+++ b/Engine/Source/Editor/MaterialEditor/Private/MaterialEditorDetailCustomization.cpp
@@ -18,6 +18,7 @@
 #include "Materials/MaterialExpressionTextureSampleParameter.h"
 #include "MaterialLayersFunctionsCustomization.h"
 #include "PropertyEditorModule.h"
+#include "PropertyRestriction.h"
 #include "MaterialEditor/MaterialEditorPreviewParameters.h"
 #include "Widgets/Input/SButton.h"
 #include "MaterialEditor/MaterialEditorInstanceConstant.h"
@@ -785,6 +786,7 @@ void FMaterialDetailCustomization::CustomizeDetails( IDetailLayoutBuilder& Detai
 
 	bool bUIMaterial = true;
 	bool bIsShadingModelFromMaterialExpression = false;
+	bool bIsAlphaHoldout = false;
 	for( TWeakObjectPtr<UObject>& Object : Objects )
 	{
 		UMaterial* Material = Cast<UMaterial>( Object.Get() );
@@ -794,6 +796,8 @@ void FMaterialDetailCustomization::CustomizeDetails( IDetailLayoutBuilder& Detai
 
 			// If any Object has its shading model from material expression
 			bIsShadingModelFromMaterialExpression |= Material->IsShadingModelFromMaterialExpression();
+			
+			bIsAlphaHoldout |= Material->BlendMode == BLEND_AlphaHoldout;
 		}
 		else
 		{
@@ -839,6 +843,12 @@ void FMaterialDetailCustomization::CustomizeDetails( IDetailLayoutBuilder& Detai
 			{
 				FByteProperty* BlendModeProperty = (FByteProperty*)Property;
 				BlendModeProperty->Enum = GetBlendModeEnum();
+
+				PropertyHandle->SetOnPropertyValueChanged(FSimpleDelegate::CreateLambda([&DetailLayout]()
+					{
+						// Refresh to update translucency pass restrictions below
+						DetailLayout.ForceRefreshDetails();
+					}));
 			}
 
 			if (PropertyName == GET_MEMBER_NAME_CHECKED(UMaterial, bEnableExecWire) && !AllowMaterialControlFlow())
@@ -861,6 +871,36 @@ void FMaterialDetailCustomization::CustomizeDetails( IDetailLayoutBuilder& Detai
 		}
 	}
 
+	// Translucency category
+	{
+		IDetailCategoryBuilder& TranslucencyCategory = DetailLayout.EditCategory(TEXT("Translucency"));
+
+		TArray<TSharedRef<IPropertyHandle>> AdvancedProperties;
+		TranslucencyCategory.GetDefaultProperties(AdvancedProperties, false, true);
+
+		for (TSharedRef<IPropertyHandle>& PropertyHandle : AdvancedProperties)
+		{
+			FProperty* Property = PropertyHandle->GetProperty();
+			FName PropertyName = Property->GetFName();
+
+			if (PropertyName == GET_MEMBER_NAME_CHECKED(UMaterial, TranslucencyPass))
+			{
+				if (bIsAlphaHoldout)
+				{
+					static FText RestrictReason = NSLOCTEXT("PropertyEditor", "TranslucencyPassRestriction", "Seperate translucency pass should not be used with Alpha Holdout blend mode.");
+					TSharedPtr<FPropertyRestriction> EnumRestriction = MakeShared<FPropertyRestriction>(RestrictReason);
+
+					FByteProperty* TranslucencyPassProperty = (FByteProperty*)Property;
+					EnumRestriction->AddDisabledValue(TranslucencyPassProperty->Enum->GetNameStringByValue(static_cast<int64>(MTP_AfterDOF)));
+					EnumRestriction->AddDisabledValue(TranslucencyPassProperty->Enum->GetNameStringByValue(static_cast<int64>(MTP_AfterMotionBlur)));
+
+					PropertyHandle->AddRestriction(EnumRestriction.ToSharedRef());
+				}
+			}
+		}
+	}
+
+
 	if( bUIMaterial )
 	{
 		DetailLayout.HideCategory( TEXT("TranslucencySelfShadowing") );
