[Shaders][Material Editor] - Add generic, data driven key, value shader statistics.

- Shader stats are simply key (FName) and value (POD types) that are stored per-shader.
- A Shader Statistic is a variant that can take on the following POD types: bool, float, int32, uint32
- Currently don't support strings since those weren't as trivial to implement (needed conversion and special case code to support serializing these to the memory images required by shader maps)
- Adds another shader statistics row in the Material Editor that will show the platform specific stats.
- The platform statistics are turned into a comma and newline separated string that is displayed to the user.
- Bump the shader output worker version since I'm changing the output file format (to contain the stats).
- The shader compiler output stores the stats as a `TMap<FString, FShaderStatVariant>`.  Serializing across FNames between Shader Compiler Worker and the main process wasn't working (different name tables).
- When storing and serializing the stats to the shader map we convert them to FNames. (see `FShader::FShader()`)

diff --git a/Engine/Shaders/Public/ShaderVersion.ush b/Engine/Shaders/Public/ShaderVersion.ush
--- a/Engine/Shaders/Public/ShaderVersion.ush
+++ b/Engine/Shaders/Public/ShaderVersion.ush
@@ -3,4 +3,4 @@
 // Each time the console command is executed it generates a new GUID. As this file is included
 // in Platform.ush (which should be included in any shader) it allows to invalidate the shader DDC.
 // 
-#pragma message("UESHADERMETADATA_VERSION 08294D12-4290-46A4-8DD7-0FDE597F2914")
+#pragma message("UESHADERMETADATA_VERSION 3c752777-4213-49d1-9357-e1b47a284227")
diff --git a/Engine/Source/Editor/MaterialEditor/Private/MaterialStats.h b/Engine/Source/Editor/MaterialEditor/Private/MaterialStats.h
--- a/Engine/Source/Editor/MaterialEditor/Private/MaterialStats.h
+++ b/Engine/Source/Editor/MaterialEditor/Private/MaterialStats.h
@@ -17,6 +17,7 @@ struct FShaderStatsInfo
 	};
 
 	TMap<ERepresentativeShader, FContent> ShaderInstructionCount;
+	TMap<ERepresentativeShader, FContent> GenericShaderStatistics;
 	FContent SamplersCount;
 	FContent InterpolatorsCount;
 	FContent TextureSampleCount;
@@ -27,6 +28,7 @@ struct FShaderStatsInfo
 	void Reset()
 	{
 		ShaderInstructionCount.Empty();
+		GenericShaderStatistics.Empty();
 
 		SamplersCount.StrDescription = TEXT("Compiling...");
 		SamplersCount.StrDescriptionLong = TEXT("Compiling...");
@@ -49,6 +51,7 @@ struct FShaderStatsInfo
 	void Empty()
 	{
 		ShaderInstructionCount.Empty();
+		GenericShaderStatistics.Empty();
 
 		SamplersCount.StrDescription.Empty();
 		SamplersCount.StrDescriptionLong.Empty();
diff --git a/Engine/Source/Editor/MaterialEditor/Private/MaterialStatsCommon.cpp b/Engine/Source/Editor/MaterialEditor/Private/MaterialStatsCommon.cpp
--- a/Engine/Source/Editor/MaterialEditor/Private/MaterialStatsCommon.cpp
+++ b/Engine/Source/Editor/MaterialEditor/Private/MaterialStatsCommon.cpp
@@ -527,6 +527,21 @@ void FMaterialStatsUtils::GetRepresentativeShaderTypesAndDescriptions(TMap<FName
 	}
 }
 
+static FString GetShaderString(const FShader::FShaderStatisticMap& Statistics)
+{
+	TStringBuilder<2048> StatisticsStrBuilder;
+	for (const auto& Stat : Statistics)
+	{
+		StatisticsStrBuilder << Stat.Key << ": ";
+		Visit([&StatisticsStrBuilder](auto& StoredValue)
+		{
+			StatisticsStrBuilder << StoredValue << "\n";
+		}, Stat.Value);
+	}
+
+	return StatisticsStrBuilder.ToString();
+}
+
 /**
 * Gets instruction counts that best represent the likely usage of this material based on shading model and other factors.
 * @param Results - an array of descriptions to be populated
@@ -563,6 +578,11 @@ void FMaterialStatsUtils::GetRepresentativeInstructionCounts(TArray<FShaderInstr
 						Info.ShaderType = ShaderInfo.ShaderType;
 						Info.ShaderDescription = ShaderInfo.ShaderDescription;
 						Info.InstructionCount = NumInstructions;
+						Info.ShaderStatisticsString = GetShaderString(MaterialShaderMap->GetShaderStatisticsMapForShader(ShaderType));
+						if (Info.ShaderStatisticsString.Len() == 0)
+						{
+							Info.ShaderStatisticsString = TEXT("n/a");
+						}
 
 						Results.Push(Info);
 
@@ -600,6 +620,11 @@ void FMaterialStatsUtils::GetRepresentativeInstructionCounts(TArray<FShaderInstr
 									Info.ShaderType = ShaderInfo.ShaderType;
 									Info.ShaderDescription = ShaderInfo.ShaderDescription;
 									Info.InstructionCount = NumInstructions;
+									Info.ShaderStatisticsString = GetShaderString(MeshShaderMap->GetShaderStatisticsMapForShader(*MaterialShaderMap, ShaderType));
+									if (Info.ShaderStatisticsString.Len() == 0)
+									{
+										Info.ShaderStatisticsString = TEXT("n/a");
+									}
 
 									Results.Push(Info);
 
@@ -646,6 +671,12 @@ void FMaterialStatsUtils::ExtractMatertialStatsInfo(EShaderPlatform ShaderPlatfo
 				TEXT("Offline shader compiler not available or an error was encountered!");
 
 			OutInfo.ShaderInstructionCount.Add(ShaderInstructionInfo[InstructionIndex].ShaderType, Content);
+
+			FString Description = ShaderInstructionInfo[InstructionIndex].ShaderStatisticsString;
+			FShaderStatsInfo::FContent GenericContent;
+			GenericContent.StrDescription = Description;
+			GenericContent.StrDescriptionLong = Description;
+			OutInfo.GenericShaderStatistics.Add(ShaderInstructionInfo[InstructionIndex].ShaderType, GenericContent);
 		}
 
 		// extract samplers info
diff --git a/Engine/Source/Editor/MaterialEditor/Private/MaterialStatsGrid.cpp b/Engine/Source/Editor/MaterialEditor/Private/MaterialStatsGrid.cpp
--- a/Engine/Source/Editor/MaterialEditor/Private/MaterialStatsGrid.cpp
+++ b/Engine/Source/Editor/MaterialEditor/Private/MaterialStatsGrid.cpp
@@ -76,6 +76,7 @@ void FStatsGridRow_Empty::CreateRow(TSharedPtr<FMaterialStats> StatsManager)
 	// just an array of empty cells
 	AddCell(FMaterialStatsGrid::DescriptorColumnName, MakeShareable(new FGridCell_Empty()));
 	AddCell(FMaterialStatsGrid::ShaderColumnName, MakeShareable(new FGridCell_Empty()));
+	AddCell(FMaterialStatsGrid::ShaderStatisticColumnName, MakeShareable(new FGridCell_Empty()));
 
 	FillPlatformCellsHelper(StatsManager);
 }
@@ -99,6 +100,7 @@ void FStatsGridRow_Name::CreateRow(TSharedPtr<FMaterialStats> StatsManager)
 	// we don't use a descriptor for this row
 	AddCell(FMaterialStatsGrid::DescriptorColumnName, MakeShareable(new FGridCell_Empty()));
 	AddCell(FMaterialStatsGrid::ShaderColumnName, MakeShareable(new FGridCell_Empty()));
+	AddCell(FMaterialStatsGrid::ShaderStatisticColumnName, MakeShareable(new FGridCell_Empty()));
 
 	FillPlatformCellsHelper(StatsManager);
 }
@@ -124,6 +126,7 @@ void FStatsGridRow_Quality::CreateRow(TSharedPtr<FMaterialStats> StatsManager)
 	// we don't use a descriptor for this row
 	AddCell(FMaterialStatsGrid::DescriptorColumnName, MakeShareable(new FGridCell_Empty()));
 	AddCell(FMaterialStatsGrid::ShaderColumnName, MakeShareable(new FGridCell_Empty()));
+	AddCell(FMaterialStatsGrid::ShaderStatisticColumnName, MakeShareable(new FGridCell_Empty()));
 
 	FillPlatformCellsHelper(StatsManager);
 }
@@ -148,10 +151,11 @@ void FStatsGridRow_Quality::AddPlatform(TSharedPtr<FMaterialStats> StatsManager,
 /*==============================================================================================================*/
 /* FStatsGridRow_Shaders functions*/
 
-FStatsGridRow_Shaders::FStatsGridRow_Shaders(ERepresentativeShader RepresentativeShader, bool bHeader)
+FStatsGridRow_Shaders::FStatsGridRow_Shaders(ERepresentativeShader RepresentativeShader, bool bHeader, bool bInstructionRow)
+	: bIsHeaderRow(bHeader)
+	, bInstructionRow(bInstructionRow)
+	, ShaderType(RepresentativeShader)
 {
-	bIsHeaderRow = bHeader;
-	ShaderType = RepresentativeShader;
 }
 
 FStatsGridRow_Shaders::EShaderClass FStatsGridRow_Shaders::GetShaderClass(const ERepresentativeShader Shader)
@@ -207,6 +211,10 @@ void FStatsGridRow_Shaders::CreateRow(TSharedPtr<FMaterialStats> StatsManager)
 	ShaderNameCell->SetColor(FStyleColors::Foreground);
 	AddCell(FMaterialStatsGrid::ShaderColumnName, ShaderNameCell);
 
+	// now add a cell that can display the what statistic we are describing.
+	FString ColumnContent = bInstructionRow ? TEXT("Instruction Count") : TEXT("Platform Statistics");
+	AddCell(FMaterialStatsGrid::ShaderStatisticColumnName, MakeShareable(new FGridCell_StaticString(ColumnContent, ColumnContent)));
+
 	FillPlatformCellsHelper(StatsManager);
 }
 
@@ -214,7 +222,7 @@ void FStatsGridRow_Shaders::AddPlatform(TSharedPtr<FMaterialStats> StatsManager,
 {
 	// add a cell that display the instruction count for this platform
 	const FString CellContent = FMaterialStatsUtils::MaterialQualityToShortString(QualityLevel);
-	TSharedPtr<FGridCell_ShaderValue> Cell = MakeShareable(new FGridCell_ShaderValue(StatsManager, EShaderInfoType::InstructionsCount, ShaderType, QualityLevel, Platform->GetPlatformShaderType(), InstanceIndex));
+	TSharedPtr<FGridCell_ShaderValue> Cell = MakeShareable(new FGridCell_ShaderValue(StatsManager, bInstructionRow ? EShaderInfoType::InstructionsCount : EShaderInfoType::GenericShaderStatistics, ShaderType, QualityLevel, Platform->GetPlatformShaderType(), InstanceIndex));
 
 	const FName ColumnName = FMaterialStatsGrid::MakePlatformColumnName(Platform, QualityLevel, InstanceIndex);
 	AddCell(ColumnName, Cell);
@@ -235,6 +243,7 @@ void FStatsGridRow_Samplers::CreateRow(TSharedPtr<FMaterialStats> StatsManager)
 	AddCell(FMaterialStatsGrid::DescriptorColumnName, HeaderCell);
 
 	AddCell(FMaterialStatsGrid::ShaderColumnName, MakeShareable(new FGridCell_Empty()));
+	AddCell(FMaterialStatsGrid::ShaderStatisticColumnName, MakeShareable(new FGridCell_Empty()));
 
 	FillPlatformCellsHelper(StatsManager);
 }
@@ -264,6 +273,7 @@ void FStatsGridRow_Interpolators::CreateRow(TSharedPtr<FMaterialStats> StatsMana
 	AddCell(FMaterialStatsGrid::DescriptorColumnName, HeaderCell);
 
 	AddCell(FMaterialStatsGrid::ShaderColumnName, MakeShareable(new FGridCell_Empty()));
+	AddCell(FMaterialStatsGrid::ShaderStatisticColumnName, MakeShareable(new FGridCell_Empty()));
 
 	FillPlatformCellsHelper(StatsManager);
 }
@@ -293,6 +303,7 @@ void FStatsGridRow_NumTextureSamples::CreateRow(TSharedPtr<FMaterialStats> Stats
 	AddCell(FMaterialStatsGrid::DescriptorColumnName, HeaderCell);
 
 	AddCell(FMaterialStatsGrid::ShaderColumnName, MakeShareable(new FGridCell_Empty()));
+	AddCell(FMaterialStatsGrid::ShaderStatisticColumnName, MakeShareable(new FGridCell_Empty()));
 
 	FillPlatformCellsHelper(StatsManager);
 }
@@ -322,6 +333,7 @@ void FStatsGridRow_NumVirtualTextureLookups::CreateRow(TSharedPtr<FMaterialStats
 	AddCell(FMaterialStatsGrid::DescriptorColumnName, HeaderCell);
 
 	AddCell(FMaterialStatsGrid::ShaderColumnName, MakeShareable(new FGridCell_Empty()));
+	AddCell(FMaterialStatsGrid::ShaderStatisticColumnName, MakeShareable(new FGridCell_Empty()));
 
 	FillPlatformCellsHelper(StatsManager);
 }
@@ -351,6 +363,7 @@ void FStatsGridRow_NumShaders::CreateRow(TSharedPtr<FMaterialStats> StatsManager
 	AddCell(FMaterialStatsGrid::DescriptorColumnName, HeaderCell);
 
 	AddCell(FMaterialStatsGrid::ShaderColumnName, MakeShareable(new FGridCell_Empty()));
+	AddCell(FMaterialStatsGrid::ShaderStatisticColumnName, MakeShareable(new FGridCell_Empty()));
 
 	FillPlatformCellsHelper(StatsManager);
 }
@@ -373,6 +386,7 @@ void FStatsGridRow_NumShaders::AddPlatform(TSharedPtr<FMaterialStats> StatsManag
 
 const FName FMaterialStatsGrid::DescriptorColumnName = TEXT("Descriptor");
 const FName FMaterialStatsGrid::ShaderColumnName = TEXT("ShaderList");
+const FName FMaterialStatsGrid::ShaderStatisticColumnName = TEXT("ShaderStatistics");
 
 FMaterialStatsGrid::FMaterialStatsGrid(TWeakPtr<FMaterialStats> _StatsManager)
 {
@@ -652,10 +666,15 @@ void FMaterialStatsGrid::BuildShaderRows()
 		{
 			bool bFirstShader = !FragmentShaderRows.Num();
 
-			TSharedPtr<FStatsGridRow> FragShaderRow = MakeShareable(new FStatsGridRow_Shaders((ERepresentativeShader)i, bFirstShader));
+			TSharedPtr<FStatsGridRow> FragShaderRow = MakeShareable(new FStatsGridRow_Shaders((ERepresentativeShader)i, bFirstShader, true));
 			FragShaderRow->CreateRow(StatsManager);
 
 			FragmentShaderRows.Add(FragShaderRow);
+
+			TSharedPtr<FStatsGridRow> FragShaderRow2 = MakeShareable(new FStatsGridRow_Shaders((ERepresentativeShader)i, false, false));
+			FragShaderRow2->CreateRow(StatsManager);
+
+			FragmentShaderRows.Add(FragShaderRow2);
 		}
 	}
 
@@ -666,7 +685,7 @@ void FMaterialStatsGrid::BuildShaderRows()
 		{
 			bool bFirstShader = !VertexShaderRows.Num();
 
-			TSharedPtr<FStatsGridRow> VertShaderRow = MakeShareable(new FStatsGridRow_Shaders((ERepresentativeShader)i, bFirstShader));
+			TSharedPtr<FStatsGridRow> VertShaderRow = MakeShareable(new FStatsGridRow_Shaders((ERepresentativeShader)i, bFirstShader, true));
 			VertShaderRow->CreateRow(StatsManager);
 
 			VertexShaderRows.Add(VertShaderRow);
@@ -679,6 +698,7 @@ void FMaterialStatsGrid::BuildColumnInfo()
 	GridColumnContent.Empty();
 	GridColumnContent.Add(DescriptorColumnName, FColumnInfo());
 	GridColumnContent.Add(ShaderColumnName, FColumnInfo());
+	GridColumnContent.Add(ShaderStatisticColumnName, FColumnInfo());
 
 	TSharedPtr<FMaterialStats> StatsManager = StatsManagerWPtr.Pin();
 	if (!StatsManager.IsValid())
@@ -940,6 +960,16 @@ FString FGridCell_ShaderValue::InternalGetContent(bool bLongContent)
 		}
 		break;
 
+		case EShaderInfoType::GenericShaderStatistics:
+		{
+			auto* Count = InstanceData.ShaderStatsInfo.GenericShaderStatistics.Find(ShaderType);
+			if (Count)
+			{
+				return bLongContent ? Count->StrDescriptionLong : Count->StrDescription;
+			}
+		}
+		break;
+
 		case EShaderInfoType::InterpolatorsCount:
 			return bLongContent ? InstanceData.ShaderStatsInfo.InterpolatorsCount.StrDescriptionLong : InstanceData.ShaderStatsInfo.InterpolatorsCount.StrDescription;
 		break;
diff --git a/Engine/Source/Editor/MaterialEditor/Private/MaterialStatsGrid.h b/Engine/Source/Editor/MaterialEditor/Private/MaterialStatsGrid.h
--- a/Engine/Source/Editor/MaterialEditor/Private/MaterialStatsGrid.h
+++ b/Engine/Source/Editor/MaterialEditor/Private/MaterialStatsGrid.h
@@ -87,6 +87,7 @@ enum class EShaderInfoType
 	TextureSampleCount,
 	VirtualTextureLookupCount,
 	ShaderCount,
+	GenericShaderStatistics,
 };
 
 /** this type of cell will query certain type of informations from the material */
@@ -186,13 +187,15 @@ private:
 	// if this is true it will add a text in the 'description' column with 'fragment/vertex shader' text
 	bool bIsHeaderRow = false;
 
+	bool bInstructionRow = true;
+
 	//EShaderType ShaderType;
 	ERepresentativeShader ShaderType;
 private:
 	EShaderClass GetShaderClass(const ERepresentativeShader Shader);
 
 public:
-	FStatsGridRow_Shaders(ERepresentativeShader RepresentativeShader, bool bHeader);
+	FStatsGridRow_Shaders(ERepresentativeShader RepresentativeShader, bool bHeader, bool bInstructionRow);
 
 	void CreateRow(TSharedPtr<FMaterialStats> StatsManager) override;
 
@@ -310,6 +313,7 @@ class FMaterialStatsGrid
 public:
 	const static FName DescriptorColumnName;
 	const static FName ShaderColumnName;
+	const static FName ShaderStatisticColumnName;
 
 private:
 	void AddColumnInfo(TSharedPtr<FShaderPlatformSettings> PlatformPtr, const EMaterialQualityLevel::Type QualityLevel, const int32 InstanceIndex);
diff --git a/Engine/Source/Editor/MaterialEditor/Private/SMaterialEditorStatsWidget.cpp b/Engine/Source/Editor/MaterialEditor/Private/SMaterialEditorStatsWidget.cpp
--- a/Engine/Source/Editor/MaterialEditor/Private/SMaterialEditorStatsWidget.cpp
+++ b/Engine/Source/Editor/MaterialEditor/Private/SMaterialEditorStatsWidget.cpp
@@ -212,6 +212,10 @@ float SMaterialEditorStatsWidget::GetColumnSize(const FName ColumnName) const
 	{
 		ColumnSize = ColumnSizeLarge;
 	}
+	else if (ColumnName == FMaterialStatsGrid::ShaderStatisticColumnName)
+	{
+		ColumnSize = ColumnSizeMedium;
+	}
 	else
 	{
 		auto StatsPtr = MaterialStatsWPtr.Pin();
diff --git a/Engine/Source/Editor/MaterialEditor/Public/MaterialStatsCommon.h b/Engine/Source/Editor/MaterialEditor/Public/MaterialStatsCommon.h
--- a/Engine/Source/Editor/MaterialEditor/Public/MaterialStatsCommon.h
+++ b/Engine/Source/Editor/MaterialEditor/Public/MaterialStatsCommon.h
@@ -76,6 +76,7 @@ public:
 	{
 		ERepresentativeShader ShaderType;
 		FString ShaderDescription;
+		FString ShaderStatisticsString;
 		int32 InstructionCount;
 	};
 
diff --git a/Engine/Source/Runtime/Engine/Public/MaterialShared.h b/Engine/Source/Runtime/Engine/Public/MaterialShared.h
--- a/Engine/Source/Runtime/Engine/Public/MaterialShared.h
+++ b/Engine/Source/Runtime/Engine/Public/MaterialShared.h
@@ -1267,6 +1267,8 @@ public:
 	uint32 GetMaxNumInstructionsForShader(FShaderType* ShaderType) const { return GetContent()->GetMaxNumInstructionsForShader(*this, ShaderType); }
 
 #if WITH_EDITOR
+	FShader::FShaderStatisticMap GetShaderStatisticsMapForShader(FShaderType* ShaderType) const { return GetContent()->GetShaderStatisticsMapForShader(*this, ShaderType); }
+
 	/** Submits compile jobs for this shadermap, returns number of jobs submitted. */
 	int32 SubmitCompileJobs(uint32 CompilingShaderMapId,
 		const FMaterial* Material,
diff --git a/Engine/Source/Runtime/RenderCore/Private/Shader.cpp b/Engine/Source/Runtime/RenderCore/Private/Shader.cpp
--- a/Engine/Source/Runtime/RenderCore/Private/Shader.cpp
+++ b/Engine/Source/Runtime/RenderCore/Private/Shader.cpp
@@ -632,6 +632,7 @@ FShaderCompiledShaderInitializerType::FShaderCompiledShaderInitializerType(
 	, NumTextureSamplers(CompilerOutput.NumTextureSamplers)
 	, CodeSize(CompilerOutput.ShaderCode.GetShaderCodeSize())
 	, PermutationId(InPermutationId)
+	, ShaderStatistics(CompilerOutput.ShaderStatistics)
 {
 }
 
@@ -683,6 +684,11 @@ FShader::FShader(const CompiledShaderInitializerType& Initializer)
 		// Store off the VF source hash that this shader was compiled with
 		VFSourceHash = Initializer.VertexFactoryType->GetSourceHash(Initializer.Target.GetPlatform());
 	}
+
+	for (const TPair<FString, FShaderStatVariant>& Pair : Initializer.ShaderStatistics)
+	{
+		ShaderStatistics.Add(FMemoryImageName(FName(Pair.Key)), Pair.Value);
+	}
 #endif // WITH_EDITORONLY_DATA
 
 	BuildParameterMapInfo(Initializer.ParameterMap.GetParameterMap());
diff --git a/Engine/Source/Runtime/RenderCore/Private/ShaderMap.cpp b/Engine/Source/Runtime/RenderCore/Private/ShaderMap.cpp
--- a/Engine/Source/Runtime/RenderCore/Private/ShaderMap.cpp
+++ b/Engine/Source/Runtime/RenderCore/Private/ShaderMap.cpp
@@ -641,6 +641,21 @@ uint32 FShaderMapContent::GetMaxNumInstructionsForShader(const FShaderMapBase& I
 	return MaxNumInstructions;
 }
 
+#if WITH_EDITOR
+const FShader::FShaderStatisticMap FShaderMapContent::GetShaderStatisticsMapForShader(const FShaderMapBase& InShaderMap, FShaderType* ShaderType) const
+{
+	FShader::FShaderStatisticMap Statistics;
+
+	FShader* Shader = GetShader(ShaderType);
+	if (Shader)
+	{
+		Statistics = Shader->GetShaderStatistics();
+	}
+
+	return Statistics;
+}
+#endif // WITH_EDITOR
+
 struct FSortedShaderEntry
 {
 	FHashedName TypeName;
diff --git a/Engine/Source/Runtime/RenderCore/Public/Shader.h b/Engine/Source/Runtime/RenderCore/Public/Shader.h
--- a/Engine/Source/Runtime/RenderCore/Public/Shader.h
+++ b/Engine/Source/Runtime/RenderCore/Public/Shader.h
@@ -844,6 +844,8 @@ public:
 	using FPermutationParameters = FShaderPermutationParameters;
 	using CompiledShaderInitializerType = FShaderCompiledShaderInitializerType;
 	using ShaderMetaType = FShaderType;
+	using ShaderStatKeyType = FMemoryImageName;
+	using FShaderStatisticMap = TMemoryImageMap<ShaderStatKeyType, FShaderStatVariant>;
 
 	/** 
 	 * Used to construct a shader for deserialization.
@@ -899,6 +901,7 @@ public:
 	inline uint32 GetNumTextureSamplers() const { return NumTextureSamplers; }
 	inline uint32 GetCodeSize() const { return CodeSize; }
 	inline void SetNumInstructions(uint32 Value) { NumInstructions = Value; }
+	inline const FShaderStatisticMap& GetShaderStatistics() const { return ShaderStatistics; }
 #else
 	inline uint32 GetNumTextureSamplers() const { return 0u; }
 	inline uint32 GetCodeSize() const { return 0u; }
@@ -1008,6 +1011,9 @@ private:
 
 	/** Size of shader's compiled code */
 	LAYOUT_FIELD_EDITORONLY(uint32, CodeSize);
+
+	/** Generic, data-driven key/value pairs of statistics. */
+	LAYOUT_FIELD_EDITORONLY(FShaderStatisticMap, ShaderStatistics);
 };
 
 RENDERCORE_API const FTypeLayoutDesc& GetTypeLayoutDesc(const FPointerTableBase* PtrTable, const FShader& Shader);
@@ -1528,6 +1534,7 @@ struct FShaderCompiledShaderInitializerType
 	uint32 NumTextureSamplers;
 	uint32 CodeSize;
 	int32 PermutationId;
+	TMap<FString, FShaderStatVariant> ShaderStatistics;
 
 	RENDERCORE_API FShaderCompiledShaderInitializerType(
 		const FShaderType* InType,
@@ -2196,6 +2203,8 @@ public:
 	RENDERCORE_API void GetOutdatedTypes(const FShaderMapBase& InShaderMap, TArray<const FShaderType*>& OutdatedShaderTypes, TArray<const FShaderPipelineType*>& OutdatedShaderPipelineTypes, TArray<const FVertexFactoryType*>& OutdatedFactoryTypes) const;
 
 	RENDERCORE_API void SaveShaderStableKeys(const FShaderMapBase& InShaderMap, EShaderPlatform TargetShaderPlatform, const struct FStableShaderKeyAndValue& SaveKeyVal);
+
+	RENDERCORE_API const FShader::FShaderStatisticMap GetShaderStatisticsMapForShader(const FShaderMapBase& InShaderMap, FShaderType* ShaderType) const;
 #endif // WITH_EDITOR
 
 	/** @return true if the map is empty */
diff --git a/Engine/Source/Runtime/RenderCore/Public/ShaderCompilerCore.h b/Engine/Source/Runtime/RenderCore/Public/ShaderCompilerCore.h
--- a/Engine/Source/Runtime/RenderCore/Public/ShaderCompilerCore.h
+++ b/Engine/Source/Runtime/RenderCore/Public/ShaderCompilerCore.h
@@ -33,7 +33,7 @@ typedef TSharedPtr<TArray<ANSICHAR>, ESPMode::ThreadSafe> FShaderSharedAnsiStrin
 // this is for the protocol, not the data, bump if FShaderCompilerInput/FShaderPreprocessOutput serialization, SerializeWorkerInput or ProcessInputFromArchive changes.
 inline const int32 ShaderCompileWorkerInputVersion = 26;
 // this is for the protocol, not the data, bump if FShaderCompilerOutput or WriteToOutputArchive changes.
-inline const int32 ShaderCompileWorkerOutputVersion = 19;
+inline const int32 ShaderCompileWorkerOutputVersion = 20;
 // this is for the protocol, not the data.
 inline const int32 ShaderCompileWorkerSingleJobHeader = 'S';
 // this is for the protocol, not the data.
@@ -665,6 +665,8 @@ struct FShaderCompilerOutput
 
 	TArray<uint8> PlatformDebugData;
 
+	TMap<FString, FShaderStatVariant> ShaderStatistics;
+
 	/** Generates OutputHash from the compiler output. */
 	RENDERCORE_API void GenerateOutputHash();
 
@@ -677,6 +679,7 @@ struct FShaderCompilerOutput
 	/** Add optional diagnostic data in ShaderCode to perform assert translation at runtime*/
 	RENDERCORE_API void SerializeShaderDiagnosticData();
 
+	// Bump ShaderCompileWorkerOutputVersion if FShaderCompilerOutput changes
 	friend FArchive& operator<<(FArchive& Ar, FShaderCompilerOutput& Output)
 	{
 		// Note: this serialize is used to pass between UE and the shader compile worker, recompile both when modifying
@@ -691,6 +694,7 @@ struct FShaderCompilerOutput
 			Ar << Output.ModifiedEntryPointName;
 		}
 		Ar << Output.PlatformDebugData;
+		Ar << Output.ShaderStatistics;
 
 		return Ar;
 	}
diff --git a/Engine/Source/Runtime/RenderCore/Public/ShaderCore.h b/Engine/Source/Runtime/RenderCore/Public/ShaderCore.h
--- a/Engine/Source/Runtime/RenderCore/Public/ShaderCore.h
+++ b/Engine/Source/Runtime/RenderCore/Public/ShaderCore.h
@@ -54,6 +54,9 @@ class FShaderPreprocessorUtilities;
 class FSHA1;
 class ITargetPlatform;
 
+using FShaderStatVariant = TVariant<bool, float, int32, uint32>;
+DECLARE_INTRINSIC_TYPE_LAYOUT(FShaderStatVariant);
+
 /**
  * Controls whether shader related logs are visible.
  * Note: The runtime verbosity is driven by the console variable 'r.ShaderDevelopmentMode'
