Adding preshader statistic to the newer shader platform statistics window in the Material Editor.

- This is in preparation for removing the older material statistics window in the Material Editor.

todo: So much copy pasta.  It should be data driven and easy to add a new stat.

#tests [Editor][Win64] Opening materials in the material with platform stats open

diff --git a/Engine/Source/Editor/MaterialEditor/Private/MaterialStats.h b/Engine/Source/Editor/MaterialEditor/Private/MaterialStats.h
--- a/Engine/Source/Editor/MaterialEditor/Private/MaterialStats.h
+++ b/Engine/Source/Editor/MaterialEditor/Private/MaterialStats.h
@@ -23,6 +23,7 @@ struct FShaderStatsInfo
 	FContent TextureSampleCount;
 	FContent VirtualTextureLookupCount;
 	FContent ShaderCount;
+	FContent PreShaderCount;
 	FString StrShaderErrors;
 
 	void Reset()
@@ -45,6 +46,9 @@ struct FShaderStatsInfo
 		ShaderCount.StrDescription = TEXT("Compiling...");
 		ShaderCount.StrDescriptionLong = TEXT("Compiling...");
 
+		PreShaderCount.StrDescription = TEXT("Compiling...");
+		PreShaderCount.StrDescriptionLong = TEXT("Compiling...");
+
 		StrShaderErrors.Empty();
 	}
 
@@ -68,6 +72,9 @@ struct FShaderStatsInfo
 		ShaderCount.StrDescription.Empty();
 		ShaderCount.StrDescriptionLong.Empty();
 
+		PreShaderCount.StrDescription.Empty();
+		PreShaderCount.StrDescriptionLong.Empty();
+
 		StrShaderErrors.Empty();
 	}
 
diff --git a/Engine/Source/Editor/MaterialEditor/Private/MaterialStatsCommon.cpp b/Engine/Source/Editor/MaterialEditor/Private/MaterialStatsCommon.cpp
--- a/Engine/Source/Editor/MaterialEditor/Private/MaterialStatsCommon.cpp
+++ b/Engine/Source/Editor/MaterialEditor/Private/MaterialStatsCommon.cpp
@@ -728,6 +728,15 @@ void FMaterialStatsUtils::ExtractMatertialStatsInfo(EShaderPlatform ShaderPlatfo
 
 		OutInfo.ShaderCount.StrDescription = FString::Printf(TEXT("%u"), TotalShadersForMaterial);
 		OutInfo.ShaderCount.StrDescriptionLong = FString::Printf(TEXT("Total Shaders: %u"), TotalShadersForMaterial);
+
+		if (FMaterialShaderMap* ShaderMap = MaterialResource->GetGameThreadShaderMap())
+		{
+			// Add number of preshaders and stats
+			uint32 TotalParams, TotalOps;
+			MaterialResource->GetPreshaderStats(TotalParams, TotalOps);
+			OutInfo.PreShaderCount.StrDescription = FString::Printf(TEXT("%u outputs\n%u params\n%u ops"), ShaderMap->GetNumPreshaders(), TotalParams, TotalOps);
+			OutInfo.PreShaderCount.StrDescriptionLong = FString::Printf(TEXT("%u outputs, %u parameter fetches, %u total operations"), ShaderMap->GetNumPreshaders(), TotalParams, TotalOps);
+		}
 	}
 }
 
diff --git a/Engine/Source/Editor/MaterialEditor/Private/MaterialStatsGrid.cpp b/Engine/Source/Editor/MaterialEditor/Private/MaterialStatsGrid.cpp
--- a/Engine/Source/Editor/MaterialEditor/Private/MaterialStatsGrid.cpp
+++ b/Engine/Source/Editor/MaterialEditor/Private/MaterialStatsGrid.cpp
@@ -352,7 +352,7 @@ void FStatsGridRow_NumVirtualTextureLookups::AddPlatform(TSharedPtr<FMaterialSta
 /*==============================================================================================================*/
 
 /*==============================================================================================================*/
-/* FStatsGridRow_NumVirtualTextureLookups functions*/
+/* FStatsGridRow_NumShaders functions*/
 
 void FStatsGridRow_NumShaders::CreateRow(TSharedPtr<FMaterialStats> StatsManager)
 {
@@ -381,6 +381,36 @@ void FStatsGridRow_NumShaders::AddPlatform(TSharedPtr<FMaterialStats> StatsManag
 /*end FStatsGridRow_NumTotalShaders functions*/
 /*==============================================================================================================*/
 
+/*==============================================================================================================*/
+/* FStatsGridRow_NumPreshaders functions*/
+
+void FStatsGridRow_NumPreshaders::CreateRow(TSharedPtr<FMaterialStats> StatsManager)
+{
+	// static string in the descriptor column
+	TSharedPtr<FGridCell> HeaderCell = MakeShareable(new FGridCell_StaticString(TEXT("Preshader Count"), TEXT("Number of preshader instructions that will be evaluated on the CPU for this material.")));
+	HeaderCell->SetColor(FStyleColors::Foreground);
+	HeaderCell->SetContentBold(true);
+	AddCell(FMaterialStatsGrid::DescriptorColumnName, HeaderCell);
+
+	AddCell(FMaterialStatsGrid::ShaderColumnName, MakeShareable(new FGridCell_Empty()));
+	AddCell(FMaterialStatsGrid::ShaderStatisticColumnName, MakeShareable(new FGridCell_Empty()));
+
+	FillPlatformCellsHelper(StatsManager);
+}
+
+void FStatsGridRow_NumPreshaders::AddPlatform(TSharedPtr<FMaterialStats> StatsManager, const TSharedPtr<FShaderPlatformSettings> Platform, const EMaterialQualityLevel::Type QualityLevel, const int32 InstanceIndex)
+{
+	// cell that will enumerate the number of pre shaders
+	const FString CellContent = FMaterialStatsUtils::MaterialQualityToShortString(QualityLevel);
+	TSharedPtr<FGridCell_ShaderValue> Cell = MakeShareable(new FGridCell_ShaderValue(StatsManager, EShaderInfoType::PreShaderCount, ERepresentativeShader::Num, QualityLevel, Platform->GetPlatformShaderType(), InstanceIndex));
+
+	const FName ColumnName = FMaterialStatsGrid::MakePlatformColumnName(Platform, QualityLevel, InstanceIndex);
+	AddCell(ColumnName, Cell);
+}
+
+/*end FStatsGridRow_NumPreshaders functions*/
+/*==============================================================================================================*/
+
 /***********************************************************************************************************************/
 /*FShaderStatsGrid functions*/
 
@@ -576,6 +606,7 @@ void FMaterialStatsGrid::BuildRowIds()
 		BuildKeyAndInsert(ERowType::VirtualTextureLookups);
 		BuildKeyAndInsert(ERowType::Interpolators);
 		BuildKeyAndInsert(ERowType::Shaders);
+		BuildKeyAndInsert(ERowType::PreShaders);
 	}
 }
 
@@ -646,6 +677,11 @@ void FMaterialStatsGrid::BuildStaticRows()
 		Row->CreateRow(StatsManager);
 		StaticRows.Add(ERowType::Shaders, Row);
 	}
+	{
+		TSharedPtr<FStatsGridRow> Row = MakeShareable(new FStatsGridRow_NumPreshaders());
+		Row->CreateRow(StatsManager);
+		StaticRows.Add(ERowType::PreShaders, Row);
+	}
 }
 
 void FMaterialStatsGrid::BuildShaderRows()
@@ -989,6 +1025,10 @@ FString FGridCell_ShaderValue::InternalGetContent(bool bLongContent)
 		case EShaderInfoType::ShaderCount:
 			return bLongContent ? InstanceData.ShaderStatsInfo.ShaderCount.StrDescriptionLong : InstanceData.ShaderStatsInfo.ShaderCount.StrDescription;
 		break;
+
+		case EShaderInfoType::PreShaderCount:
+			return bLongContent ? InstanceData.ShaderStatsInfo.PreShaderCount.StrDescriptionLong : InstanceData.ShaderStatsInfo.PreShaderCount.StrDescription;
+		break;
 	}
 
 	return TEXT("");
diff --git a/Engine/Source/Editor/MaterialEditor/Private/MaterialStatsGrid.h b/Engine/Source/Editor/MaterialEditor/Private/MaterialStatsGrid.h
--- a/Engine/Source/Editor/MaterialEditor/Private/MaterialStatsGrid.h
+++ b/Engine/Source/Editor/MaterialEditor/Private/MaterialStatsGrid.h
@@ -87,6 +87,7 @@ enum class EShaderInfoType
 	TextureSampleCount,
 	VirtualTextureLookupCount,
 	ShaderCount,
+	PreShaderCount,
 	GenericShaderStatistics,
 };
 
@@ -247,6 +248,15 @@ public:
 	void AddPlatform(TSharedPtr<FMaterialStats> StatsManager, const TSharedPtr<FShaderPlatformSettings> Platform, const EMaterialQualityLevel::Type QualityLevel, const int32 InstanceIndex) override;
 };
 
+/** this row will display the total number of shaders present in the material for a specified platform */
+class FStatsGridRow_NumPreshaders : public FStatsGridRow
+{
+public:
+	void CreateRow(TSharedPtr<FMaterialStats> StatsManager) override;
+
+	void AddPlatform(TSharedPtr<FMaterialStats> StatsManager, const TSharedPtr<FShaderPlatformSettings> Platform, const EMaterialQualityLevel::Type QualityLevel, const int32 InstanceIndex) override;
+};
+
 /** class that models the logical material stats grid */
 class FMaterialStatsGrid
 {
@@ -272,6 +282,7 @@ class FMaterialStatsGrid
 		TextureSamples,
 		VirtualTextureLookups,
 		Shaders,
+		PreShaders,
 
 		VertexShader,
 		FragmentShader,
