Adding LWC statistics to the newer shader platform statistics window in the Material Editor.

- This is in preparation for removing the older material statistics window in the Material Editor.

todo: So much copy pasta.  It should be data driven and easy to add a new stat.

#tests [Editor][Win64] Opening materials in the material with platform stats open

diff --git a/Engine/Source/Editor/MaterialEditor/Private/MaterialStats.h b/Engine/Source/Editor/MaterialEditor/Private/MaterialStats.h
--- a/Engine/Source/Editor/MaterialEditor/Private/MaterialStats.h
+++ b/Engine/Source/Editor/MaterialEditor/Private/MaterialStats.h
@@ -24,6 +24,7 @@ struct FShaderStatsInfo
 	FContent VirtualTextureLookupCount;
 	FContent ShaderCount;
 	FContent PreShaderCount;
+	FContent LWCUsage;
 	FString StrShaderErrors;
 
 	void Reset()
@@ -49,6 +50,9 @@ struct FShaderStatsInfo
 		PreShaderCount.StrDescription = TEXT("Compiling...");
 		PreShaderCount.StrDescriptionLong = TEXT("Compiling...");
 
+		LWCUsage.StrDescription = TEXT("Compiling...");
+		LWCUsage.StrDescriptionLong = TEXT("Compiling...");
+
 		StrShaderErrors.Empty();
 	}
 
diff --git a/Engine/Source/Editor/MaterialEditor/Private/MaterialStatsCommon.cpp b/Engine/Source/Editor/MaterialEditor/Private/MaterialStatsCommon.cpp
--- a/Engine/Source/Editor/MaterialEditor/Private/MaterialStatsCommon.cpp
+++ b/Engine/Source/Editor/MaterialEditor/Private/MaterialStatsCommon.cpp
@@ -729,6 +729,20 @@ void FMaterialStatsUtils::ExtractMatertialStatsInfo(EShaderPlatform ShaderPlatfo
 		OutInfo.ShaderCount.StrDescription = FString::Printf(TEXT("%u"), TotalShadersForMaterial);
 		OutInfo.ShaderCount.StrDescriptionLong = FString::Printf(TEXT("Total Shaders: %u"), TotalShadersForMaterial);
 
+		FString LWCMessage;
+		TStaticArray<uint16, (int)ELWCFunctionKind::Max> LWCFuncUsages = MaterialResource->GetEstimatedLWCFuncUsages();
+		for (int KindIndex = 0; KindIndex < (int)ELWCFunctionKind::Max; ++KindIndex)
+		{
+			int Usages = LWCFuncUsages[KindIndex];
+			if (LWCFuncUsages[KindIndex] > 0)
+			{
+				LWCMessage += FString::Printf(TEXT("%s: %u\n"), *UEnum::GetDisplayValueAsText((ELWCFunctionKind)KindIndex).ToString(), Usages);
+			}
+		}
+
+		OutInfo.LWCUsage.StrDescription = LWCMessage;
+		OutInfo.LWCUsage.StrDescriptionLong = LWCMessage;
+
 		if (FMaterialShaderMap* ShaderMap = MaterialResource->GetGameThreadShaderMap())
 		{
 			// Add number of preshaders and stats
diff --git a/Engine/Source/Editor/MaterialEditor/Private/MaterialStatsGrid.cpp b/Engine/Source/Editor/MaterialEditor/Private/MaterialStatsGrid.cpp
--- a/Engine/Source/Editor/MaterialEditor/Private/MaterialStatsGrid.cpp
+++ b/Engine/Source/Editor/MaterialEditor/Private/MaterialStatsGrid.cpp
@@ -411,6 +411,36 @@ void FStatsGridRow_NumPreshaders::AddPlatform(TSharedPtr<FMaterialStats> StatsMa
 /*end FStatsGridRow_NumPreshaders functions*/
 /*==============================================================================================================*/
 
+/*==============================================================================================================*/
+/* FStatsGridRow_LWCUsage functions*/
+
+void FStatsGridRow_LWCUsage::CreateRow(TSharedPtr<FMaterialStats> StatsManager)
+{
+	// static string in the descriptor column
+	TSharedPtr<FGridCell> HeaderCell = MakeShareable(new FGridCell_StaticString(TEXT("Large World Coordinate Usages (Est.)"), TEXT("Estimates for the number of large world coordinate operations in the material.")));
+	HeaderCell->SetColor(FStyleColors::Foreground);
+	HeaderCell->SetContentBold(true);
+	AddCell(FMaterialStatsGrid::DescriptorColumnName, HeaderCell);
+
+	AddCell(FMaterialStatsGrid::ShaderColumnName, MakeShareable(new FGridCell_Empty()));
+	AddCell(FMaterialStatsGrid::ShaderStatisticColumnName, MakeShareable(new FGridCell_Empty()));
+
+	FillPlatformCellsHelper(StatsManager);
+}
+
+void FStatsGridRow_LWCUsage::AddPlatform(TSharedPtr<FMaterialStats> StatsManager, const TSharedPtr<FShaderPlatformSettings> Platform, const EMaterialQualityLevel::Type QualityLevel, const int32 InstanceIndex)
+{
+	// cell that will enumerate the number of pre shaders
+	const FString CellContent = FMaterialStatsUtils::MaterialQualityToShortString(QualityLevel);
+	TSharedPtr<FGridCell_ShaderValue> Cell = MakeShareable(new FGridCell_ShaderValue(StatsManager, EShaderInfoType::LWCUsage, ERepresentativeShader::Num, QualityLevel, Platform->GetPlatformShaderType(), InstanceIndex));
+
+	const FName ColumnName = FMaterialStatsGrid::MakePlatformColumnName(Platform, QualityLevel, InstanceIndex);
+	AddCell(ColumnName, Cell);
+}
+
+/*end FStatsGridRow_LWCUsage functions*/
+/*==============================================================================================================*/
+
 /***********************************************************************************************************************/
 /*FShaderStatsGrid functions*/
 
@@ -607,6 +637,7 @@ void FMaterialStatsGrid::BuildRowIds()
 		BuildKeyAndInsert(ERowType::Interpolators);
 		BuildKeyAndInsert(ERowType::Shaders);
 		BuildKeyAndInsert(ERowType::PreShaders);
+		BuildKeyAndInsert(ERowType::LWCUsage);
 	}
 }
 
@@ -682,6 +713,11 @@ void FMaterialStatsGrid::BuildStaticRows()
 		Row->CreateRow(StatsManager);
 		StaticRows.Add(ERowType::PreShaders, Row);
 	}
+	{
+		TSharedPtr<FStatsGridRow> Row = MakeShareable(new FStatsGridRow_LWCUsage());
+		Row->CreateRow(StatsManager);
+		StaticRows.Add(ERowType::LWCUsage, Row);
+	}
 }
 
 void FMaterialStatsGrid::BuildShaderRows()
@@ -1029,6 +1065,10 @@ FString FGridCell_ShaderValue::InternalGetContent(bool bLongContent)
 		case EShaderInfoType::PreShaderCount:
 			return bLongContent ? InstanceData.ShaderStatsInfo.PreShaderCount.StrDescriptionLong : InstanceData.ShaderStatsInfo.PreShaderCount.StrDescription;
 		break;
+
+		case EShaderInfoType::LWCUsage:
+			return bLongContent ? InstanceData.ShaderStatsInfo.LWCUsage.StrDescriptionLong : InstanceData.ShaderStatsInfo.LWCUsage.StrDescription;
+		break;
 	}
 
 	return TEXT("");
diff --git a/Engine/Source/Editor/MaterialEditor/Private/MaterialStatsGrid.h b/Engine/Source/Editor/MaterialEditor/Private/MaterialStatsGrid.h
--- a/Engine/Source/Editor/MaterialEditor/Private/MaterialStatsGrid.h
+++ b/Engine/Source/Editor/MaterialEditor/Private/MaterialStatsGrid.h
@@ -88,6 +88,7 @@ enum class EShaderInfoType
 	VirtualTextureLookupCount,
 	ShaderCount,
 	PreShaderCount,
+	LWCUsage,
 	GenericShaderStatistics,
 };
 
@@ -257,6 +258,15 @@ public:
 	void AddPlatform(TSharedPtr<FMaterialStats> StatsManager, const TSharedPtr<FShaderPlatformSettings> Platform, const EMaterialQualityLevel::Type QualityLevel, const int32 InstanceIndex) override;
 };
 
+/** this row will display the LWC usage in the material for a specified platform */
+class FStatsGridRow_LWCUsage : public FStatsGridRow
+{
+public:
+	void CreateRow(TSharedPtr<FMaterialStats> StatsManager) override;
+
+	void AddPlatform(TSharedPtr<FMaterialStats> StatsManager, const TSharedPtr<FShaderPlatformSettings> Platform, const EMaterialQualityLevel::Type QualityLevel, const int32 InstanceIndex) override;
+};
+
 /** class that models the logical material stats grid */
 class FMaterialStatsGrid
 {
@@ -283,6 +293,7 @@ class FMaterialStatsGrid
 		VirtualTextureLookups,
 		Shaders,
 		PreShaders,
+		LWCUsage,
 
 		VertexShader,
 		FragmentShader,
