Implemented the same material graph node context menu actions for SparseVolumeTexture nodes as for regular Texture nodes. This now allows to convert between Sample, Object and Parameter nodes exactly the same as with regular textures. Also fixed the graph node title bar color being different to regular texture nodes.

#rnx

diff --git a/Engine/Source/Editor/MaterialEditor/Private/MaterialEditor.cpp b/Engine/Source/Editor/MaterialEditor/Private/MaterialEditor.cpp
--- a/Engine/Source/Editor/MaterialEditor/Private/MaterialEditor.cpp
+++ b/Engine/Source/Editor/MaterialEditor/Private/MaterialEditor.cpp
@@ -42,6 +42,7 @@
 #include "Engine/TextureCube.h"
 #include "Engine/Texture2DArray.h"
 #include "Engine/TextureCubeArray.h"
+#include "SparseVolumeTexture/SparseVolumeTexture.h"
 #include "Dialogs/Dialogs.h"
 #include "UnrealEdGlobals.h"
 #include "Editor.h"
@@ -85,6 +86,8 @@
 #include "Materials/MaterialExpressionTextureSampleParameter2DArray.h"
 #include "Materials/MaterialExpressionTextureSampleParameterCubeArray.h"
 #include "Materials/MaterialExpressionTextureSampleParameterSubUV.h"
+#include "Materials/MaterialExpressionSparseVolumeTextureSample.h"
+#include "Materials/MaterialExpressionSparseVolumeTextureObject.h"
 #include "Materials/MaterialExpressionTransformPosition.h"
 #include "Materials/MaterialExpressionVectorParameter.h"
 #include "Materials/MaterialExpressionDoubleVectorParameter.h"
@@ -3913,7 +3916,8 @@ void FMaterialEditor::OnUseCurrentTexture()
 	// as the texture to use in all selected texture sample expressions.
 	FEditorDelegates::LoadSelectedAssetsIfNeeded.Broadcast();
 	UTexture* SelectedTexture = GEditor->GetSelectedObjects()->GetTop<UTexture>();
-	if ( SelectedTexture )
+	USparseVolumeTexture* SelectedSparseVolumeTexture = GEditor->GetSelectedObjects()->GetTop<USparseVolumeTexture>();
+	if ( SelectedTexture || SelectedSparseVolumeTexture )
 	{
 		const FScopedTransaction Transaction( NSLOCTEXT("UnrealEd", "UseCurrentTexture", "Use Current Texture") );
 		const FGraphPanelSelectionSet SelectedNodes = GetSelectedNodes();
@@ -3921,13 +3925,19 @@ void FMaterialEditor::OnUseCurrentTexture()
 		for (FGraphPanelSelectionSet::TConstIterator NodeIt(SelectedNodes); NodeIt; ++NodeIt)
 		{
 			UMaterialGraphNode* GraphNode = Cast<UMaterialGraphNode>(*NodeIt);
-			if (GraphNode && GraphNode->MaterialExpression->IsA(UMaterialExpressionTextureBase::StaticClass()) )
+			if (SelectedTexture && GraphNode && GraphNode->MaterialExpression->IsA(UMaterialExpressionTextureBase::StaticClass()) )
 			{
 				UMaterialExpressionTextureBase* TextureBase = static_cast<UMaterialExpressionTextureBase*>(GraphNode->MaterialExpression);
 				TextureBase->Modify();
 				TextureBase->Texture = SelectedTexture;
 				TextureBase->AutoSetSampleType();
 			}
+			else if (SelectedSparseVolumeTexture && GraphNode && GraphNode->MaterialExpression->IsA(UMaterialExpressionSparseVolumeTextureBase::StaticClass()))
+			{
+				UMaterialExpressionSparseVolumeTextureBase* TextureBase = static_cast<UMaterialExpressionSparseVolumeTextureBase*>(GraphNode->MaterialExpression);
+				TextureBase->Modify();
+				TextureBase->SparseVolumeTexture = SelectedSparseVolumeTexture;
+			}
 		}
 
 		// Update the current preview material. 
@@ -4062,6 +4072,8 @@ void FMaterialEditor::OnConvertObjects()
 				UMaterialExpressionTextureObjectParameter* TextureObjectParameterExpression = Cast<UMaterialExpressionTextureObjectParameter>(CurrentSelectedExpression);
 				UMaterialExpressionRuntimeVirtualTextureSample* RuntimeVirtualTextureSampleExpression = Cast<UMaterialExpressionRuntimeVirtualTextureSample>(CurrentSelectedExpression);
 				UMaterialExpressionSparseVolumeTextureSample* SparseVolumeTextureSampleExpression = Cast<UMaterialExpressionSparseVolumeTextureSample>(CurrentSelectedExpression);
+				UMaterialExpressionSparseVolumeTextureObject* SparseVolumeTextureObjectExpression = Cast<UMaterialExpressionSparseVolumeTextureObject>(CurrentSelectedExpression);
+				UMaterialExpressionSparseVolumeTextureObjectParameter* SparseVolumeTextureObjectParameterExpression = Cast<UMaterialExpressionSparseVolumeTextureObjectParameter>(CurrentSelectedExpression);
 
 				// Setup the class to convert to
 				UClass* ClassToCreate = NULL;
@@ -4105,10 +4117,18 @@ void FMaterialEditor::OnConvertObjects()
 				{
 					ClassToCreate = UMaterialExpressionRuntimeVirtualTextureSampleParameter::StaticClass();
 				}
+				else if (SparseVolumeTextureObjectParameterExpression) // Has to come before SparseVolumeTextureSample comparison 
+				{
+					ClassToCreate = UMaterialExpressionSparseVolumeTextureObject::StaticClass();
+				}
 				else if (SparseVolumeTextureSampleExpression)
 				{
 					ClassToCreate = UMaterialExpressionSparseVolumeTextureSampleParameter::StaticClass();
 				}
+				else if (SparseVolumeTextureObjectExpression)
+				{
+					ClassToCreate = UMaterialExpressionSparseVolumeTextureObjectParameter::StaticClass();
+				}
 				else if (ComponentMaskExpression)
 				{
 					ClassToCreate = UMaterialExpressionStaticComponentMaskParameter::StaticClass();
@@ -4203,13 +4223,25 @@ void FMaterialEditor::OnConvertObjects()
 							NewRuntimeVirtualTextureExpression->MipValueMode = RuntimeVirtualTextureSampleExpression->MipValueMode;
 							NewGraphNode->ReconstructNode();
 						}
-						else if (SparseVolumeTextureSampleExpression)
+						else if (SparseVolumeTextureSampleExpression && !SparseVolumeTextureObjectParameterExpression) // Sample -> SampleParameter
 						{
 							bNeedsRefresh = true;
 							UMaterialExpressionSparseVolumeTextureSampleParameter* NewSparseVolumeTextureExpression = CastChecked<UMaterialExpressionSparseVolumeTextureSampleParameter>(NewExpression);
 							NewSparseVolumeTextureExpression->SparseVolumeTexture = SparseVolumeTextureSampleExpression->SparseVolumeTexture;
 							NewGraphNode->ReconstructNode();
 						}
+						else if (SparseVolumeTextureObjectExpression && !SparseVolumeTextureObjectParameterExpression) // Object -> ObjectParameter
+						{
+							bNeedsRefresh = true;
+							UMaterialExpressionSparseVolumeTextureObjectParameter* NewSparseVolumeTextureObjectParameterExpression = CastChecked<UMaterialExpressionSparseVolumeTextureObjectParameter>(NewExpression);
+							NewSparseVolumeTextureObjectParameterExpression->SparseVolumeTexture = SparseVolumeTextureObjectExpression->SparseVolumeTexture;
+						}
+						else if (SparseVolumeTextureObjectParameterExpression) // ObjectParameter -> Object
+						{
+							bNeedsRefresh = true;
+							UMaterialExpressionSparseVolumeTextureObject* NewTextureObjectExpression = CastChecked<UMaterialExpressionSparseVolumeTextureObject>(NewExpression);
+							NewTextureObjectExpression->SparseVolumeTexture = SparseVolumeTextureObjectParameterExpression->SparseVolumeTexture;
+						}
 						else if (ComponentMaskExpression)
 						{
 							bNeedsRefresh = true;
@@ -4285,6 +4317,8 @@ void FMaterialEditor::OnConvertTextures()
 				UMaterialExpression* CurrentSelectedExpression = GraphNode->MaterialExpression;
 				UMaterialExpressionTextureSample* TextureSampleExpression = Cast<UMaterialExpressionTextureSample>(CurrentSelectedExpression);
 				UMaterialExpressionTextureObject* TextureObjectExpression = Cast<UMaterialExpressionTextureObject>(CurrentSelectedExpression);
+				UMaterialExpressionSparseVolumeTextureSample* SparseVolumeTextureSampleExpression = Cast<UMaterialExpressionSparseVolumeTextureSample>(CurrentSelectedExpression);
+				UMaterialExpressionSparseVolumeTextureObject* SparseVolumeTextureObjectExpression = Cast<UMaterialExpressionSparseVolumeTextureObject>(CurrentSelectedExpression);
 
 				// Setup the class to convert to
 				UClass* ClassToCreate = NULL;
@@ -4296,6 +4330,14 @@ void FMaterialEditor::OnConvertTextures()
 				{
 					ClassToCreate = UMaterialExpressionTextureSample::StaticClass();
 				}
+				else if (SparseVolumeTextureSampleExpression)
+				{
+					ClassToCreate = UMaterialExpressionSparseVolumeTextureObject::StaticClass();
+				}
+				else if (SparseVolumeTextureObjectExpression)
+				{
+					ClassToCreate = UMaterialExpressionSparseVolumeTextureSample::StaticClass();
+				}
 
 				if (ClassToCreate)
 				{
@@ -4324,6 +4366,18 @@ void FMaterialEditor::OnConvertTextures()
 							NewTextureExpr->IsDefaultMeshpaintTexture = TextureObjectExpression->IsDefaultMeshpaintTexture;
 							NewTextureExpr->MipValueMode = TMVM_None;
 						}
+						else if (SparseVolumeTextureSampleExpression)
+						{
+							bNeedsRefresh = true;
+							UMaterialExpressionSparseVolumeTextureObject* NewTextureExpr = CastChecked<UMaterialExpressionSparseVolumeTextureObject>(NewExpression);
+							NewTextureExpr->SparseVolumeTexture = SparseVolumeTextureSampleExpression->SparseVolumeTexture;
+						}
+						else if (SparseVolumeTextureObjectExpression)
+						{
+							bNeedsRefresh = true;
+							UMaterialExpressionSparseVolumeTextureSample* NewTextureExpr = CastChecked<UMaterialExpressionSparseVolumeTextureSample>(NewExpression);
+							NewTextureExpr->SparseVolumeTexture = SparseVolumeTextureObjectExpression->SparseVolumeTexture;
+						}
 
 						if (bNeedsRefresh)
 						{
diff --git a/Engine/Source/Editor/UnrealEd/Private/MaterialGraphNode.cpp b/Engine/Source/Editor/UnrealEd/Private/MaterialGraphNode.cpp
--- a/Engine/Source/Editor/UnrealEd/Private/MaterialGraphNode.cpp
+++ b/Engine/Source/Editor/UnrealEd/Private/MaterialGraphNode.cpp
@@ -26,6 +26,7 @@
 #include "Materials/MaterialExpressionMaterialAttributeLayers.h"
 #include "Materials/MaterialExpressionRuntimeVirtualTextureSample.h"
 #include "Materials/MaterialExpressionSparseVolumeTextureSample.h"
+#include "Materials/MaterialExpressionSparseVolumeTextureObject.h"
 #include "Materials/MaterialExpressionScalarParameter.h"
 #include "Materials/MaterialExpressionStaticBool.h"
 #include "Materials/MaterialExpressionStaticBoolParameter.h"
@@ -375,7 +376,7 @@ void UMaterialGraphNode::GetNodeContextMenuActions(UToolMenu* Menu, UGraphNodeCo
 		FToolMenuSection& NodeSection = Menu->AddSection("MaterialSchemaNodeActions", LOCTEXT("NodeActionsMenuHeader", "Node Actions"), MenuPosition);
 		if (MaterialExpression)
 		{
-			if (MaterialExpression->IsA(UMaterialExpressionTextureBase::StaticClass()))
+			if (MaterialExpression->IsA(UMaterialExpressionTextureBase::StaticClass()) || MaterialExpression->IsA(UMaterialExpressionSparseVolumeTextureBase::StaticClass()))
 			{
 				{
 					NodeSection.AddMenuEntry(FMaterialEditorCommands::Get().UseCurrentTexture);
@@ -383,11 +384,11 @@ void UMaterialGraphNode::GetNodeContextMenuActions(UToolMenu* Menu, UGraphNodeCo
 
 				// Add a 'Convert To Texture' option for convertible types
 				{
-					if ( MaterialExpression->IsA(UMaterialExpressionTextureSample::StaticClass()) && !MaterialExpression->HasAParameterName())
+					if ( (MaterialExpression->IsA(UMaterialExpressionTextureSample::StaticClass()) || MaterialExpression->IsA(UMaterialExpressionSparseVolumeTextureSample::StaticClass())) && !MaterialExpression->HasAParameterName())
 					{
 						NodeSection.AddMenuEntry(FMaterialEditorCommands::Get().ConvertToTextureObjects);
 					}
-					else if ( MaterialExpression->IsA(UMaterialExpressionTextureObject::StaticClass()))
+					else if ( MaterialExpression->IsA(UMaterialExpressionTextureObject::StaticClass()) || MaterialExpression->IsA(UMaterialExpressionSparseVolumeTextureObject::StaticClass()) )
 					{
 						NodeSection.AddMenuEntry(FMaterialEditorCommands::Get().ConvertToTextureSamples);
 					}
@@ -421,8 +422,9 @@ void UMaterialGraphNode::GetNodeContextMenuActions(UToolMenu* Menu, UGraphNodeCo
 				|| MaterialExpression->IsA(UMaterialExpressionConstant4Vector::StaticClass())
 				|| (MaterialExpression->IsA(UMaterialExpressionTextureSample::StaticClass()) && !MaterialExpression->HasAParameterName())
 				|| (MaterialExpression->IsA(UMaterialExpressionRuntimeVirtualTextureSample::StaticClass()) && !MaterialExpression->HasAParameterName())
-				|| (MaterialExpression->IsA(UMaterialExpressionSparseVolumeTextureSampleParameter::StaticClass()) && !MaterialExpression->HasAParameterName())
+				|| (MaterialExpression->IsA(UMaterialExpressionSparseVolumeTextureSample::StaticClass()) && !MaterialExpression->HasAParameterName())
 				|| MaterialExpression->IsA(UMaterialExpressionTextureObject::StaticClass())
+				|| MaterialExpression->IsA(UMaterialExpressionSparseVolumeTextureObject::StaticClass())
 				|| MaterialExpression->IsA(UMaterialExpressionComponentMask::StaticClass()))
 			{
 				{
@@ -437,7 +439,8 @@ void UMaterialGraphNode::GetNodeContextMenuActions(UToolMenu* Menu, UGraphNodeCo
 			// Add a 'Convert To Constant' option for convertible types
 			if (bIsAStandardScalarParam
 				|| MaterialExpression->IsA(UMaterialExpressionVectorParameter::StaticClass())
-				|| MaterialExpression->IsA(UMaterialExpressionTextureObjectParameter::StaticClass()))
+				|| MaterialExpression->IsA(UMaterialExpressionTextureObjectParameter::StaticClass())
+				|| MaterialExpression->IsA(UMaterialExpressionSparseVolumeTextureObjectParameter::StaticClass()))
 			{
 				{
 					NodeSection.AddMenuEntry(FMaterialEditorCommands::Get().ConvertToConstant);
@@ -1055,6 +1058,10 @@ bool UMaterialGraphNode::UsesObjectColour(UMaterialExpression* Expression)
 	{
 		return true;
 	}
+	else if (Expression->IsA<UMaterialExpressionSparseVolumeTextureBase>())
+	{
+		return true;
+	}
 	else if (Expression->IsA<UMaterialExpressionFontSample>())
 	{
 		return true;
