Option to force always compiling materials for mobile in material editor

diff --git a/Engine/Source/Editor/MaterialEditor/Private/MaterialEditor.cpp b/Engine/Source/Editor/MaterialEditor/Private/MaterialEditor.cpp
--- a/Engine/Source/Editor/MaterialEditor/Private/MaterialEditor.cpp
+++ b/Engine/Source/Editor/MaterialEditor/Private/MaterialEditor.cpp
@@ -191,6 +191,11 @@ static TAutoConsoleVariable<int32> CVarMaterialEdMaxDerivedMaterialInstances(
 	-1,
 	TEXT("Limits amount of derived material instance shown in platform stats. Use negative number to disable the limit. Material editor must be re-opened if changed at runtime."));
 
+TAutoConsoleVariable<bool> CVarMaterialEdAllowIgnoringCompilationErrors(
+	TEXT("r.MaterialEditor.AllowIgnoringCompilationErrors"),
+	true,
+	TEXT("Allow ignoring compilation errors of platform shaders and derived materials."));
+
 ///////////////////////////
 // FMatExpressionPreview //
 ///////////////////////////
@@ -565,7 +570,7 @@ void FMaterialEditor::InitMaterialEditor( const EToolkitMode::Type Mode, const T
 
 	GEditor->RegisterForUndo(this);
 
-	MaterialStatsManager = FMaterialStatsUtils::CreateMaterialStats(this, true);
+	MaterialStatsManager = FMaterialStatsUtils::CreateMaterialStats(this, true, CVarMaterialEdAllowIgnoringCompilationErrors.GetValueOnGameThread());
 	MaterialStatsManager->SetMaterialsDisplayNames({OriginalMaterial->GetName()});
 	MaterialStatsManager->GetOldStatsListing()->OnMessageTokenClicked().AddSP(this, &FMaterialEditor::OnMessageLogLinkActivated);
 
diff --git a/Engine/Source/Editor/MaterialEditor/Private/MaterialInstanceEditor.cpp b/Engine/Source/Editor/MaterialEditor/Private/MaterialInstanceEditor.cpp
--- a/Engine/Source/Editor/MaterialEditor/Private/MaterialInstanceEditor.cpp
+++ b/Engine/Source/Editor/MaterialEditor/Private/MaterialInstanceEditor.cpp
@@ -69,6 +69,8 @@ const FName FMaterialInstanceEditor::PropertiesTabId( TEXT( "MaterialInstanceEdi
 const FName FMaterialInstanceEditor::LayerPropertiesTabId(TEXT("MaterialInstanceEditor_MaterialLayerProperties"));
 const FName FMaterialInstanceEditor::PreviewSettingsTabId(TEXT("MaterialInstanceEditor_PreviewSettings"));
 
+extern TAutoConsoleVariable<bool> CVarMaterialEdAllowIgnoringCompilationErrors;
+
 //////////////////////////////////////////////////////////////////////////
 // SMaterialTreeWidgetItem
 class SMaterialTreeWidgetItem : public SMultiColumnTableRow< TWeakObjectPtr<UMaterialInterface> >
@@ -414,7 +416,7 @@ void FMaterialInstanceEditor::InitMaterialInstanceEditor( const EToolkitMode::Ty
 	MaterialEditorInstance->SetSourceInstance(InstanceConstant);
 	MaterialEditorInstance->SetSourceFunction(MaterialFunctionOriginal);
 
-	MaterialStatsManager = FMaterialStatsUtils::CreateMaterialStats(this, false);
+	MaterialStatsManager = FMaterialStatsUtils::CreateMaterialStats(this, false, CVarMaterialEdAllowIgnoringCompilationErrors.GetValueOnGameThread());
 	MaterialStatsManager->SetMaterialsDisplayNames({MaterialEditorInstance->SourceInstance->GetName()});
 
 	// Register our commands. This will only register them if not previously registered
diff --git a/Engine/Source/Editor/MaterialEditor/Private/MaterialStats.cpp b/Engine/Source/Editor/MaterialEditor/Private/MaterialStats.cpp
--- a/Engine/Source/Editor/MaterialEditor/Private/MaterialStats.cpp
+++ b/Engine/Source/Editor/MaterialEditor/Private/MaterialStats.cpp
@@ -18,6 +18,7 @@
 #include "DataDrivenShaderPlatformInfo.h"
 #include "Modules/ModuleManager.h"
 #include "MessageLogModule.h"
+#include "Interfaces/IShaderFormat.h"
 
 #define LOCTEXT_NAMESPACE "MaterialStats"
 
@@ -34,11 +35,14 @@ FShaderPlatformSettings::FShaderPlatformSettings(
 	const FName _Name,
 	const bool _bAllowPresenceInGrid,
 	const bool _bAllowCodeView,
-	const FString& _Description)
+	const FString& _Description,
+	const bool bAlwaysOn
+	)
 	: PlatformType(_PlatformType)
 	, PlatformShaderID(_ShaderPlatformID)
 	, PlatformName(_Name)
 	, PlatformDescription(_Description)
+	, bAlwaysOn(bAlwaysOn)
 	, bAllowCodeView(_bAllowCodeView)
 	, bAllowPresenceInGrid(_bAllowPresenceInGrid)
 {
@@ -413,21 +417,22 @@ FMaterialStats::~FMaterialStats()
 	SaveSettings();
 }
 
-void FMaterialStats::Initialize(IMaterialEditor* InMaterialEditor, const bool ShowMaterialInstancesMenu)
+void FMaterialStats::Initialize(IMaterialEditor* InMaterialEditor, const bool bShowMaterialInstancesMenu, const bool bAllowIgnoringCompilationErrors)
 {
 	MaterialEditor = InMaterialEditor;
 
 	StatsGrid = MakeShareable(new FMaterialStatsGrid(AsShared()));
 
-	BuildShaderPlatformDB();
+	BuildShaderPlatformDB(bAllowIgnoringCompilationErrors);
 
-	LoadSettings();
+	LoadSettings(bAllowIgnoringCompilationErrors);
 
 	StatsGrid->BuildGrid();
 
 	GridStatsWidget = SNew(SMaterialEditorStatsWidget)
 		.MaterialStatsWPtr(SharedThis(this))
-		.ShowMaterialInstancesMenu(ShowMaterialInstancesMenu);
+		.ShowMaterialInstancesMenu(bShowMaterialInstancesMenu)
+		.AllowIgnoringCompilationErrors(bAllowIgnoringCompilationErrors);
 
 	FMessageLogModule& MessageLogModule = FModuleManager::LoadModuleChecked<FMessageLogModule>("MessageLog");
 	FMessageLogInitializationOptions LogOptions;
@@ -455,7 +460,7 @@ void FMaterialStats::Initialize(IMaterialEditor* InMaterialEditor, const bool Sh
 		FIsActionChecked::CreateSP(this, &FMaterialStats::IsShowingOldStats));
 }
 
-void FMaterialStats::LoadSettings()
+void FMaterialStats::LoadSettings(const bool bAllowIgnoringCompilationErrors)
 {
 	Options = NewObject<UMaterialStatsOptions>();
 
@@ -471,19 +476,25 @@ void FMaterialStats::LoadSettings()
 	for (int32 i = 0; i < EMaterialQualityLevel::Num; ++i)
 	{
 		const bool bUsed = !!Options->bMaterialQualityUsed[i];
-		bArrStatsQualitySelector[(EMaterialQualityLevel::Type)i] = bUsed;
+		bArrStatsQualitySelector[(EMaterialQualityLevel::Type)i] = bUsed || bArrStatsQualitySelectorAlwaysOn[(EMaterialQualityLevel::Type)i];
 
 		for (const auto& PlatformEntry : ShaderPlatformStatsDB)
 		{
 			TSharedPtr<FShaderPlatformSettings> SomePlatform = PlatformEntry.Value;
 			if (SomePlatform.IsValid())
 			{
-				SomePlatform->SetExtractStatsQualityLevel((EMaterialQualityLevel::Type)i, bUsed);
+				SomePlatform->SetExtractStatsFlag((EMaterialQualityLevel::Type)i, bUsed);
 			}
 		}
 	}
 
 	MaterialStatsDerivedMIOption = Options->MaterialStatsDerivedMIOption;
+
+	// force compilation of derived instances if we're not allowed to ignore compilation errors
+	if (!bAllowIgnoringCompilationErrors && (MaterialStatsDerivedMIOption == EMaterialStatsDerivedMIOption::Ignore || MaterialStatsDerivedMIOption == EMaterialStatsDerivedMIOption::InvalidOrMax))
+	{
+		MaterialStatsDerivedMIOption = EMaterialStatsDerivedMIOption::CompileOnly;
+	}
 }
 
 void FMaterialStats::SaveSettings()
@@ -568,35 +579,44 @@ void FMaterialStats::RefreshStatsGrid()
 	GetGridStatsWidget()->RequestRefresh();
 }
 
-void FMaterialStats::BuildShaderPlatformDB()
+void FMaterialStats::BuildShaderPlatformDB(const bool bAllowIgnoringCompilationErrors)
 {
-	bool bCanCompileMacDesktopSPs = (PLATFORM_MAC != 0);
-	bool bCanCompileMacMobileSPs = (PLATFORM_MAC != 0) || (PLATFORM_WINDOWS != 0);
+	// set High quality level as always on if we're not allowed to ignore compilation errors in UI
+	// this will not allow to disable this particular quality level 
+	if (!bAllowIgnoringCompilationErrors)
+	{
+		bArrStatsQualitySelectorAlwaysOn[EMaterialQualityLevel::High] = true;
+		bArrStatsQualitySelector[EMaterialQualityLevel::High] = true;
+	}
 
 #if PLATFORM_WINDOWS
 	// DirectX
-	AddShaderPlatform(EPlatformCategoryType::Desktop, SP_PCD3D_SM5, TEXT("DirectX SM5"), true, true, TEXT("Desktop, DirectX, Shader Model 5"));
-	AddShaderPlatform(EPlatformCategoryType::Desktop, SP_PCD3D_SM6, TEXT("DirectX SM6"), true, true, TEXT("Desktop, DirectX, Shader Model 6"));
-	AddShaderPlatform(EPlatformCategoryType::Desktop, SP_PCD3D_ES3_1, TEXT("DirectX  Mobile"), true, true, TEXT("Desktop, DirectX, Mobile"));
+	AddShaderPlatform(EPlatformCategoryType::Desktop, SP_PCD3D_SM5, TEXT("DirectX SM5"), true, TEXT("Desktop, DirectX, Shader Model 5"));
+	AddShaderPlatform(EPlatformCategoryType::Desktop, SP_PCD3D_SM6, TEXT("DirectX SM6"), true, TEXT("Desktop, DirectX, Shader Model 6"), !bAllowIgnoringCompilationErrors);
+	AddShaderPlatform(EPlatformCategoryType::Desktop, SP_PCD3D_ES3_1, TEXT("DirectX  Mobile"), true, TEXT("Desktop, DirectX, Mobile"), !bAllowIgnoringCompilationErrors);
 #endif
 
 	// Vulkan
-	AddShaderPlatform(EPlatformCategoryType::Desktop, SP_VULKAN_SM5, TEXT("Vulkan SM5"), true, true, TEXT("Desktop, Vulkan, Shader Model 5"));
-	AddShaderPlatform(EPlatformCategoryType::Desktop, SP_VULKAN_SM6, TEXT("Vulkan SM6"), true, true, TEXT("Desktop, Vulkan, Shader Model 6"));
-	AddShaderPlatform(EPlatformCategoryType::Desktop, SP_VULKAN_PCES3_1, TEXT("Vulkan Mobile"), true, true, TEXT("Desktop, Vulkan, Mobile"));
+	AddShaderPlatform(EPlatformCategoryType::Desktop, SP_VULKAN_SM5, TEXT("Vulkan SM5"), true, TEXT("Desktop, Vulkan, Shader Model 5"));
+	AddShaderPlatform(EPlatformCategoryType::Desktop, SP_VULKAN_SM6, TEXT("Vulkan SM6"), true, TEXT("Desktop, Vulkan, Shader Model 6"));
+	AddShaderPlatform(EPlatformCategoryType::Desktop, SP_VULKAN_PCES3_1, TEXT("Vulkan Mobile"), true, TEXT("Desktop, Vulkan, Mobile"));
 
 	// Android
-	AddShaderPlatform(EPlatformCategoryType::Android, SP_OPENGL_ES3_1_ANDROID, TEXT("Android GLES Mobile"), true, true, TEXT("Android, OpenGLES Mobile"));
-	AddShaderPlatform(EPlatformCategoryType::Android, SP_VULKAN_ES3_1_ANDROID, TEXT("Android Vulkan Mobile"), true, true, TEXT("Android, Vulkan Mobile"));
-	AddShaderPlatform(EPlatformCategoryType::Android, SP_VULKAN_SM5_ANDROID, TEXT("Android Vulkan SM5"), true, true, TEXT("Android, Vulkan SM5"));
-
-	// Apple
-	AddShaderPlatform(EPlatformCategoryType::Desktop, SP_METAL_SM5, TEXT("Metal SM5"), bCanCompileMacDesktopSPs, true, TEXT("macOS, Metal, Shader Model 5"));
-    AddShaderPlatform(EPlatformCategoryType::Desktop, SP_METAL_SM6, TEXT("Metal SM6"), bCanCompileMacDesktopSPs, true, TEXT("macOS, Metal, Shader Model 6"));
-	AddShaderPlatform(EPlatformCategoryType::Desktop, SP_METAL_MRT_MAC, TEXT("Metal SM5 (MRT)"), bCanCompileMacDesktopSPs, true, TEXT("macOS, Metal, Shader Model 5"));
-	AddShaderPlatform(EPlatformCategoryType::IOS, SP_METAL, TEXT("Metal"), bCanCompileMacMobileSPs, true, TEXT("iOS, Metal, Mobile"));
-	AddShaderPlatform(EPlatformCategoryType::IOS, SP_METAL_MRT, TEXT("Metal MRT"), bCanCompileMacMobileSPs, true, TEXT("iOS, Metal, Shader Model 5"));
-	AddShaderPlatform(EPlatformCategoryType::IOS, SP_METAL_SIM, TEXT("Metal Simulator"), bCanCompileMacMobileSPs, true, TEXT("iOS, Metal, Mobile"));
+	AddShaderPlatform(EPlatformCategoryType::Android, SP_OPENGL_ES3_1_ANDROID, TEXT("Android GLES Mobile"), true, TEXT("Android, OpenGLES Mobile"));
+	AddShaderPlatform(EPlatformCategoryType::Android, SP_VULKAN_ES3_1_ANDROID, TEXT("Android Vulkan Mobile"), true, TEXT("Android, Vulkan Mobile"));
+	AddShaderPlatform(EPlatformCategoryType::Android, SP_VULKAN_SM5_ANDROID, TEXT("Android Vulkan SM5"), true, TEXT("Android, Vulkan SM5"));
+
+#if PLATFORM_MAC
+	// macOS
+	AddShaderPlatform(EPlatformCategoryType::Desktop, SP_METAL_SM5, TEXT("Metal SM5"), true, TEXT("macOS, Metal, Shader Model 5"));
+	AddShaderPlatform(EPlatformCategoryType::Desktop, SP_METAL_SM6, TEXT("Metal SM6"), true, TEXT("macOS, Metal, Shader Model 6"), !bAllowIgnoringCompilationErrors);
+	AddShaderPlatform(EPlatformCategoryType::Desktop, SP_METAL_MRT_MAC, TEXT("Metal SM5 (MRT)"), true, TEXT("macOS, Metal, Shader Model 5"));
+#endif
+
+	// iOS
+	AddShaderPlatform(EPlatformCategoryType::IOS, SP_METAL, TEXT("Metal"), true, TEXT("iOS, Metal, Mobile"));
+	AddShaderPlatform(EPlatformCategoryType::IOS, SP_METAL_MRT, TEXT("Metal MRT"), true, TEXT("iOS, Metal, Shader Model 5"));
+	AddShaderPlatform(EPlatformCategoryType::IOS, SP_METAL_SIM, TEXT("Metal Simulator"), true, TEXT("iOS, Metal, Mobile"));
 
 	ITargetPlatformManagerModule& TPM = GetTargetPlatformManagerRef();
 
@@ -612,22 +632,25 @@ void FMaterialStats::BuildShaderPlatformDB()
 			if (TPM.FindShaderFormat(ShaderFormat) != nullptr)
 			{
 				const FName PlatformName = ShaderPlatformToPlatformName(ShaderPlatform);
-				AddShaderPlatform(bIsConsole ? EPlatformCategoryType::Console : EPlatformCategoryType::Desktop, ShaderPlatform, PlatformName, true, true, PlatformName.ToString());
+				AddShaderPlatform(bIsConsole ? EPlatformCategoryType::Console : EPlatformCategoryType::Desktop, ShaderPlatform, PlatformName, true, PlatformName.ToString());
 			}
 		}
 	}
 
 }
 
-TSharedPtr<FShaderPlatformSettings> FMaterialStats::AddShaderPlatform(const EPlatformCategoryType PlatformType, const EShaderPlatform PlatformID, const FName PlatformName,
-	const bool bAllowPresenceInGrid, const bool bAllowCodeView, const FString& Description)
+TSharedPtr<FShaderPlatformSettings> FMaterialStats::AddShaderPlatform(const EPlatformCategoryType PlatformType, const EShaderPlatform PlatformID, const FName PlatformName, const bool bAllowCodeView, const FString& Description, const bool bAlwaysOn)
 {
+	ITargetPlatformManagerModule& TPM = GetTargetPlatformManagerRef();
+	const IShaderFormat* ShaderFormat = TPM.FindShaderFormat(LegacyShaderPlatformToShaderFormat(PlatformID));
+	const bool bAllowPresenceInGrid = ShaderFormat ? ShaderFormat->CanCompileBinaryShaders() : false;
+
 	if (!FDataDrivenShaderPlatformInfo::IsValid(PlatformID))
 	{
 		return TSharedPtr<FShaderPlatformSettings>();
 	}
 
-	TSharedPtr<FShaderPlatformSettings> PlatformPtr = MakeShareable(new FShaderPlatformSettings(PlatformType, PlatformID, PlatformName, bAllowPresenceInGrid, bAllowCodeView, Description));
+	TSharedPtr<FShaderPlatformSettings> PlatformPtr = MakeShareable(new FShaderPlatformSettings(PlatformType, PlatformID, PlatformName, bAllowPresenceInGrid, bAllowCodeView, Description, bAlwaysOn));
 	ShaderPlatformStatsDB.Add(PlatformID, PlatformPtr);
 
 	auto& ArrayPlatforms = PlatformTypeDB.FindOrAdd(PlatformType);
@@ -635,7 +658,7 @@ TSharedPtr<FShaderPlatformSettings> FMaterialStats::AddShaderPlatform(const EPla
 
 	for (int32 i = 0; i < EMaterialQualityLevel::Num; ++i)
 	{
-		PlatformPtr->SetExtractStatsFlag((EMaterialQualityLevel::Type)i, bArrStatsQualitySelector[i]);
+		PlatformPtr->SetExtractStatsFlag((EMaterialQualityLevel::Type)i, bAlwaysOn && bArrStatsQualitySelectorAlwaysOn[i]);
 	}
 
 	return PlatformPtr;
@@ -672,14 +695,14 @@ void FMaterialStats::SetStatusQualityFlag(const EMaterialQualityLevel::Type Qual
 {
 	check(QualityLevel < EMaterialQualityLevel::Num);
 
-	bArrStatsQualitySelector[QualityLevel] = bValue;
+	bArrStatsQualitySelector[QualityLevel] = bValue || bArrStatsQualitySelectorAlwaysOn[QualityLevel];
 
 	for (const auto& PlatformEntry : ShaderPlatformStatsDB)
 	{
 		TSharedPtr<FShaderPlatformSettings> SomePlatform = PlatformEntry.Value;
 		if (SomePlatform.IsValid())
 		{
-			SomePlatform->SetExtractStatsQualityLevel(QualityLevel, bValue);
+			SomePlatform->SetExtractStatsFlag(QualityLevel, bValue);
 		}
 	}
 
diff --git a/Engine/Source/Editor/MaterialEditor/Private/MaterialStats.h b/Engine/Source/Editor/MaterialEditor/Private/MaterialStats.h
--- a/Engine/Source/Editor/MaterialEditor/Private/MaterialStats.h
+++ b/Engine/Source/Editor/MaterialEditor/Private/MaterialStats.h
@@ -169,6 +169,9 @@ private:
 	/** if true this will be visible in the material stats. grid */
 	bool bPresentInGrid = false;
 
+	/** if true this platform is always present in material stats grid */
+	bool bAlwaysOn = false;
+
 	/** if true this will be listed in the 'view code' menu */
 	bool bAllowCodeView = false;
 
@@ -200,7 +203,8 @@ public:
 		const FName _Name,
 		const bool _bAllowPresenceInGrid,
 		const bool _bAllowCodeView,
-		const FString& _Description);
+		const FString& _Description,
+		const bool bAlwaysOn);
 	~FShaderPlatformSettings();
 
 	/** returns the name of this platform given in the constructor */
@@ -241,15 +245,15 @@ public:
 	/** returns whether or not this material was chosen to be displayed in the stats grid widget with any material quality level */
 	FORCEINLINE bool IsPresentInGrid() const;
 
+	/** returns whether or not this material is always presnet in the stats grid widget */
+	FORCEINLINE bool IsAlwaysOn() const;
+
 	/** used by the grid widget to enable or disable the presence of this material */
 	FORCEINLINE bool FlipPresentInGrid();
 
 	/** used  to enable or disable the presence of this material in the stats grid widget */
 	FORCEINLINE void SetPresentInGrid(const bool bValue);
 
-	/** marks the need to extract statistics for the specified material quality */
-	FORCEINLINE void SetExtractStatsQualityLevel(const EMaterialQualityLevel::Type Quality, const bool bActive);
-
 	/** flags shader compilation for a specific quality level */
 	FORCEINLINE void SetNeedShaderCompilation(const EMaterialQualityLevel::Type QualityLevel, const bool bValue, const bool bOnlyCompileDerivedMI);
 
@@ -306,6 +310,7 @@ class FMaterialStats : public FGCObject, public TSharedFromThis<FMaterialStats>
 
 	/** array of bools that flag a specific global material quality setting of the stats grid widget */
 	bool bArrStatsQualitySelector[EMaterialQualityLevel::Num] = { false };
+	bool bArrStatsQualitySelectorAlwaysOn[EMaterialQualityLevel::Num] = { false };
 
 	/** inspect all derived material instances also */
 	EMaterialStatsDerivedMIOption MaterialStatsDerivedMIOption = EMaterialStatsDerivedMIOption::CompileOnly;
@@ -349,11 +354,10 @@ class FMaterialStats : public FGCObject, public TSharedFromThis<FMaterialStats>
 
 private:
 	/** adds a specified platform in the grid widget for analysis; usually called from BuildShaderPlatformDB() */
-	TSharedPtr<FShaderPlatformSettings> AddShaderPlatform(const EPlatformCategoryType PlatformType, const EShaderPlatform PlatformID, const FName PlatformName,
-		const bool bAllowPresenceInGrid, const bool bAllowCodeView, const FString& Description);
+	TSharedPtr<FShaderPlatformSettings> AddShaderPlatform(const EPlatformCategoryType PlatformType, const EShaderPlatform PlatformID, const FName PlatformName, const bool bAllowCodeView, const FString& Description, const bool bAlwaysOn = false);
 
 	/** build a collection of available shader platform for which we can extract various statistics */
-	void BuildShaderPlatformDB();
+	void BuildShaderPlatformDB(const bool bAllowIgnoringCompilationErrors);
 
 	/** this will spawn the window that will display the a specific set of shaders from the analyzed material */
 	TSharedRef<class SDockTab> SpawnTab_ShaderCode(const class FSpawnTabArgs& Args, const EShaderPlatform PlatformID, const EMaterialQualityLevel::Type QualityType, const int32 InstanceIndex);
@@ -394,7 +398,7 @@ private:
 	void DisplayStatsGrid(const bool bShow);
 	void DisplayOldStats(const bool bShow);
 
-	void LoadSettings();
+	void LoadSettings(const bool bAllowIgnoringCompilationErrors);
 	void SaveSettings();
 
 	/** function that will collect (eventual) warning messages when the stats grid is not properly configured */
@@ -433,6 +437,7 @@ public:
 	bool SwitchStatsQualityFlag(const EMaterialQualityLevel::Type Quality);
 	void SetStatusQualityFlag(const EMaterialQualityLevel::Type Quality, const bool bValue);
 	FORCEINLINE bool GetStatsQualityFlag(const EMaterialQualityLevel::Type Quality);
+	FORCEINLINE bool GetStatsQualityFlagAlwaysOn(const EMaterialQualityLevel::Type Quality);
 
 	void SetMaterialStatsDerivedMIOption(const EMaterialStatsDerivedMIOption value);
 	FORCEINLINE EMaterialStatsDerivedMIOption GetMaterialStatsDerivedMIOption() const;
@@ -479,7 +484,7 @@ public:
 
 private:
 	/** this function will do the setup procedure and its called from FMaterialStatsUtils::CreateMaterialStats()  */
-	void Initialize(IMaterialEditor *MaterialEditor, const bool ShowMaterialInstancesMenu);
+	void Initialize(IMaterialEditor *MaterialEditor, const bool bShowMaterialInstancesMenu, const bool bAllowIgnoringCompilationErrors);
 public:
 	//end Setup Functions
 	///////////////////////////////////////////
@@ -570,6 +575,12 @@ FORCEINLINE bool FMaterialStats::GetStatsQualityFlag(const EMaterialQualityLevel
 	return bArrStatsQualitySelector[(int32)Quality];
 }
 
+FORCEINLINE bool FMaterialStats::GetStatsQualityFlagAlwaysOn(const EMaterialQualityLevel::Type Quality)
+{
+	check(Quality < EMaterialQualityLevel::Num);
+	return bArrStatsQualitySelectorAlwaysOn[(int32)Quality];
+}
+
 FORCEINLINE EMaterialStatsDerivedMIOption FMaterialStats::GetMaterialStatsDerivedMIOption() const
 {
 	return MaterialStatsDerivedMIOption;
@@ -665,6 +676,11 @@ FORCEINLINE bool FShaderPlatformSettings::IsPresentInGrid() const
 	return bPresentInGrid;
 }
 
+FORCEINLINE bool FShaderPlatformSettings::IsAlwaysOn() const
+{
+	return bAlwaysOn;
+}
+
 FORCEINLINE bool FShaderPlatformSettings::FlipPresentInGrid()
 {
 	SetPresentInGrid(!IsPresentInGrid());
@@ -674,12 +690,24 @@ FORCEINLINE bool FShaderPlatformSettings::FlipPresentInGrid()
 
 FORCEINLINE void FShaderPlatformSettings::SetPresentInGrid(const bool bValue)
 {
-	bPresentInGrid = bValue;
-}
-
-FORCEINLINE void FShaderPlatformSettings::SetExtractStatsQualityLevel(const EMaterialQualityLevel::Type Quality, const bool bActive)
-{
-	PlatformData[Quality].bExtractStats = bActive;
+	if (bAllowPresenceInGrid) // if we're allowed to be present in the stats grid
+	{
+		if (!bAlwaysOn) // and bAlwaysOn is false - then any bValue is fine
+		{
+			bPresentInGrid = bValue;
+		}
+		else if (bValue == true) // otherwise only enabling is fine, as we shouldn't be able to disable this platform
+		{
+			bPresentInGrid = bValue;
+		}
+	}
+	else
+	{
+		if (bValue == false) // if we're not allowed to be present, only disabling this platform is fine
+		{
+			bPresentInGrid = bValue;
+		}
+	}
 }
 
 FORCEINLINE void FShaderPlatformSettings::SetNeedShaderCompilation(const EMaterialQualityLevel::Type QualityLevel, const bool bValue, const bool bOnlyCompileDerivedMI)
diff --git a/Engine/Source/Editor/MaterialEditor/Private/MaterialStatsCommon.cpp b/Engine/Source/Editor/MaterialEditor/Private/MaterialStatsCommon.cpp
--- a/Engine/Source/Editor/MaterialEditor/Private/MaterialStatsCommon.cpp
+++ b/Engine/Source/Editor/MaterialEditor/Private/MaterialStatsCommon.cpp
@@ -27,10 +27,10 @@ void FMaterialResourceStats::SetupExtraCompilationSettings(const EShaderPlatform
 /***********************************************************************************************************************/
 /*begin FMaterialStatsUtils */
 
-TSharedPtr<FMaterialStats> FMaterialStatsUtils::CreateMaterialStats(class IMaterialEditor* MaterialEditor, const bool ShowMaterialInstancesMenu)
+TSharedPtr<FMaterialStats> FMaterialStatsUtils::CreateMaterialStats(class IMaterialEditor* MaterialEditor, const bool bShowMaterialInstancesMenu, const bool bAllowIgnoringCompilationErrors)
 {
 	TSharedPtr<FMaterialStats> MaterialStats = MakeShareable(new FMaterialStats());
-	MaterialStats->Initialize(MaterialEditor, ShowMaterialInstancesMenu);
+	MaterialStats->Initialize(MaterialEditor, bShowMaterialInstancesMenu, bAllowIgnoringCompilationErrors);
 
 	return MaterialStats;
 }
diff --git a/Engine/Source/Editor/MaterialEditor/Private/MaterialStatsGrid.cpp b/Engine/Source/Editor/MaterialEditor/Private/MaterialStatsGrid.cpp
--- a/Engine/Source/Editor/MaterialEditor/Private/MaterialStatsGrid.cpp
+++ b/Engine/Source/Editor/MaterialEditor/Private/MaterialStatsGrid.cpp
@@ -927,7 +927,7 @@ void FMaterialStatsGrid::OnQualitySettingChanged(const EMaterialQualityLevel::Ty
 	if (!StatsManager.IsValid())
 		return;
 
-	bool bQualityOn = StatsManager->GetStatsQualityFlag(QualityLevel);
+	const bool bQualityOn = StatsManager->GetStatsQualityFlag(QualityLevel);
 
 	const auto& PlatformDB = StatsManager->GetPlatformsDB();
 	for (auto PlatformPair : PlatformDB)
diff --git a/Engine/Source/Editor/MaterialEditor/Private/SMaterialEditorStatsWidget.cpp b/Engine/Source/Editor/MaterialEditor/Private/SMaterialEditorStatsWidget.cpp
--- a/Engine/Source/Editor/MaterialEditor/Private/SMaterialEditorStatsWidget.cpp
+++ b/Engine/Source/Editor/MaterialEditor/Private/SMaterialEditorStatsWidget.cpp
@@ -354,7 +354,7 @@ uint32 SMaterialEditorStatsWidget::CountSubPlatforms(EPlatformCategoryType Categ
 	for (int32 i = 0; i < ArrPlatforms.Num(); ++i)
 	{
 		auto PlatformPtr = ArrPlatforms[i];
-		if (!PlatformPtr.IsValid() || !PlatformPtr->IsStatsGridPresenceAllowed())
+		if (!PlatformPtr.IsValid())
 		{
 			continue;
 		}
@@ -387,7 +387,7 @@ void SMaterialEditorStatsWidget::CreatePlatformMenus(FMenuBuilder& Builder, EPla
 	for (int32 i = 0; i < ArrPlatforms.Num(); ++i)
 	{
 		auto PlatformPtr = ArrPlatforms[i];
-		if (!PlatformPtr.IsValid() || !PlatformPtr->IsStatsGridPresenceAllowed())
+		if (!PlatformPtr.IsValid())
 		{
 			continue;
 		}
@@ -405,6 +405,12 @@ void SMaterialEditorStatsWidget::CreatePlatformMenus(FMenuBuilder& Builder, EPla
 			return ECheckBoxState::Unchecked;
 		};
 
+		// lambda function used to determine the enable state for the shader platform checkboxes
+		auto Lamda_PlatformEnableState = [PlatformPtr = PlatformPtr]()
+		{
+			return PlatformPtr.IsValid() && !PlatformPtr->IsAlwaysOn() && PlatformPtr->IsStatsGridPresenceAllowed();
+		};
+
 		// lambda used with shader platform checkboxes to add or remove selected shader platforms
 		auto Lamda_PlatformFlipState = [WidgetPtr = this, PlatformPtr = PlatformPtr](const ECheckBoxState NewState)
 		{
@@ -418,9 +424,10 @@ void SMaterialEditorStatsWidget::CreatePlatformMenus(FMenuBuilder& Builder, EPla
 			WidgetPtr->RequestRefresh();
 		};
 
-		auto PlatformWidget = SNew(SCheckBox)			
+		auto PlatformWidget = SNew(SCheckBox)
 			.OnCheckStateChanged_Lambda(Lamda_PlatformFlipState)
 			.IsChecked_Lambda(Lamda_PlatformCheckState)
+			.IsEnabled_Lambda(Lamda_PlatformEnableState)
 			.Content()
 			[
 				SNew(STextBlock)
@@ -480,11 +487,18 @@ void SMaterialEditorStatsWidget::CreateQualityMenus(FMenuBuilder& Builder)
 			return ECheckBoxState::Unchecked;
 		};
 
+		auto Lamba_QualityCheckAlwaysOn = [QualityType = (EMaterialQualityLevel::Type)i, MaterialStatsWPtr = MaterialStatsWPtr]()
+		{
+			const auto StatsPtr = MaterialStatsWPtr.Pin();
+			return StatsPtr.IsValid() && !StatsPtr->GetStatsQualityFlagAlwaysOn(QualityType);
+		};
+
 		const FText QualitySettingName = FText::FromString(FMaterialStatsUtils::MaterialQualityToString(QualityLevel));
 
 		auto QualityWidget = SNew(SCheckBox)
 			.OnCheckStateChanged(this, &SMaterialEditorStatsWidget::OnFlipQualityState, QualityLevel)
 			.IsChecked_Lambda(Lamda_QualityCheckState)
+			.IsEnabled_Lambda(Lamba_QualityCheckAlwaysOn)
 			.Content()
 			[
 				SNew(STextBlock)
@@ -514,7 +528,10 @@ void SMaterialEditorStatsWidget::CreateDerivedMaterialsMenu(class FMenuBuilder&
 	{
 		DerivedMaterialInstancesComboBoxItems[static_cast<int32>(Option)] = MakeShared<EMaterialStatsDerivedMIOption>(Option);
 	};
-	AddOption(EMaterialStatsDerivedMIOption::Ignore);
+	if (bAllowIgnoringCompilationErrors)
+	{
+		AddOption(EMaterialStatsDerivedMIOption::Ignore);
+	}
 	AddOption(EMaterialStatsDerivedMIOption::CompileOnly);
 	AddOption(EMaterialStatsDerivedMIOption::ShowStats);
 
@@ -566,7 +583,7 @@ TSharedRef<SWidget> SMaterialEditorStatsWidget::GetSettingsButtonContent()
 	FMenuBuilder Builder(false, nullptr);
 
 	CreatePlatformCategoryMenus(Builder);
-	if (ShowMaterialInstancesMenu)
+	if (bShowMaterialInstancesMenu)
 	{
 		Builder.AddMenuSeparator();
 		CreateDerivedMaterialsMenu(Builder);
@@ -625,7 +642,8 @@ void SMaterialEditorStatsWidget::ClearMessages()
 void SMaterialEditorStatsWidget::Construct(const FArguments& InArgs)
 {
 	MaterialStatsWPtr = InArgs._MaterialStatsWPtr;
-	ShowMaterialInstancesMenu = InArgs._ShowMaterialInstancesMenu;
+	bShowMaterialInstancesMenu = InArgs._ShowMaterialInstancesMenu;
+	bAllowIgnoringCompilationErrors = InArgs._AllowIgnoringCompilationErrors;
 
 	const auto StatsPtr = MaterialStatsWPtr.Pin();
 
diff --git a/Engine/Source/Editor/MaterialEditor/Private/SMaterialEditorStatsWidget.h b/Engine/Source/Editor/MaterialEditor/Private/SMaterialEditorStatsWidget.h
--- a/Engine/Source/Editor/MaterialEditor/Private/SMaterialEditorStatsWidget.h
+++ b/Engine/Source/Editor/MaterialEditor/Private/SMaterialEditorStatsWidget.h
@@ -20,10 +20,12 @@ public:
 	SLATE_BEGIN_ARGS(SMaterialEditorStatsWidget)
 		: _MaterialStatsWPtr(nullptr)
 		, _ShowMaterialInstancesMenu(true)
+		, _AllowIgnoringCompilationErrors(true)
 	{}
 
 		SLATE_ARGUMENT(TWeakPtr<FMaterialStats>, MaterialStatsWPtr)
 		SLATE_ARGUMENT(bool, ShowMaterialInstancesMenu)
+		SLATE_ARGUMENT(bool, AllowIgnoringCompilationErrors)
 	SLATE_END_ARGS()
 
 	void Construct(const FArguments& InArgs);
@@ -82,7 +84,8 @@ protected:
 	TSharedPtr<SVerticalBox> MessageBoxWidget;
 
 	TWeakPtr<FMaterialStats> MaterialStatsWPtr;
-	bool ShowMaterialInstancesMenu;
+	bool bShowMaterialInstancesMenu;
+	bool bAllowIgnoringCompilationErrors;
 
 	TSharedPtr<SListView<TSharedPtr<int32>>> MaterialInfoList;
 
diff --git a/Engine/Source/Editor/MaterialEditor/Public/MaterialStatsCommon.h b/Engine/Source/Editor/MaterialEditor/Public/MaterialStatsCommon.h
--- a/Engine/Source/Editor/MaterialEditor/Public/MaterialStatsCommon.h
+++ b/Engine/Source/Editor/MaterialEditor/Public/MaterialStatsCommon.h
@@ -93,7 +93,7 @@ public:
 
 public:
 	/** call this to create an instance to FMaterialStats */
-	static TSharedPtr<class FMaterialStats> CreateMaterialStats(class IMaterialEditor* MaterialEditor, const bool ShowMaterialInstancesMenu);
+	static TSharedPtr<class FMaterialStats> CreateMaterialStats(class IMaterialEditor* MaterialEditor, const bool bShowMaterialInstancesMenu, const bool bAllowIgnoringCompilationErrors);
 
 	/** utility functions that translate various enum values to strings */
 	static FString MaterialQualityToString(const EMaterialQualityLevel::Type Quality);
