[Shaders][ODSC] - Implement on demand shader compilation when recompile all materials after saving a material function.

- Reduces shaders compiled from 13,000 to 300 in my test case
- ~4 seconds saved in UMaterialFunctionInterface::ForceRecompileForRendering due to not having to build and submit shader jobs.

diff --git a/Engine/Source/Editor/MaterialEditor/Private/MaterialEditingLibrary.cpp b/Engine/Source/Editor/MaterialEditor/Private/MaterialEditingLibrary.cpp
--- a/Engine/Source/Editor/MaterialEditor/Private/MaterialEditingLibrary.cpp
+++ b/Engine/Source/Editor/MaterialEditor/Private/MaterialEditingLibrary.cpp
@@ -984,6 +984,8 @@ void UMaterialEditingLibrary::DeleteMaterialExpressionInFunction(UMaterialFuncti
 
 void UMaterialEditingLibrary::UpdateMaterialFunction(UMaterialFunctionInterface* MaterialFunction, UMaterial* PreviewMaterial)
 {
+	TRACE_CPUPROFILER_EVENT_SCOPE(UMaterialEditingLibrary::UpdateMaterialFunction)
+
 	if (MaterialFunction)
 	{
 		// Create a material update context so we can safely update materials using this function.
@@ -994,20 +996,24 @@ void UMaterialEditingLibrary::UpdateMaterialFunction(UMaterialFunctionInterface*
 			MaterialFunction->ForceRecompileForRendering(UpdateContext, PreviewMaterial);
 			MaterialFunction->MarkPackageDirty();
 
-			// Go through all function instances in memory and recompile them if they are children
-			for (TObjectIterator<UMaterialFunctionInstance> It; It; ++It)
 			{
-				UMaterialFunctionInstance* FunctionInstance = *It;
+				TRACE_CPUPROFILER_EVENT_SCOPE(UpdateAllMaterialInstances)
 
-				TArray<UMaterialFunctionInterface*> Functions;
-				FunctionInstance->GetDependentFunctions(Functions);
-				if (Functions.Contains(MaterialFunction))
+				// Go through all function instances in memory and recompile them if they are children
+				for (TObjectIterator<UMaterialFunctionInstance> It; It; ++It)
 				{
-					FunctionInstance->UpdateParameterSet();
-					FunctionInstance->ForceRecompileForRendering(UpdateContext, PreviewMaterial);
+					UMaterialFunctionInstance* FunctionInstance = *It;
+
+					TArray<UMaterialFunctionInterface*> Functions;
+					FunctionInstance->GetDependentFunctions(Functions);
+					if (Functions.Contains(MaterialFunction))
+					{
+						FunctionInstance->UpdateParameterSet();
+						FunctionInstance->ForceRecompileForRendering(UpdateContext, PreviewMaterial);
 
-					// ForceRecompileForRendering will update StateId, so need to mark the package as dirty
-					FunctionInstance->MarkPackageDirty();
+						// ForceRecompileForRendering will update StateId, so need to mark the package as dirty
+						FunctionInstance->MarkPackageDirty();
+					}
 				}
 			}
 
diff --git a/Engine/Source/Runtime/Engine/Classes/Materials/Material.h b/Engine/Source/Runtime/Engine/Classes/Materials/Material.h
--- a/Engine/Source/Runtime/Engine/Classes/Materials/Material.h
+++ b/Engine/Source/Runtime/Engine/Classes/Materials/Material.h
@@ -1187,7 +1187,7 @@ public:
 	/** Allows material properties to be compiled with the option of being overridden by the material attributes input. */
 	ENGINE_API virtual int32 CompilePropertyEx( class FMaterialCompiler* Compiler, const FGuid& AttributeID ) override;
 	ENGINE_API virtual bool ShouldForcePlanePreview() override;
-	ENGINE_API virtual void ForceRecompileForRendering() override;
+	ENGINE_API virtual void ForceRecompileForRendering(EMaterialShaderPrecompileMode CompileMode = EMaterialShaderPrecompileMode::Default) override;
 #endif // WITH_EDITOR
 	//~ End UMaterialInterface Interface.
 
diff --git a/Engine/Source/Runtime/Engine/Classes/Materials/MaterialInstance.h b/Engine/Source/Runtime/Engine/Classes/Materials/MaterialInstance.h
--- a/Engine/Source/Runtime/Engine/Classes/Materials/MaterialInstance.h
+++ b/Engine/Source/Runtime/Engine/Classes/Materials/MaterialInstance.h
@@ -786,7 +786,7 @@ public:
 	virtual ENGINE_API FGraphEventArray PrecachePSOs(const FPSOPrecacheVertexFactoryDataList& VertexFactoryDataList, const FPSOPrecacheParams& PreCacheParams, EPSOPrecachePriority Priority, TArray<FMaterialPSOPrecacheRequestID>& OutMaterialPSORequestIDs) override;
 
 #if WITH_EDITOR
-	ENGINE_API virtual void ForceRecompileForRendering() override;
+	ENGINE_API virtual void ForceRecompileForRendering(EMaterialShaderPrecompileMode CompileMode = EMaterialShaderPrecompileMode::Default) override;
 #endif // WITH_EDITOR
 
 	ENGINE_API virtual float GetOpacityMaskClipValue() const override;
diff --git a/Engine/Source/Runtime/Engine/Classes/Materials/MaterialInterface.h b/Engine/Source/Runtime/Engine/Classes/Materials/MaterialInterface.h
--- a/Engine/Source/Runtime/Engine/Classes/Materials/MaterialInterface.h
+++ b/Engine/Source/Runtime/Engine/Classes/Materials/MaterialInterface.h
@@ -1036,7 +1036,7 @@ public:
 
 #if WITH_EDITOR
 	/** Clears the shader cache and recompiles the shader for rendering. */
-	virtual void ForceRecompileForRendering() {}
+	virtual void ForceRecompileForRendering(EMaterialShaderPrecompileMode CompileMode = EMaterialShaderPrecompileMode::Default) {}
 #endif // WITH_EDITOR
 
 	/**
diff --git a/Engine/Source/Runtime/Engine/Private/Materials/Material.cpp b/Engine/Source/Runtime/Engine/Private/Materials/Material.cpp
--- a/Engine/Source/Runtime/Engine/Private/Materials/Material.cpp
+++ b/Engine/Source/Runtime/Engine/Private/Materials/Material.cpp
@@ -6440,10 +6440,10 @@ void UMaterial::NotifyCompilationFinished(UMaterialInterface* Material)
 	UMaterial::OnMaterialCompilationFinished().Broadcast(Material);
 }
 
-void UMaterial::ForceRecompileForRendering()
+void UMaterial::ForceRecompileForRendering(EMaterialShaderPrecompileMode CompileMode)
 {
 	UpdateCachedExpressionData();
-	CacheResourceShadersForRendering( false );
+	CacheResourceShadersForRendering(false,  CompileMode);
 }
 
 bool UMaterial::CheckInValidStateForCompilation(class FMaterialCompiler* Compiler) const
diff --git a/Engine/Source/Runtime/Engine/Private/Materials/MaterialExpressions.cpp b/Engine/Source/Runtime/Engine/Private/Materials/MaterialExpressions.cpp
--- a/Engine/Source/Runtime/Engine/Private/Materials/MaterialExpressions.cpp
+++ b/Engine/Source/Runtime/Engine/Private/Materials/MaterialExpressions.cpp
@@ -14591,6 +14591,8 @@ void UMaterialFunctionInterface::PostEditChangeProperty(FPropertyChangedEvent& P
 
 void UMaterialFunctionInterface::ForceRecompileForRendering(FMaterialUpdateContext& UpdateContext, UMaterial* InPreviewMaterial)
 {
+	TRACE_CPUPROFILER_EVENT_SCOPE(UMaterialFunctionInterface::ForceRecompileForRendering)
+
 	//@todo - recreate guid only when needed, not when a comment changes
 	StateId = FGuid::NewGuid();
 
@@ -14636,7 +14638,7 @@ void UMaterialFunctionInterface::ForceRecompileForRendering(FMaterialUpdateConte
 		{
 			// Propagate the change to this material
 			UpdateContext.AddMaterialInterface(CurrentMaterialInterface);
-			CurrentMaterialInterface->ForceRecompileForRendering();
+			CurrentMaterialInterface->ForceRecompileForRendering(EMaterialShaderPrecompileMode::None);
 		}
 	}
 }
diff --git a/Engine/Source/Runtime/Engine/Private/Materials/MaterialInstance.cpp b/Engine/Source/Runtime/Engine/Private/Materials/MaterialInstance.cpp
--- a/Engine/Source/Runtime/Engine/Private/Materials/MaterialInstance.cpp
+++ b/Engine/Source/Runtime/Engine/Private/Materials/MaterialInstance.cpp
@@ -2034,10 +2034,10 @@ void UMaterialInstance::GetDependentFunctions(TArray<UMaterialFunctionInterface*
 #endif // WITH_EDITORONLY_DATA
 
 #if WITH_EDITOR
-void UMaterialInstance::ForceRecompileForRendering()
+void UMaterialInstance::ForceRecompileForRendering(EMaterialShaderPrecompileMode CompileMode)
 {
 	UpdateCachedData();
-	CacheResourceShadersForRendering();
+	CacheResourceShadersForRendering(CompileMode);
 }
 #endif // WITH_EDITOR
 
