[Material Editor] Fixed occasional crash when trying to create the preview of a deleted expression.

diff --git a/Engine/Source/Editor/MaterialEditor/Private/MaterialEditor.cpp b/Engine/Source/Editor/MaterialEditor/Private/MaterialEditor.cpp
--- a/Engine/Source/Editor/MaterialEditor/Private/MaterialEditor.cpp
+++ b/Engine/Source/Editor/MaterialEditor/Private/MaterialEditor.cpp
@@ -7282,13 +7282,14 @@ FMatExpressionPreview* FMaterialEditor::GetExpressionPreview(UMaterialExpression
 			}
 		}
 
-		if( !Preview )
+		if (!Preview && MaterialExpression->Material->GetExpressions().Contains(MaterialExpression))
 		{
 			bNewlyCreated = true;
 			Preview = new FMatExpressionPreview(MaterialExpression);
 			ExpressionPreviews.Add(Preview);
 			Preview->CacheShaders(GMaxRHIShaderPlatform, EMaterialShaderPrecompileMode::None);
 		}
+
 		return Preview;
 	}
 
