Delay the material palette(s) refresh on the next tick when disccovering a new material function to improve the performance of the editor during the asset discovery.

diff --git a/Engine/Source/Editor/MaterialEditor/Private/SMaterialPalette.cpp b/Engine/Source/Editor/MaterialEditor/Private/SMaterialPalette.cpp
--- a/Engine/Source/Editor/MaterialEditor/Private/SMaterialPalette.cpp
+++ b/Engine/Source/Editor/MaterialEditor/Private/SMaterialPalette.cpp
@@ -233,8 +233,17 @@ void SMaterialPalette::RenameAssetFromRegistry(const FAssetData& InAddedAssetDat
 void SMaterialPalette::RefreshAssetInRegistry(const FAssetData& InAddedAssetData)
 {
 	if (InAddedAssetData.IsInstanceOf(UMaterialFunction::StaticClass()))
+	{
+		bNeedRefresh = true;
+	}
+}
+
+void SMaterialPalette::Tick(const FGeometry& AllottedGeometry, const double InCurrentTime, const float InDeltaTime)
+{
+	if (bNeedRefresh)
 	{
 		RefreshActionsList(true);
+		bNeedRefresh = false;
 	}
 }
 
diff --git a/Engine/Source/Editor/MaterialEditor/Private/SMaterialPalette.h b/Engine/Source/Editor/MaterialEditor/Private/SMaterialPalette.h
--- a/Engine/Source/Editor/MaterialEditor/Private/SMaterialPalette.h
+++ b/Engine/Source/Editor/MaterialEditor/Private/SMaterialPalette.h
@@ -60,6 +60,8 @@ protected:
 
 	void RefreshAssetInRegistry(const FAssetData& InAddedAssetData);
 
+	virtual void Tick(const FGeometry& AllottedGeometry, const double InCurrentTime, const float InDeltaTime) override;
+
 protected:
 	/** Pointer back to the material editor that owns us */
 	TWeakPtr<FMaterialEditor> MaterialEditorPtr;
@@ -69,4 +71,6 @@ protected:
 
 	/** Combo box used to select category */
 	TSharedPtr<STextComboBox> CategoryComboBox;
+
+	bool bNeedRefresh = false;
 };
