Fix crash when swapping parent material on a material instance.

- When the parent property changes we need to notify the material stats manager to update its cached material resources (see FMaterialInstanceEditor::NotifyPostChange)
- The SetMaterial functions in MaterialStats.cpp never took into account the parent on the material changing and therefore having to update its cached state.  The FShaderPlatformSettings cached off a UMaterial pointer to base material.  When switching the parent material on the instance this code wasn't updating this pointer.  This meant the stats thought it was the world grid material instead of the new parent material.
- I relaxed some of the conditionals in SetMaterial to always update, and only update if the base material changed.

diff --git a/Engine/Source/Editor/MaterialEditor/Private/MaterialInstanceEditor.cpp b/Engine/Source/Editor/MaterialEditor/Private/MaterialInstanceEditor.cpp
--- a/Engine/Source/Editor/MaterialEditor/Private/MaterialInstanceEditor.cpp
+++ b/Engine/Source/Editor/MaterialEditor/Private/MaterialInstanceEditor.cpp
@@ -1140,6 +1140,10 @@ void FMaterialInstanceEditor::NotifyPostChange( const FPropertyChangedEvent& Pro
 		RebuildInheritanceList();
 
 		UpdatePropertyWindow();
+
+		// If the parent of this instance changed we need to update the cached state on the stats manager to have the updated parent.
+		MaterialStatsManager->SetMaterial(MaterialEditorInstance->SourceInstance);
+		MaterialStatsManager->Update();
 	}
 	else if(PropertyThatChanged->GetName() == TEXT("PreviewMesh"))
 	{
diff --git a/Engine/Source/Editor/MaterialEditor/Private/MaterialStats.cpp b/Engine/Source/Editor/MaterialEditor/Private/MaterialStats.cpp
--- a/Engine/Source/Editor/MaterialEditor/Private/MaterialStats.cpp
+++ b/Engine/Source/Editor/MaterialEditor/Private/MaterialStats.cpp
@@ -187,7 +187,9 @@ void FShaderPlatformSettings::SetMaterial(UMaterial *InBaseMaterial, UMaterialIn
 		bReallocate = true;
 	}
 
-	if (InBaseMaterial == nullptr && InBaseMaterialInstance != nullptr && MaterialInstance != InBaseMaterialInstance)
+	if (InBaseMaterial == nullptr && 
+		InBaseMaterialInstance != nullptr && 
+		(MaterialInstance != InBaseMaterialInstance || Material != InBaseMaterialInstance->GetMaterial())) // If the Material Instance changed or its parent changed
 	{
 		Material = InBaseMaterialInstance->GetMaterial();
 		MaterialInstance = InBaseMaterialInstance;
@@ -878,25 +880,22 @@ void FMaterialStats::SetMaterial(UMaterial *InMaterial, const TArray<TObjectPtr<
 
 void FMaterialStats::SetMaterial(UMaterialInstance *InMaterialInstance)
 {
-	if (MaterialInterface != InMaterialInstance || DerivedMaterialInstances.Num() > 0)
-	{
-		MaterialInterface = InMaterialInstance;
-		DerivedMaterialInstances.Empty();
+	MaterialInterface = InMaterialInstance;
+	DerivedMaterialInstances.Empty();
 
-		for (const auto& Entry : ShaderPlatformStatsDB)
+	for (const auto& Entry : ShaderPlatformStatsDB)
+	{
+		auto& Platform = Entry.Value;
+		if (Platform.IsValid())
 		{
-			auto& Platform = Entry.Value;
-			if (Platform.IsValid())
-			{
-				Platform->SetMaterial(nullptr, InMaterialInstance, DerivedMaterialInstances);
-			}
+			Platform->SetMaterial(nullptr, InMaterialInstance, DerivedMaterialInstances);
 		}
+	}
 
-		auto GridPtr = GetGridStatsWidget();
-		if (GridPtr.IsValid())
-		{
-			GridPtr->OnColumnNumChanged();
-		}
+	auto GridPtr = GetGridStatsWidget();
+	if (GridPtr.IsValid())
+	{
+		GridPtr->OnColumnNumChanged();
 	}
 }
 
