Added in material editor info about wether or not a material is compatible with the light function atlas.

diff --git a/Engine/Source/Editor/MaterialEditor/Private/MaterialEditor.cpp b/Engine/Source/Editor/MaterialEditor/Private/MaterialEditor.cpp
--- a/Engine/Source/Editor/MaterialEditor/Private/MaterialEditor.cpp
+++ b/Engine/Source/Editor/MaterialEditor/Private/MaterialEditor.cpp
@@ -3199,6 +3199,15 @@ void FMaterialEditor::UpdateMaterialinfoList_Old()
 					TSharedRef<FTokenizedMessage> PreshaderCountLine = FTokenizedMessage::Create(EMessageSeverity::Info);
 					PreshaderCountLine->AddToken(FTextToken::Create(FText::FromString(PreshaderCountString)));
 					Messages.Add(PreshaderCountLine);
+
+					if (MaterialResource->GetMaterialDomain() == MD_LightFunction)
+					{
+						const bool bIsCompatibleWithLightFunctionAtlas = ShaderMap->IsLightFunctionAtlasCompatible();
+						FString LightFunctionAtlasStr = FString::Printf(TEXT("Light function material%s compatible with the light function atlas for fast batched deferred light shading."), bIsCompatibleWithLightFunctionAtlas ? TEXT(" IS") : TEXT(" IS NOT"));
+						TSharedRef<FTokenizedMessage> LightFunctionAtlasCountLine = FTokenizedMessage::Create(EMessageSeverity::Info);
+						LightFunctionAtlasCountLine->AddToken(FTextToken::Create(FText::FromString(LightFunctionAtlasStr)));
+						Messages.Add(LightFunctionAtlasCountLine);
+					}
 				}
 			}
 
