[MaterialLayers] - Bug fix to get tooltips working for scalar, vector,  , and static bool parameters in the layer parameters view.

- This code in SMaterialLayersFunctionsInstanceTreeItem::Construct is a duplicate of the code in SMaterialParametersOverviewTreeItem::Construct, and it layer version was missing a tooltip.
- Texture parameters and channel mask parameter names are still broken, but for other reasons

todo: cleanup the material parameter display code to not be duplicated.

#tests

diff --git a/Engine/Source/Editor/MaterialEditor/Private/SMaterialLayersFunctionsTree.cpp b/Engine/Source/Editor/MaterialEditor/Private/SMaterialLayersFunctionsTree.cpp
--- a/Engine/Source/Editor/MaterialEditor/Private/SMaterialLayersFunctionsTree.cpp
+++ b/Engine/Source/Editor/MaterialEditor/Private/SMaterialLayersFunctionsTree.cpp
@@ -901,6 +901,10 @@ void SMaterialLayersFunctionsInstanceTreeItem::Construct(const FArguments& InArg
 					]
 				];
 			}
+			else
+			{
+				Row.ToolTip(FMaterialPropertyHelpers::GetParameterExpressionDescription(StackParameterData->Parameter, MaterialEditorInstance));
+			}
 
 			bisPaddedProperty = true;
 		}
@@ -2457,6 +2461,10 @@ void SMaterialLayersFunctionsMaterialTreeItem::Construct(const FArguments& InArg
 					]
 				];
 			}
+			else
+			{
+				Row.ToolTip(FMaterialPropertyHelpers::GetParameterExpressionDescription(StackParameterData->Parameter, MaterialEditorInstance));
+			}
 		}
 
 		FNodeWidgets NodeWidgets = Node.CreateNodeWidgets();
