[Shaders] Improve UI for shader code view in material stats.

- Changes the text style of the shader code view to a monospace font (like the output log).
- Changes the text block of the shader code view to an editable (read-only) text block to allow selecting and copying out of the text field.
- Adds a line-number text bock to the left side of the view.
- Also fix crash when FShaderPlatformSettings::GetShaderCode() is invoked before its instances are initialized.

#rnx

diff --git a/Engine/Source/Editor/MaterialEditor/Private/MaterialStats.cpp b/Engine/Source/Editor/MaterialEditor/Private/MaterialStats.cpp
--- a/Engine/Source/Editor/MaterialEditor/Private/MaterialStats.cpp
+++ b/Engine/Source/Editor/MaterialEditor/Private/MaterialStats.cpp
@@ -9,6 +9,9 @@
 #include "Widgets/Input/SButton.h"
 #include "Widgets/Layout/SSeparator.h"
 #include "Widgets/Layout/SScrollBox.h"
+#include "Widgets/Layout/SSplitter.h"
+#include "Widgets/Text/SMultiLineEditableText.h"
+#include "Text/HLSLSyntaxHighlighterMarshaller.h"
 #include "HAL/PlatformApplicationMisc.h"
 #include "MaterialEditorActions.h"
 #include "Materials/MaterialInstance.h"
@@ -104,6 +107,12 @@ void FShaderPlatformSettings::OnShaderViewComboSelectionChanged(TSharedPtr<FMate
 
 FText FShaderPlatformSettings::GetShaderCode(const EMaterialQualityLevel::Type QualityType, const int32 InstanceIndex)
 {
+	if (PlatformData[QualityType].Instances.IsEmpty())
+	{
+		return FText::GetEmpty();
+	}
+
+	check(InstanceIndex >= 0 && InstanceIndex < PlatformData[QualityType].Instances.Num());
 	auto& Instance = PlatformData[QualityType].Instances[InstanceIndex];
 	// if there were no change to the material return the cached shader code
 	if (!Instance.bUpdateShaderCode)
@@ -923,6 +932,54 @@ bool FMaterialStats::AnyNewCompilationErrors(const int32 StartingFromInstanceInd
 	return bNewCompilationErrors;
 }
 
+// Generates a string of line numbers by counting the number of new-line characters in input source code.
+static FText GenerateLineNoTextForCodeView(const FString& InSourceCode)
+{
+	FString LineNumberStr;
+	int32 LineNo = 1;
+
+	for (TCHAR Chr : InSourceCode)
+	{
+		if (Chr == TEXT('\n'))
+		{
+			LineNumberStr += FString::Printf(TEXT("% 6d\n"), LineNo);
+			++LineNo;
+		}
+	}
+
+	return FText::FromString(MoveTemp(LineNumberStr));
+}
+
+TSharedRef<SScrollBox> FMaterialStats::BuildShaderCodeWidget(TFunction<FText(void)>&& InShaderCodeCallback)
+{
+	return SNew(SScrollBox)
+		+ SScrollBox::Slot().Padding(5)
+		[
+			SNew(SSplitter)
+			.Orientation(EOrientation::Orient_Horizontal)
+			+SSplitter::Slot()
+			.Resizable(false)
+			.SizeRule(SSplitter::SizeToContent)
+			[
+				SNew(STextBlock)
+				.Text_Lambda([InShaderCodeCallback]()
+				{
+					return GenerateLineNoTextForCodeView(InShaderCodeCallback().ToString());
+				})
+				.TextStyle(FAppStyle::Get(), "MessageLog")
+				.ColorAndOpacity(FLinearColor::Gray)
+			]
+			+SSplitter::Slot()
+			[
+				SNew(SMultiLineEditableText)
+				.Text_Lambda([InShaderCodeCallback](){ return InShaderCodeCallback(); })
+				.TextStyle(FAppStyle::Get(), "MessageLog")
+				.IsReadOnly(true)
+				.Marshaller(SyntaxHighlighter)
+			]
+		];
+}
+
 TSharedRef<class SDockTab> FMaterialStats::SpawnTab_HLSLCode(const class FSpawnTabArgs& Args)
 {
 	auto CodeViewUtility =
@@ -959,15 +1016,6 @@ TSharedRef<class SDockTab> FMaterialStats::SpawnTab_HLSLCode(const class FSpawnT
 			SNew(SSeparator)
 		];
 
-	auto CodeView =
-		SNew(SScrollBox)
-		+ SScrollBox::Slot().Padding(5)
-		[
-			SNew(STextBlock)
-			.Text_Lambda([Code = &HLSLCode]() { return FText::FromString(*Code); })
-		];
-
-
 	TSharedRef<SDockTab> SpawnedTab = SNew(SDockTab)
 		.Label(LOCTEXT("HLSLCodeTitle", "HLSL Code"))
 		[
@@ -980,7 +1028,7 @@ TSharedRef<class SDockTab> FMaterialStats::SpawnTab_HLSLCode(const class FSpawnT
 			+ SVerticalBox::Slot()
 			.FillHeight(1)
 			[
-				CodeView
+				BuildShaderCodeWidget([Code = &HLSLCode]() { return FText::FromString(*Code); })
 			]
 		];
 
@@ -1128,6 +1176,19 @@ void FMaterialStats::BuildViewShaderCodeMenus()
 	TSharedPtr<FWorkspaceItem> PlatformGroupMenuItem = ParentCategoryRef->AddGroup(LOCTEXT("ViewShaderCodePlatformsGroupMenu", "Shader Code"),
 		FSlateIcon(FAppStyle::GetAppStyleSetName(), "MaterialEditor.Tabs.HLSLCode"));
 
+	// Create the syntax highlighter
+	FHLSLSyntaxHighlighterMarshaller::FSyntaxTextStyle CodeStyle(
+		FAppStyle::Get().GetWidgetStyle<FTextBlockStyle>("SyntaxHighlight.SourceCode.Normal"),
+		FAppStyle::Get().GetWidgetStyle<FTextBlockStyle>("SyntaxHighlight.SourceCode.Operator"),
+		FAppStyle::Get().GetWidgetStyle<FTextBlockStyle>("SyntaxHighlight.SourceCode.Keyword"),
+		FAppStyle::Get().GetWidgetStyle<FTextBlockStyle>("SyntaxHighlight.SourceCode.String"),
+		FAppStyle::Get().GetWidgetStyle<FTextBlockStyle>("SyntaxHighlight.SourceCode.Number"),
+		FAppStyle::Get().GetWidgetStyle<FTextBlockStyle>("SyntaxHighlight.SourceCode.Comment"),
+		FAppStyle::Get().GetWidgetStyle<FTextBlockStyle>("SyntaxHighlight.SourceCode.PreProcessorKeyword"),
+		FAppStyle::Get().GetWidgetStyle<FTextBlockStyle>("SyntaxHighlight.SourceCode.Error"));
+
+	SyntaxHighlighter = FHLSLSyntaxHighlighterMarshaller::Create(CodeStyle);
+
 	// add hlsl code viewer tab
 	TabManager->RegisterTabSpawner( HLSLCodeTabId, FOnSpawnTab::CreateSP(this, &FMaterialStats::SpawnTab_HLSLCode))
 		.SetDisplayName( LOCTEXT("HLSLCodeTab", "HLSL Code") )
@@ -1175,18 +1236,16 @@ void FMaterialStats::BuildViewShaderCodeMenus()
 					.SetGroup(ShaderPlatformMenuItem.ToSharedRef())
 					.SetDisplayName(FText::FromString(MaterialQualityName));
 
-				auto CodeScrollBox = SNew(SScrollBox)
+				TSharedRef<SScrollBox> CodeScrollBox = SNew(SScrollBox)
 					+ SScrollBox::Slot().Padding(5)
 					[
-						SNew(STextBlock)
-						.Text_Lambda([MaterialStats = TWeakPtr<FMaterialStats>(SharedThis(this)), PlatformID, QualityLevel]()
+						BuildShaderCodeWidget([MaterialStats = TWeakPtr<FMaterialStats>(SharedThis(this)), PlatformID, QualityLevel]()
 						{
 							auto StatsPtr = MaterialStats.Pin();
 							if (StatsPtr.IsValid())
 							{
 								return StatsPtr->GetShaderCode(PlatformID, QualityLevel, InstanceIndex);
 							}
-
 							return FText::FromString(TEXT("Error reading shader code!"));
 						})
 					];
diff --git a/Engine/Source/Editor/MaterialEditor/Private/MaterialStats.h b/Engine/Source/Editor/MaterialEditor/Private/MaterialStats.h
--- a/Engine/Source/Editor/MaterialEditor/Private/MaterialStats.h
+++ b/Engine/Source/Editor/MaterialEditor/Private/MaterialStats.h
@@ -7,6 +7,8 @@
 #include "UObject/GCObject.h"
 #include "Preferences/MaterialStatsOptions.h"
 
+class FHLSLSyntaxHighlighterMarshaller;
+
 /** structure used to store various statistics extracted from compiled shaders */
 struct FShaderStatsInfo
 {
@@ -354,6 +356,8 @@ class FMaterialStats : public FGCObject, public TSharedFromThis<FMaterialStats>
 
 	TMulticastDelegate<void()> RefreshDependentTabs;
 
+	TSharedPtr<FHLSLSyntaxHighlighterMarshaller> SyntaxHighlighter;
+
 private:
 	/** adds a specified platform in the grid widget for analysis; usually called from BuildShaderPlatformDB() */
 	TSharedPtr<FShaderPlatformSettings> AddShaderPlatform(const EPlatformCategoryType PlatformType, const EShaderPlatform PlatformID, const FName PlatformName, const bool bAllowCodeView, const FString& Description, const bool bAlwaysOn = false);
@@ -365,6 +369,8 @@ private:
 	TSharedRef<class SDockTab> SpawnTab_ShaderCode(const class FSpawnTabArgs& Args, const EShaderPlatform PlatformID, const EMaterialQualityLevel::Type QualityType, const int32 InstanceIndex);
 	TSharedRef<class SDockTab> SpawnTab_HLSLCode(const class FSpawnTabArgs& Args);
 
+	TSharedRef<SScrollBox> BuildShaderCodeWidget(TFunction<FText(void)>&& InShaderCodeCallback);
+
 	/** utility function used to build names for the shader viewing tabs  */
 	static FName MakeTabName(const EPlatformCategoryType PlatformType, const EShaderPlatform ShaderPlatformType, const EMaterialQualityLevel::Type QualityLevel, const int32 InstanceIndex);
 
