Fixed crash in material editor stats, when preview material instances get GCed

diff --git a/Engine/Source/Editor/MaterialEditor/Private/MaterialEditor.cpp b/Engine/Source/Editor/MaterialEditor/Private/MaterialEditor.cpp
--- a/Engine/Source/Editor/MaterialEditor/Private/MaterialEditor.cpp
+++ b/Engine/Source/Editor/MaterialEditor/Private/MaterialEditor.cpp
@@ -52,6 +52,7 @@
 #include "HAL/PlatformApplicationMisc.h"
 #include "MaterialCachedData.h"
 #include "DataDrivenShaderPlatformInfo.h"
+#include "Algo/ForEach.h"
 
 #include "Materials/MaterialExpressionBreakMaterialAttributes.h"
 #include "Materials/MaterialExpressionCollectionParameter.h"
@@ -3474,10 +3475,8 @@ void FMaterialEditor::AddReferencedObjects( FReferenceCollector& Collector )
 {
 	Collector.AddReferencedObject(EditorOptions);
 	Collector.AddReferencedObject(Material);
-	for (auto& DerivedMaterialInstance: DerivedMaterialInstances)
-	{
-		Collector.AddReferencedObject(DerivedMaterialInstance);
-	}
+	Algo::ForEach(DerivedMaterialInstances, [&Collector](TObjectPtr<UMaterialInstance> MaterialInstance){ Collector.AddReferencedObject(MaterialInstance); });
+	Algo::ForEach(OriginalDerivedMaterialInstances, [&Collector](TObjectPtr<UMaterialInstance> MaterialInstance){ Collector.AddReferencedObject(MaterialInstance); });
 	Collector.AddReferencedObject(OriginalMaterial);
 	Collector.AddReferencedObject(MaterialFunction);
 	Collector.AddReferencedObject(ExpressionPreviewMaterial);
