MaterialLayer UI updates

diff --git a/Engine/Source/Editor/MaterialEditor/Private/MaterialEditorInstanceDetailCustomization.cpp b/Engine/Source/Editor/MaterialEditor/Private/MaterialEditorInstanceDetailCustomization.cpp
--- a/Engine/Source/Editor/MaterialEditor/Private/MaterialEditorInstanceDetailCustomization.cpp
+++ b/Engine/Source/Editor/MaterialEditor/Private/MaterialEditorInstanceDetailCustomization.cpp
@@ -21,6 +21,7 @@
 #include "MaterialEditor/DEditorTextureParameterValue.h"
 #include "MaterialEditor/DEditorVectorParameterValue.h"
 #include "MaterialEditor/MaterialEditorInstanceConstant.h"
+#include "SMaterialSubstrateTree.h"
 #include "Materials/MaterialInterface.h"
 #include "Materials/MaterialInstance.h"
 #include "Materials/MaterialExpressionParameter.h"
@@ -56,13 +57,14 @@
 
 
 
-TSharedRef<IDetailCustomization> FMaterialInstanceParameterDetails::MakeInstance(UMaterialEditorInstanceConstant* MaterialInstance, FGetShowHiddenParameters InShowHiddenDelegate)
+TSharedRef<IDetailCustomization> FMaterialInstanceParameterDetails::MakeInstance(UMaterialEditorInstanceConstant* MaterialInstance, SMaterialLayersFunctionsInstanceWrapper* InMaterialLayersFunctionsInstance, FGetShowHiddenParameters InShowHiddenDelegate)
 {
-	return MakeShareable(new FMaterialInstanceParameterDetails(MaterialInstance, InShowHiddenDelegate));
+	return MakeShareable(new FMaterialInstanceParameterDetails(MaterialInstance, InMaterialLayersFunctionsInstance, InShowHiddenDelegate));
 }
 
-FMaterialInstanceParameterDetails::FMaterialInstanceParameterDetails(UMaterialEditorInstanceConstant* MaterialInstance, FGetShowHiddenParameters InShowHiddenDelegate)
+FMaterialInstanceParameterDetails::FMaterialInstanceParameterDetails(UMaterialEditorInstanceConstant* MaterialInstance, SMaterialLayersFunctionsInstanceWrapper* InMaterialLayersFunctionsInstance, FGetShowHiddenParameters InShowHiddenDelegate)
 	: MaterialEditorInstance(MaterialInstance)
+	, MaterialLayersFunctionsInstance(InMaterialLayersFunctionsInstance) 
 	, ShowHiddenDelegate(InShowHiddenDelegate)
 {
 }
@@ -104,32 +106,142 @@ void FMaterialInstanceParameterDetails::CustomizeDetails(IDetailLayoutBuilder& D
 	IDetailCategoryBuilder& GroupsCategory = DetailLayout.EditCategory(GroupsCategoryName, LOCTEXT("MICParamGroupsTitle", "Parameter Groups"));
 	TSharedRef<IPropertyHandle> ParameterGroupsProperty = DetailLayout.GetProperty("ParameterGroups");
 
-	CreateGroupsWidget(ParameterGroupsProperty, GroupsCategory);
-
-	// Create default category for class properties
-	const FName DefaultCategoryName = NAME_None;
-	IDetailCategoryBuilder& DefaultCategory = DetailLayout.EditCategory(DefaultCategoryName);
-	DetailLayout.HideProperty("MaterialLayersParameterValues");
-	if (MaterialEditorInstance->bIsFunctionPreviewMaterial)
+	// check if tree has any selection, we show parameter properties for selected layer only
+	if (this->MaterialLayersFunctionsInstance->NestedTree->GetNumItemsSelected() > 0)
 	{
-		// Customize Parent property so we can check for recursively set parents
-		bool bShowParent = false;
-		if(MaterialEditorInstance->SourceFunction->GetMaterialFunctionUsage() != EMaterialFunctionUsage::Default)
+		// for each selected FSortedParamData item (type stack)
+		for (TSharedPtr<FSortedParamData> SelectedItem : this->MaterialLayersFunctionsInstance->NestedTree->GetSelectedItems())
 		{
-			bShowParent = true;
+			// we go through list of assets
+			for(TSharedPtr<FSortedParamData> ChildAsset : SelectedItem->Children)
+			{
+				// we look for parameter group children
+				for(TSharedPtr<FSortedParamData> GroupParamData : ChildAsset->Children)
+				{
+					if (GroupParamData->StackDataType == EStackDataType::Group)
+					{
+						int32 GroupIdx = MaterialEditorInstance->ParameterGroups.IndexOfByPredicate(
+							[&](const FEditorParameterGroup& Group)
+								{
+									return Group.GroupName == GroupParamData->Group.GroupName;
+								});
+						
+						FEditorParameterGroup& ParameterGroup = GroupParamData->Group;
+						IDetailGroup& DetailGroup = GroupsCategory.AddGroup(ParameterGroup.GroupName, FText::FromName(ParameterGroup.GroupName), false, true);
+						TSharedPtr<IPropertyHandle> GroupPropertyHandle = ParameterGroupsProperty->GetChildHandle(GroupIdx);
+						
+						CreateSingleGroupWidget(ParameterGroup, GroupPropertyHandle, DetailGroup, GroupParamData->ParameterInfo.Index, true);
+						
+						FSimpleDelegate UpdateThumbnails = FSimpleDelegate::CreateLambda([=, this]()
+						{
+							this->MaterialLayersFunctionsInstance->NestedTree->UpdateThumbnailMaterial(ChildAsset->ParameterInfo.Association, ChildAsset->ParameterInfo.Index);
+						});
+						GroupPropertyHandle->SetOnPropertyValueChanged(UpdateThumbnails);
+						GroupPropertyHandle->SetOnChildPropertyValueChanged(UpdateThumbnails);
+						
+					}
+				}
+			}
 		}
-		if (bShowParent)
+		DetailLayout.HideCategory("MaterialEditorInstanceConstant");
+		DetailLayout.HideProperty("Parent");
+		DetailLayout.HideProperty("PhysMaterial");
+		DetailLayout.HideProperty("LightmassSettings");
+		DetailLayout.HideProperty("bUseOldStyleMICEditorGroups");
+		DetailLayout.HideProperty("ParameterGroups");
+		DetailLayout.HideProperty("RefractionDepthBias");
+		DetailLayout.HideProperty("bOverrideSubsurfaceProfile");
+		DetailLayout.HideProperty("SubsurfaceProfile");
+		DetailLayout.HideProperty("BasePropertyOverrides");
+		DetailLayout.HideProperty("MaterialLayersParameterValues");
+	}
+	else
+	{
+		CreateGroupsWidget(ParameterGroupsProperty, GroupsCategory);
+
+		// Create default category for class properties
+		const FName DefaultCategoryName = NAME_None;
+		IDetailCategoryBuilder& DefaultCategory = DetailLayout.EditCategory(DefaultCategoryName);
+		DetailLayout.HideProperty("MaterialLayersParameterValues");
+		if (MaterialEditorInstance->bIsFunctionPreviewMaterial)
 		{
+			// Customize Parent property so we can check for recursively set parents
+			bool bShowParent = false;
+			if(MaterialEditorInstance->SourceFunction->GetMaterialFunctionUsage() != EMaterialFunctionUsage::Default)
+			{
+				bShowParent = true;
+			}
+			if (bShowParent)
+			{
+				TSharedRef<IPropertyHandle> ParentPropertyHandle = DetailLayout.GetProperty("Parent");
+				IDetailPropertyRow& ParentPropertyRow = DefaultCategory.AddProperty(ParentPropertyHandle);
+				ParentPropertyHandle->MarkResetToDefaultCustomized();
+
+				TSharedPtr<SWidget> NameWidget;
+				TSharedPtr<SWidget> ValueWidget;
+				FDetailWidgetRow Row;
+
+				ParentPropertyRow.GetDefaultWidgets(NameWidget, ValueWidget, Row);
+
+				ParentPropertyHandle->ClearResetToDefaultCustomized();
+
+				const bool bShowChildren = true;
+				ParentPropertyRow.CustomWidget(bShowChildren)
+					.NameContent()
+					.MinDesiredWidth(Row.NameWidget.MinWidth)
+					.MaxDesiredWidth(Row.NameWidget.MaxWidth)
+					[
+						NameWidget.ToSharedRef()
+					]
+				.ValueContent()
+					.MinDesiredWidth(Row.ValueWidget.MinWidth)
+					.MaxDesiredWidth(Row.ValueWidget.MaxWidth)
+					[
+						SNew(SObjectPropertyEntryBox)
+						.ObjectPath(this, &FMaterialInstanceParameterDetails::GetFunctionParentPath)
+						.AllowedClass(UMaterialFunctionInterface::StaticClass())
+						.ThumbnailPool(DetailLayout.GetThumbnailPool())
+						.AllowClear(true)
+						.OnObjectChanged(this, &FMaterialInstanceParameterDetails::OnAssetChanged, ParentPropertyHandle)
+						.OnShouldSetAsset(this, &FMaterialInstanceParameterDetails::OnShouldSetAsset)
+						.NewAssetFactories(TArray<UFactory*>())
+					];
+
+				ValueWidget.Reset();
+
+
+			}
+			else
+			{
+				DetailLayout.HideProperty("Parent");
+			}
+
+			DetailLayout.HideProperty("PhysMaterial");
+			DetailLayout.HideProperty("LightmassSettings");
+			DetailLayout.HideProperty("bUseOldStyleMICEditorGroups");
+			DetailLayout.HideProperty("ParameterGroups");
+			DetailLayout.HideProperty("RefractionDepthBias");
+			DetailLayout.HideProperty("bOverrideSubsurfaceProfile");
+			DetailLayout.HideProperty("SubsurfaceProfile");
+			DetailLayout.HideProperty("BasePropertyOverrides");
+		}
+		else
+		{
+			// Add PhysMaterial property
+			DefaultCategory.AddProperty("PhysMaterial");
+
+			// Customize Parent property so we can check for recursively set parents
 			TSharedRef<IPropertyHandle> ParentPropertyHandle = DetailLayout.GetProperty("Parent");
 			IDetailPropertyRow& ParentPropertyRow = DefaultCategory.AddProperty(ParentPropertyHandle);
-			ParentPropertyHandle->MarkResetToDefaultCustomized();
 
+			ParentPropertyHandle->MarkResetToDefaultCustomized();
+	
 			TSharedPtr<SWidget> NameWidget;
 			TSharedPtr<SWidget> ValueWidget;
 			FDetailWidgetRow Row;
 
 			ParentPropertyRow.GetDefaultWidgets(NameWidget, ValueWidget, Row);
-
+	
 			ParentPropertyHandle->ClearResetToDefaultCustomized();
 
 			const bool bShowChildren = true;
@@ -140,150 +252,93 @@ void FMaterialInstanceParameterDetails::CustomizeDetails(IDetailLayoutBuilder& D
 				[
 					NameWidget.ToSharedRef()
 				]
-			.ValueContent()
+				.ValueContent()
 				.MinDesiredWidth(Row.ValueWidget.MinWidth)
 				.MaxDesiredWidth(Row.ValueWidget.MaxWidth)
 				[
 					SNew(SObjectPropertyEntryBox)
-					.ObjectPath(this, &FMaterialInstanceParameterDetails::GetFunctionParentPath)
-					.AllowedClass(UMaterialFunctionInterface::StaticClass())
+					.PropertyHandle(ParentPropertyHandle)
+					.AllowedClass(UMaterialInterface::StaticClass())
 					.ThumbnailPool(DetailLayout.GetThumbnailPool())
 					.AllowClear(true)
-					.OnObjectChanged(this, &FMaterialInstanceParameterDetails::OnAssetChanged, ParentPropertyHandle)
 					.OnShouldSetAsset(this, &FMaterialInstanceParameterDetails::OnShouldSetAsset)
-					.NewAssetFactories(TArray<UFactory*>())
 				];
 
 			ValueWidget.Reset();
 
 
-		}
-		else
-		{
-			DetailLayout.HideProperty("Parent");
-		}
-
-		DetailLayout.HideProperty("PhysMaterial");
-		DetailLayout.HideProperty("LightmassSettings");
-		DetailLayout.HideProperty("bUseOldStyleMICEditorGroups");
-		DetailLayout.HideProperty("ParameterGroups");
-		DetailLayout.HideProperty("RefractionDepthBias");
-		DetailLayout.HideProperty("bOverrideSubsurfaceProfile");
-		DetailLayout.HideProperty("SubsurfaceProfile");
-		DetailLayout.HideProperty("BasePropertyOverrides");
-	}
-	else
-	{
-		// Add PhysMaterial property
-		DefaultCategory.AddProperty("PhysMaterial");
-
-		// Customize Parent property so we can check for recursively set parents
-		TSharedRef<IPropertyHandle> ParentPropertyHandle = DetailLayout.GetProperty("Parent");
-		IDetailPropertyRow& ParentPropertyRow = DefaultCategory.AddProperty(ParentPropertyHandle);
-
-		ParentPropertyHandle->MarkResetToDefaultCustomized();
-	
-		TSharedPtr<SWidget> NameWidget;
-		TSharedPtr<SWidget> ValueWidget;
-		FDetailWidgetRow Row;
-
-		ParentPropertyRow.GetDefaultWidgets(NameWidget, ValueWidget, Row);
-	
-		ParentPropertyHandle->ClearResetToDefaultCustomized();
-
-		const bool bShowChildren = true;
-		ParentPropertyRow.CustomWidget(bShowChildren)
-			.NameContent()
-			.MinDesiredWidth(Row.NameWidget.MinWidth)
-			.MaxDesiredWidth(Row.NameWidget.MaxWidth)
-			[
-				NameWidget.ToSharedRef()
-			]
-			.ValueContent()
-			.MinDesiredWidth(Row.ValueWidget.MinWidth)
-			.MaxDesiredWidth(Row.ValueWidget.MaxWidth)
-			[
-				SNew(SObjectPropertyEntryBox)
-				.PropertyHandle(ParentPropertyHandle)
-				.AllowedClass(UMaterialInterface::StaticClass())
-				.ThumbnailPool(DetailLayout.GetThumbnailPool())
-				.AllowClear(true)
-				.OnShouldSetAsset(this, &FMaterialInstanceParameterDetails::OnShouldSetAsset)
-			];
+			// Add/hide other properties
+			DetailLayout.HideProperty("LightmassSettings");
+			CreateLightmassOverrideWidgets(DetailLayout);
+			DetailLayout.HideProperty("bUseOldStyleMICEditorGroups");
+			DetailLayout.HideProperty("ParameterGroups");
 
-		ValueWidget.Reset();
-
-
-		// Add/hide other properties
-		DetailLayout.HideProperty("LightmassSettings");
-		CreateLightmassOverrideWidgets(DetailLayout);
-		DetailLayout.HideProperty("bUseOldStyleMICEditorGroups");
-		DetailLayout.HideProperty("ParameterGroups");
-
-		{
-			FIsResetToDefaultVisible IsRefractionDepthBiasPropertyResetVisible = FIsResetToDefaultVisible::CreateLambda([this](TSharedPtr<IPropertyHandle> InHandle) {
-				float BiasValue;
-				float ParentBiasValue;
-				return MaterialEditorInstance->SourceInstance->GetRefractionSettings(BiasValue) 
-					&& MaterialEditorInstance->Parent->GetRefractionSettings(ParentBiasValue)
-					&& BiasValue != ParentBiasValue;
-			});
-			FResetToDefaultHandler ResetRefractionDepthBiasPropertyHandler = FResetToDefaultHandler::CreateLambda([this](TSharedPtr<IPropertyHandle> InHandle) {
-				MaterialEditorInstance->Parent->GetRefractionSettings(MaterialEditorInstance->RefractionDepthBias);
-			});
-			FResetToDefaultOverride ResetRefractionDepthBiasPropertyOverride = FResetToDefaultOverride::Create(IsRefractionDepthBiasPropertyResetVisible, ResetRefractionDepthBiasPropertyHandler);
-			IDetailPropertyRow& PropertyRow = DefaultCategory.AddProperty("RefractionDepthBias");
-			PropertyRow.Visibility(TAttribute<EVisibility>::Create(TAttribute<EVisibility>::FGetter::CreateSP(this, &FMaterialInstanceParameterDetails::ShouldShowMaterialRefractionSettings)));
-			PropertyRow.OverrideResetToDefault(ResetRefractionDepthBiasPropertyOverride);
-		}
+			{
+				FIsResetToDefaultVisible IsRefractionDepthBiasPropertyResetVisible = FIsResetToDefaultVisible::CreateLambda([this](TSharedPtr<IPropertyHandle> InHandle) {
+					float BiasValue;
+					float ParentBiasValue;
+					return MaterialEditorInstance->SourceInstance->GetRefractionSettings(BiasValue) 
+						&& MaterialEditorInstance->Parent->GetRefractionSettings(ParentBiasValue)
+						&& BiasValue != ParentBiasValue;
+				});
+				FResetToDefaultHandler ResetRefractionDepthBiasPropertyHandler = FResetToDefaultHandler::CreateLambda([this](TSharedPtr<IPropertyHandle> InHandle) {
+					MaterialEditorInstance->Parent->GetRefractionSettings(MaterialEditorInstance->RefractionDepthBias);
+				});
+				FResetToDefaultOverride ResetRefractionDepthBiasPropertyOverride = FResetToDefaultOverride::Create(IsRefractionDepthBiasPropertyResetVisible, ResetRefractionDepthBiasPropertyHandler);
+				IDetailPropertyRow& PropertyRow = DefaultCategory.AddProperty("RefractionDepthBias");
+				PropertyRow.Visibility(TAttribute<EVisibility>::Create(TAttribute<EVisibility>::FGetter::CreateSP(this, &FMaterialInstanceParameterDetails::ShouldShowMaterialRefractionSettings)));
+				PropertyRow.OverrideResetToDefault(ResetRefractionDepthBiasPropertyOverride);
+			}
 
-		{
-			// Add the material property override group
-			static FName GroupName(TEXT("MaterialPropertyOverrideGroup"));
-			IDetailGroup& MaterialPropertyOverrideGroup = DefaultCategory.AddGroup(GroupName, LOCTEXT("MaterialPropertyOverrideGroup", "Material Property Overrides"), false, false);
+			{
+				// Add the material property override group
+				static FName GroupName(TEXT("MaterialPropertyOverrideGroup"));
+				IDetailGroup& MaterialPropertyOverrideGroup = DefaultCategory.AddGroup(GroupName, LOCTEXT("MaterialPropertyOverrideGroup", "Material Property Overrides"), false, false);
 			
-			// Hide the originals, these will be recreated manually
-			DetailLayout.HideProperty("bOverrideSubsurfaceProfile");
-			DetailLayout.HideProperty("SubsurfaceProfile");
-			DetailLayout.HideProperty("BasePropertyOverrides");
-
-			// Set up the override logic for the subsurface profile
-			TAttribute<bool> IsParamEnabled = TAttribute<bool>::Create(TAttribute<bool>::FGetter::CreateLambda([this](){ return (bool)MaterialEditorInstance->bOverrideSubsurfaceProfile; }));
-
-			IDetailPropertyRow& PropertyRow = MaterialPropertyOverrideGroup.AddPropertyRow(DetailLayout.GetProperty("SubsurfaceProfile"));
-			PropertyRow
-				.EditCondition(IsParamEnabled, 
-					FOnBooleanValueChanged::CreateLambda([this](bool NewValue) {
-						MaterialEditorInstance->bOverrideSubsurfaceProfile = (uint32)NewValue;
-						MaterialEditorInstance->PostEditChange();
-						FEditorSupportDelegates::RedrawAllViewports.Broadcast();
-				}))
-				.Visibility(TAttribute<EVisibility>::Create(TAttribute<EVisibility>::FGetter::CreateSP(this, &FMaterialInstanceParameterDetails::ShouldShowSubsurfaceProfile)));
+				// Hide the originals, these will be recreated manually
+				DetailLayout.HideProperty("bOverrideSubsurfaceProfile");
+				DetailLayout.HideProperty("SubsurfaceProfile");
+				DetailLayout.HideProperty("BasePropertyOverrides");
+
+				// Set up the override logic for the subsurface profile
+				TAttribute<bool> IsParamEnabled = TAttribute<bool>::Create(TAttribute<bool>::FGetter::CreateLambda([this](){ return (bool)MaterialEditorInstance->bOverrideSubsurfaceProfile; }));
+
+				IDetailPropertyRow& PropertyRow = MaterialPropertyOverrideGroup.AddPropertyRow(DetailLayout.GetProperty("SubsurfaceProfile"));
+				PropertyRow
+					.EditCondition(IsParamEnabled, 
+						FOnBooleanValueChanged::CreateLambda([this](bool NewValue) {
+							MaterialEditorInstance->bOverrideSubsurfaceProfile = (uint32)NewValue;
+							MaterialEditorInstance->PostEditChange();
+							FEditorSupportDelegates::RedrawAllViewports.Broadcast();
+					}))
+					.Visibility(TAttribute<EVisibility>::Create(TAttribute<EVisibility>::FGetter::CreateSP(this, &FMaterialInstanceParameterDetails::ShouldShowSubsurfaceProfile)));
 			
-			// Append the base property overrides to the Material Property Override Group
-			CreateBasePropertyOverrideWidgets(DetailLayout, MaterialPropertyOverrideGroup);
+				// Append the base property overrides to the Material Property Override Group
+				CreateBasePropertyOverrideWidgets(DetailLayout, MaterialPropertyOverrideGroup);
 
-			// Append the nanite material override.
-			MaterialPropertyOverrideGroup.AddPropertyRow(DetailLayout.GetProperty("NaniteOverrideMaterial"));
+				// Append the nanite material override.
+				MaterialPropertyOverrideGroup.AddPropertyRow(DetailLayout.GetProperty("NaniteOverrideMaterial"));
+			}
 		}
-	}
 
-	// Add the preview mesh property directly from the material instance 
-	FName PreviewingCategoryName = TEXT("Previewing");
-	IDetailCategoryBuilder& PreviewingCategory = DetailLayout.EditCategory(PreviewingCategoryName, LOCTEXT("MICPreviewingCategoryTitle", "Previewing"));
+		// Add the preview mesh property directly from the material instance 
+		FName PreviewingCategoryName = TEXT("Previewing");
+		IDetailCategoryBuilder& PreviewingCategory = DetailLayout.EditCategory(PreviewingCategoryName, LOCTEXT("MICPreviewingCategoryTitle", "Previewing"));
 
-	TArray<UObject*> ExternalObjects;
-	ExternalObjects.Add(MaterialEditorInstance->SourceInstance);
+		TArray<UObject*> ExternalObjects;
+		ExternalObjects.Add(MaterialEditorInstance->SourceInstance);
 
-	PreviewingCategory.AddExternalObjectProperty(ExternalObjects, TEXT("PreviewMesh"));
+		PreviewingCategory.AddExternalObjectProperty(ExternalObjects, TEXT("PreviewMesh"));
 
-	DefaultCategory.AddExternalObjectProperty(ExternalObjects, TEXT("AssetUserData"), EPropertyLocation::Advanced);
+		DefaultCategory.AddExternalObjectProperty(ExternalObjects, TEXT("AssetUserData"), EPropertyLocation::Advanced);
+	}
 }
 
 void FMaterialInstanceParameterDetails::CreateGroupsWidget(TSharedRef<IPropertyHandle> ParameterGroupsProperty, IDetailCategoryBuilder& GroupsCategory)
 {
 	bool bShowSaveButtons = false;
 	check(MaterialEditorInstance);
+	
 	for (int32 GroupIdx = 0; GroupIdx < MaterialEditorInstance->ParameterGroups.Num(); ++GroupIdx)
 	{
 		FEditorParameterGroup& ParameterGroup = MaterialEditorInstance->ParameterGroups[GroupIdx];
@@ -298,7 +353,7 @@ void FMaterialInstanceParameterDetails::CreateGroupsWidget(TSharedRef<IPropertyH
 				const bool bIsVisible = MaterialEditorInstance->VisibleExpressions.Contains(Parameter->ParameterInfo);
 				bCreateGroup = bIsVisible && (!MaterialEditorInstance->bShowOnlyOverrides || FMaterialPropertyHelpers::IsOverriddenExpression(Parameter));
 			}
-			
+		
 			if (bCreateGroup)
 			{
 				IDetailGroup& DetailGroup = GroupsCategory.AddGroup(ParameterGroup.GroupName, FText::FromName(ParameterGroup.GroupName), false, false);
@@ -339,6 +394,7 @@ void FMaterialInstanceParameterDetails::CreateGroupsWidget(TSharedRef<IPropertyH
 			}
 		}
 	}
+	
 	if (bShowSaveButtons)
 	{
 		FDetailWidgetRow& SaveInstanceRow = GroupsCategory.AddCustomRow(LOCTEXT("SaveInstances", "Save Instances"));
@@ -401,7 +457,7 @@ void FMaterialInstanceParameterDetails::EnableGroupParameters(FEditorParameterGr
 	}
 }
 
-void FMaterialInstanceParameterDetails::CreateSingleGroupWidget(FEditorParameterGroup& ParameterGroup, TSharedPtr<IPropertyHandle> ParameterGroupProperty, IDetailGroup& DetailGroup )
+void FMaterialInstanceParameterDetails::CreateSingleGroupWidget(FEditorParameterGroup& ParameterGroup, TSharedPtr<IPropertyHandle> ParameterGroupProperty, IDetailGroup& DetailGroup, int32 GroupIndex /*= -1*/, bool bForceShowParam /*= false*/)
 {
 	TSharedPtr<IPropertyHandle> ParametersArrayProperty = ParameterGroupProperty->GetChildHandle("Parameters");
 
@@ -409,77 +465,80 @@ void FMaterialInstanceParameterDetails::CreateSingleGroupWidget(FEditorParameter
 	for (int32 ParamIdx = 0; ParamIdx < ParameterGroup.Parameters.Num(); ++ParamIdx)
 	{
 		TSharedPtr<IPropertyHandle> ParameterProperty = ParametersArrayProperty->GetChildHandle(ParamIdx);
-
 		UDEditorParameterValue* Parameter = ParameterGroup.Parameters[ParamIdx];
-		UDEditorFontParameterValue* FontParam = Cast<UDEditorFontParameterValue>(Parameter);
-		UDEditorMaterialLayersParameterValue* LayersParam = Cast<UDEditorMaterialLayersParameterValue>(Parameter);
-		UDEditorScalarParameterValue* ScalarParam = Cast<UDEditorScalarParameterValue>(Parameter);
-		UDEditorStaticComponentMaskParameterValue* CompMaskParam = Cast<UDEditorStaticComponentMaskParameterValue>(Parameter);
-		UDEditorStaticSwitchParameterValue* SwitchParam = Cast<UDEditorStaticSwitchParameterValue>(Parameter);
-		UDEditorTextureParameterValue* TextureParam = Cast<UDEditorTextureParameterValue>(Parameter);
-		UDEditorRuntimeVirtualTextureParameterValue* RuntimeVirtualTextureParam = Cast<UDEditorRuntimeVirtualTextureParameterValue>(Parameter);
-		UDEditorSparseVolumeTextureParameterValue* SparseVolumeTextureParam = Cast<UDEditorSparseVolumeTextureParameterValue>(Parameter);
-		UDEditorVectorParameterValue* VectorParam = Cast<UDEditorVectorParameterValue>(Parameter);
-
-		if (Parameter->ParameterInfo.Association == EMaterialParameterAssociation::GlobalParameter)
+		if (ParameterProperty.IsValid() &&
+			(GroupIndex == INDEX_NONE || Parameter->ParameterInfo.Index == GroupIndex))
 		{
-			if (VectorParam && VectorParam->bIsUsedAsChannelMask)
-			{
-				CreateVectorChannelMaskParameterValueWidget(Parameter, ParameterProperty, DetailGroup);
-			}
-			if (ScalarParam && ScalarParam->AtlasData.bIsUsedAsAtlasPosition)
+			UDEditorFontParameterValue* FontParam = Cast<UDEditorFontParameterValue>(Parameter);
+			UDEditorMaterialLayersParameterValue* LayersParam = Cast<UDEditorMaterialLayersParameterValue>(Parameter);
+			UDEditorScalarParameterValue* ScalarParam = Cast<UDEditorScalarParameterValue>(Parameter);
+			UDEditorStaticComponentMaskParameterValue* CompMaskParam = Cast<UDEditorStaticComponentMaskParameterValue>(Parameter);
+			UDEditorStaticSwitchParameterValue* SwitchParam = Cast<UDEditorStaticSwitchParameterValue>(Parameter);
+			UDEditorTextureParameterValue* TextureParam = Cast<UDEditorTextureParameterValue>(Parameter);
+			UDEditorRuntimeVirtualTextureParameterValue* RuntimeVirtualTextureParam = Cast<UDEditorRuntimeVirtualTextureParameterValue>(Parameter);
+			UDEditorSparseVolumeTextureParameterValue* SparseVolumeTextureParam = Cast<UDEditorSparseVolumeTextureParameterValue>(Parameter);
+			UDEditorVectorParameterValue* VectorParam = Cast<UDEditorVectorParameterValue>(Parameter);
+
+			if (Parameter->ParameterInfo.Association == EMaterialParameterAssociation::GlobalParameter || bForceShowParam)
 			{
-				CreateScalarAtlasPositionParameterValueWidget(Parameter, ParameterProperty, DetailGroup);
-			}
-			if (TextureParam && 
-				( !TextureParam->ChannelNames.R.IsEmpty()
-				|| !TextureParam->ChannelNames.G.IsEmpty()
-				|| !TextureParam->ChannelNames.B.IsEmpty()
-				|| !TextureParam->ChannelNames.A.IsEmpty()))
-			{
-				CreateLabeledTextureParameterValueWidget(Parameter, ParameterProperty, DetailGroup);
-			}
-			else if (LayersParam)
-			{
-			}
-			else if (CompMaskParam)
-			{
-				CreateMaskParameterValueWidget(Parameter, ParameterProperty, DetailGroup);
-			}
-			else
-			{
-				if (ScalarParam && ScalarParam->SliderMax > ScalarParam->SliderMin)
+				if (VectorParam && VectorParam->bIsUsedAsChannelMask)
 				{
-					TSharedPtr<IPropertyHandle> ParameterValueProperty = ParameterProperty->GetChildHandle("ParameterValue");
-					ParameterValueProperty->SetInstanceMetaData("UIMin", FString::Printf(TEXT("%f"), ScalarParam->SliderMin));
-					ParameterValueProperty->SetInstanceMetaData("UIMax", FString::Printf(TEXT("%f"), ScalarParam->SliderMax));
+					CreateVectorChannelMaskParameterValueWidget(Parameter, ParameterProperty, DetailGroup);
 				}
-
-				if (VectorParam)
+				if (ScalarParam && ScalarParam->AtlasData.bIsUsedAsAtlasPosition)
 				{
-					static const FName Red("R");
-					static const FName Green("G");
-					static const FName Blue("B");
-					static const FName Alpha("A");
-					if (!VectorParam->ChannelNames.R.IsEmpty())
-					{
-						ParameterProperty->GetChildHandle(Red)->SetPropertyDisplayName(VectorParam->ChannelNames.R);
-					}
-					if (!VectorParam->ChannelNames.G.IsEmpty())
-					{
-						ParameterProperty->GetChildHandle(Green)->SetPropertyDisplayName(VectorParam->ChannelNames.G);
-					}
-					if (!VectorParam->ChannelNames.B.IsEmpty())
+					CreateScalarAtlasPositionParameterValueWidget(Parameter, ParameterProperty, DetailGroup);
+				}
+				if (TextureParam && 
+					( !TextureParam->ChannelNames.R.IsEmpty()
+					|| !TextureParam->ChannelNames.G.IsEmpty()
+					|| !TextureParam->ChannelNames.B.IsEmpty()
+					|| !TextureParam->ChannelNames.A.IsEmpty()))
+				{
+					CreateLabeledTextureParameterValueWidget(Parameter, ParameterProperty, DetailGroup);
+				}
+				else if (LayersParam)
+				{
+				}
+				else if (CompMaskParam)
+				{
+					CreateMaskParameterValueWidget(Parameter, ParameterProperty, DetailGroup);
+				}
+				else
+				{
+					if (ScalarParam && ScalarParam->SliderMax > ScalarParam->SliderMin)
 					{
-						ParameterProperty->GetChildHandle(Blue)->SetPropertyDisplayName(VectorParam->ChannelNames.B);
+						TSharedPtr<IPropertyHandle> ParameterValueProperty = ParameterProperty->GetChildHandle("ParameterValue");
+						ParameterValueProperty->SetInstanceMetaData("UIMin", FString::Printf(TEXT("%f"), ScalarParam->SliderMin));
+						ParameterValueProperty->SetInstanceMetaData("UIMax", FString::Printf(TEXT("%f"), ScalarParam->SliderMax));
 					}
-					if (!VectorParam->ChannelNames.A.IsEmpty())
+
+					if (VectorParam)
 					{
-						ParameterProperty->GetChildHandle(Alpha)->SetPropertyDisplayName(VectorParam->ChannelNames.A);
+						static const FName Red("R");
+						static const FName Green("G");
+						static const FName Blue("B");
+						static const FName Alpha("A");
+						if (!VectorParam->ChannelNames.R.IsEmpty())
+						{
+							ParameterProperty->GetChildHandle(Red)->SetPropertyDisplayName(VectorParam->ChannelNames.R);
+						}
+						if (!VectorParam->ChannelNames.G.IsEmpty())
+						{
+							ParameterProperty->GetChildHandle(Green)->SetPropertyDisplayName(VectorParam->ChannelNames.G);
+						}
+						if (!VectorParam->ChannelNames.B.IsEmpty())
+						{
+							ParameterProperty->GetChildHandle(Blue)->SetPropertyDisplayName(VectorParam->ChannelNames.B);
+						}
+						if (!VectorParam->ChannelNames.A.IsEmpty())
+						{
+							ParameterProperty->GetChildHandle(Alpha)->SetPropertyDisplayName(VectorParam->ChannelNames.A);
+						}
 					}
-				}
 
-				CreateParameterValueWidget(Parameter, ParameterProperty, DetailGroup);
+					CreateParameterValueWidget(Parameter, ParameterProperty, DetailGroup);
+				}
 			}
 		}
 	}
diff --git a/Engine/Source/Editor/MaterialEditor/Private/MaterialEditorInstanceDetailCustomization.h b/Engine/Source/Editor/MaterialEditor/Private/MaterialEditorInstanceDetailCustomization.h
--- a/Engine/Source/Editor/MaterialEditor/Private/MaterialEditorInstanceDetailCustomization.h
+++ b/Engine/Source/Editor/MaterialEditor/Private/MaterialEditorInstanceDetailCustomization.h
@@ -6,6 +6,7 @@
 #include "Types/SlateEnums.h"
 #include "Layout/Visibility.h"
 #include "IDetailCustomization.h"
+#include "SMaterialLayersFunctionsTree.h"
 #include "Input/Reply.h"
 #include "Customizations/ColorStructCustomization.h"
 
@@ -26,10 +27,10 @@ class FMaterialInstanceParameterDetails : public IDetailCustomization
 {
 public:
 	/** Makes a new instance of this detail layout class for a specific detail view requesting it */
-	static TSharedRef<class IDetailCustomization> MakeInstance(UMaterialEditorInstanceConstant* MaterialInstance, FGetShowHiddenParameters InShowHiddenDelegate);
+	static TSharedRef<class IDetailCustomization> MakeInstance(UMaterialEditorInstanceConstant* MaterialInstance, SMaterialLayersFunctionsInstanceWrapper* MaterialLayersFunctionsInstance, FGetShowHiddenParameters InShowHiddenDelegate);
 	
 	/** Constructor */
-	FMaterialInstanceParameterDetails(UMaterialEditorInstanceConstant* MaterialInstance, FGetShowHiddenParameters InShowHiddenDelegate);
+	FMaterialInstanceParameterDetails(UMaterialEditorInstanceConstant* MaterialInstance, SMaterialLayersFunctionsInstanceWrapper* MaterialLayersFunctionsInstance, FGetShowHiddenParameters InShowHiddenDelegate);
 
 	/** IDetailCustomization interface */
 	virtual void CustomizeDetails(IDetailLayoutBuilder& DetailLayout) override;
@@ -45,7 +46,7 @@ private:
 	void CreateGroupsWidget(TSharedRef<IPropertyHandle> ParameterGroupsProperty, class IDetailCategoryBuilder& GroupsCategory);
 
 	/** Builds the widget for an individual parameter group */
-	void CreateSingleGroupWidget(struct FEditorParameterGroup& ParameterGroup, TSharedPtr<IPropertyHandle> ParameterGroupProperty, class IDetailGroup& DetailGroup);
+	void CreateSingleGroupWidget(struct FEditorParameterGroup& ParameterGroup, TSharedPtr<IPropertyHandle> ParameterGroupProperty, class IDetailGroup& DetailGroup, int32 GroupIndex = -1, bool bForceShowParam = false);
 
 	/** Enable/Disable all parameter properties in a group */
 	static void EnableGroupParameters(struct FEditorParameterGroup& ParameterGroup, bool ShouldEnable);
@@ -135,6 +136,8 @@ private:
 	/** Object that stores all of the possible parameters we can edit */
 	UMaterialEditorInstanceConstant* MaterialEditorInstance;
 
+	SMaterialLayersFunctionsInstanceWrapper* MaterialLayersFunctionsInstance;
+	
 	/** Delegate to call to determine if hidden parameters should be shown */
 	FGetShowHiddenParameters ShowHiddenDelegate;
 
diff --git a/Engine/Source/Editor/MaterialEditor/Private/MaterialEditorModule.cpp b/Engine/Source/Editor/MaterialEditor/Private/MaterialEditorModule.cpp
--- a/Engine/Source/Editor/MaterialEditor/Private/MaterialEditorModule.cpp
+++ b/Engine/Source/Editor/MaterialEditor/Private/MaterialEditorModule.cpp
@@ -23,6 +23,7 @@
 #include "ISettingsModule.h"
 #include "Interfaces/IMainFrameModule.h"
 #include "MaterialEditorGraphPanelPinFactory.h"
+#include "MaterialEditorStyle.h"
 
 const FName MaterialEditorAppIdentifier = FName(TEXT("MaterialEditorApp"));
 const FName MaterialInstanceEditorAppIdentifier = FName(TEXT("MaterialInstanceEditorApp"));
@@ -67,6 +68,8 @@ public:
 
 		GraphPanelPinFactory = MakeShared<FMaterialEditorGraphPanelPinFactory>();
 		FEdGraphUtilities::RegisterVisualPinFactory(GraphPanelPinFactory);
+
+		FSubstrateMaterialEditorStyle::Initialize();
 	}
 
 	/**
@@ -78,6 +81,8 @@ public:
 
 		MenuExtensibilityManager.Reset();
 		ToolBarExtensibilityManager.Reset();
+		
+		FSubstrateMaterialEditorStyle::Shutdown();
 	}
 
 	/**
diff --git a/Engine/Source/Editor/MaterialEditor/Private/MaterialEditorStyle.cpp b/Engine/Source/Editor/MaterialEditor/Private/MaterialEditorStyle.cpp
new file mode 100644
--- /dev/null
+++ b/Engine/Source/Editor/MaterialEditor/Private/MaterialEditorStyle.cpp
@@ -0,0 +1,364 @@
+// Copyright Epic Games, Inc. All Rights Reserved.
+
+#include "MaterialEditorStyle.h"
+#include "Brushes/SlateImageBrush.h"
+#include "DetailLayoutBuilder.h"
+#include "MaterialEditorModule.h"
+#include "Engine/Texture2D.h"
+#include "Interfaces/IPluginManager.h"
+#include "Misc/Paths.h"
+#include "Styling/SlateStyleRegistry.h"
+#include "Styling/StyleColors.h"
+
+TSharedPtr<FSlateStyleSet> FSubstrateMaterialEditorStyle::StyleInstance = nullptr;
+
+void FSubstrateMaterialEditorStyle::Initialize()
+{
+	if (!StyleInstance.IsValid())
+	{
+		StyleInstance = Create();
+		FSlateStyleRegistry::RegisterSlateStyle(*StyleInstance);
+	}
+}
+
+void FSubstrateMaterialEditorStyle::Shutdown()
+{
+	if (StyleInstance.IsValid())
+	{
+		FSlateStyleRegistry::UnRegisterSlateStyle(*StyleInstance);
+		ensure(StyleInstance.IsUnique());
+		StyleInstance.Reset();
+	}
+}
+
+FName FSubstrateMaterialEditorStyle::GetStyleSetName()
+{
+	static FName StyleSetName(TEXT("SubstrateMaterialEditorStyle"));
+	return StyleSetName;
+}
+
+const ISlateStyle& FSubstrateMaterialEditorStyle::Get()
+{
+	if (!StyleInstance.IsValid())
+	{
+		Initialize();
+	}
+	return *StyleInstance;
+}
+
+FLinearColor FSubstrateMaterialEditorStyle::GetColor(const FName& InName)
+{
+	return FSubstrateMaterialEditorStyle::Get().GetColor(InName);
+}
+
+const FSlateBrush* FSubstrateMaterialEditorStyle::GetBrush(const FName& InName)
+{
+	return FSubstrateMaterialEditorStyle::Get().GetBrush(InName);
+}
+
+#define IMAGE_BRUSH(RelativePath, ... ) FSlateImageBrush(Style->RootToContentDir(RelativePath, TEXT(".png")), __VA_ARGS__)
+#define BOX_BRUSH(RelativePath, ... ) FSlateBoxBrush(Style->RootToContentDir(RelativePath, TEXT(".png")), __VA_ARGS__)
+#define BORDER_BRUSH(RelativePath, ... ) FSlateBorderBrush(Style->RootToContentDir(RelativePath, TEXT(".png")), __VA_ARGS__)
+
+#define IMAGE_BRUSH_SVG(RelativePath, ...) FSlateVectorImageBrush(Style->RootToContentDir(RelativePath, TEXT(".svg")), __VA_ARGS__)
+#define BOX_BRUSH_SVG(RelativePath, ...) FSlateVectorBoxBrush(Style->RootToContentDir(RelativePath, TEXT(".svg")), __VA_ARGS__)
+#define BORDER_BRUSH_SVG(RelativePath, ...) FSlateVectorBorderBrush(Style->RootToContentDir(RelativePath, TEXT(".svg")), __VA_ARGS__)
+
+#define DEFAULT_FONT(...) FCoreStyle::GetDefaultFontStyle(__VA_ARGS__)
+
+#define IMAGE_PLUGIN_BRUSH(RelativePath, ...) FSlateImageBrush(FSubstrateMaterialEditorStyle::InContent(RelativePath, ".png"), __VA_ARGS__)
+#define IMAGE_PLUGIN_BRUSH_SVG(RelativePath, ...) FSlateVectorImageBrush(FSubstrateMaterialEditorStyle::InContent(RelativePath, ".svg"), __VA_ARGS__)
+
+const FVector2D Icon8x8(8.f, 8.f);
+const FVector2D Icon12x12(12.f, 12.f);
+const FVector2D Icon16x16(16.f, 16.f);
+const FVector2D Icon20x20(20.f, 20.f);
+const FVector2D Icon24x24(24.f, 24.f);
+const FVector2D Icon32x32(32.f, 32.f);
+const FVector2D Icon40x40(40.f, 40.f);
+
+FLinearColor ReplaceColorAlpha(const FLinearColor& InColor, const float InNewAlpha)
+{
+	FLinearColor OutColor(InColor);
+	OutColor.A = InNewAlpha;
+	return OutColor;
+}
+TSharedRef<FSlateStyleSet> FSubstrateMaterialEditorStyle::Create()
+{
+	TSharedRef<FSlateStyleSet> Style = MakeShared<FSlateStyleSet>(TEXT("SubstrateMaterial"));
+
+	Style->SetContentRoot(FPaths::EngineContentDir() / TEXT("Editor/Slate"));
+	Style->SetCoreContentRoot(FPaths::EngineContentDir() / TEXT("Slate"));
+	////////////////////////////////////////////////////////////////////////////////////////////////////
+	// Color Styles
+	const FLinearColor EngineSelectColor = FStyleColors::Primary.GetSpecifiedColor();
+	const FLinearColor EngineSelectHoverColor = FStyleColors::PrimaryHover.GetSpecifiedColor();
+	const FLinearColor EngineSelectPressColor = FStyleColors::PrimaryPress.GetSpecifiedColor();
+
+	const FLinearColor SelectColor = ReplaceColorAlpha(FStyleColors::Select.GetSpecifiedColor(), 0.9f);
+	const FLinearColor SelectHoverColor = FStyleColors::Select.GetSpecifiedColor();
+	const FLinearColor SelectPressColor = EngineSelectPressColor;
+
+	Style->Set("Color.Select", SelectColor);
+	Style->Set("Color.Select.Hover", SelectHoverColor);
+	Style->Set("Color.Select.Press", SelectPressColor);
+
+	////////////////////////////////////////////////////////////////////////////////////////////////////
+	// Brush Styles
+	Style->Set("Icons.Menu.Dropdown", new IMAGE_BRUSH_SVG("Icons/MenuDropdown", Icon16x16));
+	
+	Style->Set("Icons.Type.None", new IMAGE_BRUSH("Icons/ValueTypes/None", Icon12x12));
+	Style->Set("Icons.Type.Bool", new IMAGE_BRUSH("Icons/ValueTypes/Bool", Icon12x12));
+	Style->Set("Icons.Type.Float1", new IMAGE_BRUSH("Icons/ValueTypes/Float1", Icon12x12));
+	Style->Set("Icons.Type.Float2", new IMAGE_BRUSH("Icons/ValueTypes/Float2", Icon12x12));
+	Style->Set("Icons.Type.Float3_RPY", new IMAGE_BRUSH("Icons/ValueTypes/Float3_RPY", Icon12x12));
+	Style->Set("Icons.Type.Float3_RGB", new IMAGE_BRUSH("Icons/ValueTypes/Float3_RGB", Icon12x12));
+	Style->Set("Icons.Type.Float3_XYZ", new IMAGE_BRUSH("Icons/ValueTypes/Float3_XYZ", Icon12x12));
+	Style->Set("Icons.Type.Float4_RGBA", new IMAGE_BRUSH("Icons/ValueTypes/Float4_RGBA", Icon12x12));
+	Style->Set("Icons.Type.Float_Any", new IMAGE_BRUSH("Icons/ValueTypes/Float_Any", Icon12x12));
+	Style->Set("Icons.Type.Texture", new IMAGE_BRUSH("Icons/ValueTypes/Texture", Icon12x12));
+	
+	Style->Set("Icons.Material.DefaultLit", new IMAGE_BRUSH("Icons/EditorIcons/MaterialTypeDefaultLit", Icon32x32));
+	Style->Set("Icons.Material.Unlit", new IMAGE_BRUSH("Icons/EditorIcons/MaterialTypeUnlit", Icon32x32));
+	
+	Style->Set("Icons.Lock", new IMAGE_BRUSH_SVG("Icons/EditorIcons/Lock", Icon16x16));
+	Style->Set("Icons.Unlock", new IMAGE_BRUSH_SVG("Icons/EditorIcons/Unlock", Icon16x16));
+	
+	Style->Set("Icons.Remove", new IMAGE_BRUSH("Icons/EditorIcons/Remove_16px", Icon16x16));
+	
+	Style->Set("Icons.Normalize", new IMAGE_BRUSH_SVG("Icons/EditorIcons/Normalize", Icon16x16));
+	
+	Style->Set("Icons.Stage.EnabledButton", new IMAGE_BRUSH("Icons/EditorIcons/WhiteBall", Icon8x8));
+	Style->Set("Icons.Stage.BaseToggle", new IMAGE_BRUSH("Icons/EditorIcons/BaseToggle_16x", Icon16x16));
+	Style->Set("Icons.Stage.MaskToggle", new IMAGE_BRUSH("Icons/EditorIcons/MaskToggle_16x", Icon16x16));
+	Style->Set("Icons.Stage.Enabled", new IMAGE_BRUSH_SVG("Icons/EditorIcons/Enable", Icon24x24));
+	Style->Set("Icons.Stage.Disabled", new IMAGE_BRUSH_SVG("Icons/EditorIcons/Disable", Icon24x24));
+	
+	Style->Set("Icons.Stage.ChainLinked", new IMAGE_BRUSH_SVG("Icons/EditorIcons/ChainLinked", Icon16x16));
+	Style->Set("Icons.Stage.ChainUnlinked", new IMAGE_BRUSH_SVG("Icons/EditorIcons/ChainUnlinked", Icon16x16));
+	Style->Set("Icons.Stage.ChainLinked.Horizontal", new IMAGE_BRUSH_SVG("Icons/EditorIcons/ChainLinked_Horizontal", Icon24x24));
+	Style->Set("Icons.Stage.ChainUnlinked.Horizontal", new IMAGE_BRUSH_SVG("Icons/EditorIcons/ChainUnlinked_Horizontal", Icon24x24));
+	Style->Set("Icons.Stage.ChainLinked.Vertical", new IMAGE_BRUSH_SVG("Icons/EditorIcons/ChainLinked_Vertical", Icon24x24));
+	Style->Set("Icons.Stage.ChainUnlinked.Vertical", new IMAGE_BRUSH_SVG("Icons/EditorIcons/ChainUnlinked_Vertical", Icon24x24));
+	
+	Style->Set("ImageBorder", new FSlateRoundedBoxBrush(
+		FLinearColor::Transparent, 10.0f, 
+		FStyleColors::Panel.GetSpecifiedColor(), 2.0f));
+	
+	Style->Set("Border.SinglePixel", new BORDER_BRUSH(TEXT("Images/Borders/Border_SinglePixel"), FMargin(1.0f / 4.0f)));
+	Style->Set("Border.LeftTopRight", new BORDER_BRUSH(TEXT("Images/Borders/Border_LeftTopRight"), FMargin(1.0f / 4.0f, 1.0f / 2.0f)));
+	Style->Set("Border.LeftBottomRight", new BORDER_BRUSH(TEXT("Images/Borders/Border_LeftBottomRight"), FMargin(1.0f / 4.0f, 1.0f / 2.0f)));
+	Style->Set("Border.TopLeftBottom", new BORDER_BRUSH(TEXT("Images/Borders/Border_TopLeftBottom"), FMargin(1.0f / 2.0f, 1.0f / 4.0f)));
+	Style->Set("Border.TopRightBottom", new BORDER_BRUSH(TEXT("Images/Borders/Border_TopRightBottom"), FMargin(1.0f / 2.0f, 1.0f / 4.0f)));
+	Style->Set("Border.Top", new BORDER_BRUSH(TEXT("Images/Borders/Border_Top"), FMargin(0.0f, 1.0f / 2.0f)));
+	Style->Set("Border.Bottom", new BORDER_BRUSH(TEXT("Images/Borders/Border_Bottom"), FMargin(0.0f, 1.0f / 2.0f)));
+	Style->Set("Border.Left", new BORDER_BRUSH(TEXT("Images/Borders/Border_Left"), FMargin(1.0f / 2.0f, 0.0f)));
+	Style->Set("Border.Right", new BORDER_BRUSH(TEXT("Images/Borders/Border_Right"), FMargin(1.0f / 2.0f, 0.0f)));
+
+	////////////////////////////////////////////////////////////////////////////////////////////////////
+	// Button Styles
+	Style->Set("HoverHintOnly", FButtonStyle()
+		.SetNormal(FSlateNoResource())
+		.SetHovered(FSlateRoundedBoxBrush(FLinearColor(1, 1, 1, 0.15f), 4.0f))
+		.SetPressed(FSlateRoundedBoxBrush(FLinearColor(1, 1, 1, 0.25f), 4.0f))
+		.SetNormalPadding(FMargin(0, 0, 0, 1))
+		.SetPressedPadding(FMargin(0, 1, 0, 0)));
+	
+	Style->Set("HoverHintOnly.Bordered", FButtonStyle()
+		.SetNormal(FSlateRoundedBoxBrush(FLinearColor::Transparent, 4.0f, FLinearColor(1, 1, 1, 0.25f), 1.0f))
+		.SetHovered(FSlateRoundedBoxBrush(FLinearColor(1, 1, 1, 0.15f), 4.0f, FLinearColor(1, 1, 1, 0.4f), 1.0f))
+		.SetPressed(FSlateRoundedBoxBrush(FLinearColor(1, 1, 1, 0.25f), 4.0f, FLinearColor(1, 1, 1, 0.5f), 1.0f))
+		.SetNormalPadding(FMargin(0, 0, 0, 1))
+		.SetPressedPadding(FMargin(0, 1, 0, 0)));
+	
+	Style->Set("HoverHintOnly.Bordered.Dark", FButtonStyle()
+		.SetNormal(FSlateRoundedBoxBrush(FLinearColor::Transparent, 4.0f, FStyleColors::InputOutline.GetSpecifiedColor(), 1.0f))
+		.SetHovered(FSlateRoundedBoxBrush(FLinearColor(1, 1, 1, 0.15f), 4.0f, FLinearColor(1, 1, 1, 0.4f), 1.0f))
+		.SetPressed(FSlateRoundedBoxBrush(FLinearColor(1, 1, 1, 0.25f), 4.0f, FLinearColor(1, 1, 1, 0.5f), 1.0f))
+		.SetNormalPadding(FMargin(0, 0, 0, 1))
+		.SetPressedPadding(FMargin(0, 1, 0, 0)));
+	
+	SetupLayerViewStyles(Style);
+	SetupTextStyles(Style);
+	
+	//////////////////////////////////////////////////////
+	/// Editable TextBox Style
+	FEditableTextBoxStyle InlineEditableTextBoxStyle;
+	InlineEditableTextBoxStyle.SetPadding(FMargin(0));
+	InlineEditableTextBoxStyle.SetBackgroundColor(FSlateColor(FLinearColor::Transparent));
+	
+	Style->Set("InlineEditableTextBoxStyle", InlineEditableTextBoxStyle);
+
+	return Style;
+}
+
+void FSubstrateMaterialEditorStyle::SetupLayerViewStyles(const TSharedRef<FSlateStyleSet>& Style)
+{
+	Style->Set("LayerView.Background", new FSlateRoundedBoxBrush(
+		FStyleColors::Panel.GetSpecifiedColor()/*FLinearColor(0, 0, 0, 0.25f)*/, 6.0f,
+		FStyleColors::Header.GetSpecifiedColor()/*FLinearColor(1, 1, 1, 0.2f)*/, 1.0f));
+
+	Style->Set("LayerView.Details.Background", new FSlateRoundedBoxBrush(
+		FStyleColors::Recessed.GetSpecifiedColor()/*FLinearColor(0, 0, 0, 0.25f)*/, 6.0f,
+		FStyleColors::Header.GetSpecifiedColor()/*FLinearColor(1, 1, 1, 0.2f)*/, 1.0f));
+
+	/**
+	 * SListView and and FTableViewStyle have no support for adding padding between the background brush
+	 * and the SListView widget, so we are not using this style for the SDMSlotLayerView. Instead, we add a
+	 * SBorder around the SDMBLayerView and style that.
+	 */
+	Style->Set("LayerView", FTableViewStyle()
+		.SetBackgroundBrush(*Style->GetBrush("LayerView.Background"))
+	);
+
+	const float LayerViewItemCornerRadius = 5.0f;
+	const float LayerViewItemBorderWidth = 1.0f;
+
+	const FLinearColor LayerViewItemFillColor = FLinearColor::Transparent;
+	constexpr FLinearColor LayerViewItemBorderColor = FLinearColor(1, 1, 1, 0.15f);
+
+	const FLinearColor LayerItemHoverFillColor = FStyleColors::Recessed.GetSpecifiedColor();
+	constexpr FLinearColor LayerItemHoverBorderColor = FLinearColor(1, 1, 1, 0.2f);
+
+	const FLinearColor LayerItemSelectFillColor = FStyleColors::Header.GetSpecifiedColor();
+	const FLinearColor LayerItemSelectBorderColor = ReplaceColorAlpha(FStyleColors::Select.GetSpecifiedColor(), 0.9f);
+
+	//constexpr FLinearColor LayerItemSelectHoverFillColor = FLinearColor(1, 1, 1, 0.25f);
+	//const FLinearColor LayerItemSelectHoverBorderColor = FStyleColors::Select.GetSpecifiedColor();
+
+	Style->Set("LayerView.Row.Item", new FSlateRoundedBoxBrush(
+		LayerViewItemFillColor, LayerViewItemCornerRadius,
+		LayerViewItemBorderColor, LayerViewItemBorderWidth));
+
+	Style->Set("LayerView.Row.Hovered", new FSlateRoundedBoxBrush(
+		LayerItemHoverFillColor, LayerViewItemCornerRadius,
+		LayerItemHoverBorderColor, LayerViewItemBorderWidth));
+
+	Style->Set("LayerView.Row.Selected", new FSlateRoundedBoxBrush(
+		LayerItemSelectFillColor, LayerViewItemCornerRadius,
+		LayerItemSelectBorderColor, LayerViewItemBorderWidth));
+
+	Style->Set("LayerView.Row.ActiveBrush", new FSlateRoundedBoxBrush(
+		LayerItemSelectFillColor, LayerViewItemCornerRadius,
+		LayerItemSelectBorderColor, LayerViewItemBorderWidth));
+	Style->Set("LayerView.Row.ActiveHoveredBrush", new FSlateRoundedBoxBrush(
+		LayerItemSelectFillColor, LayerViewItemCornerRadius,
+		LayerItemSelectBorderColor, LayerViewItemBorderWidth));
+	Style->Set("LayerView.Row.InactiveBrush", new FSlateRoundedBoxBrush(
+		LayerItemSelectFillColor, LayerViewItemCornerRadius,
+		LayerItemSelectBorderColor, LayerViewItemBorderWidth));
+	Style->Set("LayerView.Row.InactiveHoveredBrush", new FSlateRoundedBoxBrush(
+		LayerItemSelectFillColor, LayerViewItemCornerRadius,
+		LayerItemSelectBorderColor, LayerViewItemBorderWidth));
+
+	Style->Set("LayerView.Row", FTableRowStyle()
+		.SetTextColor(FStyleColors::Foreground)
+		.SetSelectedTextColor(FStyleColors::ForegroundHover)
+		.SetEvenRowBackgroundBrush(*Style->GetBrush(TEXT("LayerView.Row.Item")))
+		.SetEvenRowBackgroundHoveredBrush(*Style->GetBrush(TEXT("LayerView.Row.Hovered")))
+		.SetOddRowBackgroundBrush(*Style->GetBrush(TEXT("LayerView.Row.Item")))
+		.SetOddRowBackgroundHoveredBrush(*Style->GetBrush(TEXT("LayerView.Row.Hovered")))
+		.SetSelectorFocusedBrush(*Style->GetBrush(TEXT("LayerView.Row.Selected")))
+		.SetActiveBrush(*Style->GetBrush(TEXT("LayerView.Row.ActiveBrush")))
+		.SetActiveHoveredBrush(*Style->GetBrush(TEXT("LayerView.Row.ActiveHoveredBrush")))
+		.SetInactiveBrush(*Style->GetBrush(TEXT("LayerView.Row.InactiveBrush")))
+		.SetInactiveHoveredBrush(*Style->GetBrush(TEXT("LayerView.Row.InactiveHoveredBrush")))
+		.SetSelectorFocusedBrush(BORDER_BRUSH("Common/Selector", FMargin(4.f / 16.f), Style->GetColor(TEXT("Color.Select.Hover"))))
+		.SetDropIndicator_Onto(BOX_BRUSH("Common/DropZoneIndicator_Onto", FMargin(4.0f / 16.0f), Style->GetColor(TEXT("Color.Select.Hover"))))
+		.SetDropIndicator_Above(BOX_BRUSH("Common/DropZoneIndicator_Above", FMargin(4.0f / 16.0f, 4.0f / 16.0f, 0.f, 0.f), Style->GetColor(TEXT("Color.Select.Hover"))))
+		.SetDropIndicator_Below(BOX_BRUSH("Common/DropZoneIndicator_Below", FMargin(4.0f / 16.0f, 0.f, 0.f, 4.0f / 16.0f), Style->GetColor(TEXT("Color.Select.Hover"))))
+	);
+
+	Style->Set("LayerView.AddIcon", new IMAGE_BRUSH("Icons/EditorIcons/LayerAdd", Icon16x16));
+	Style->Set("LayerView.DuplicateIcon", new IMAGE_BRUSH("Icons/EditorIcons/Duplicate_40x", Icon40x40));
+	Style->Set("LayerView.RemoveIcon", new IMAGE_BRUSH("Icons/EditorIcons/LayerRemove", Icon16x16));
+
+	Style->Set("LayerView.Row.Handle", new IMAGE_BRUSH_SVG("Icons/DragHandle", Icon16x16));
+}
+
+void FSubstrateMaterialEditorStyle::SetupTextStyles(const TSharedRef<FSlateStyleSet>& Style)
+{
+	const FTextBlockStyle NormalTextStyle(FAppStyle::GetWidgetStyle<FTextBlockStyle>(TEXT("NormalText")));
+
+	const FLinearColor LayerViewItemTextShadowColor = FLinearColor(0.0f, 0.0f, 0.0f, 0.5f);
+
+	FFontOutlineSettings HandleFontOutline;
+	HandleFontOutline.OutlineColor = LayerViewItemTextShadowColor;
+	HandleFontOutline.OutlineSize = 1;
+
+	FTextBlockStyle SmallTextStyle(NormalTextStyle);
+	SmallTextStyle.SetFont(DEFAULT_FONT("Regular", 8));
+	Style->Set("SmallFont", SmallTextStyle);
+
+	FTextBlockStyle RegularTextStyle(NormalTextStyle);
+	RegularTextStyle.SetFont(DEFAULT_FONT("Regular", 10));
+	Style->Set("RegularFont", RegularTextStyle);
+
+	FTextBlockStyle BoldTextStyle(NormalTextStyle);
+	BoldTextStyle.SetFont(DEFAULT_FONT("Bold", 10));
+	Style->Set("BoldFont", BoldTextStyle);
+
+	Style->Set("ActorName", RegularTextStyle);
+
+	FTextBlockStyle ActorNameBigTextStyle(NormalTextStyle);
+	ActorNameBigTextStyle.SetFont(DEFAULT_FONT("Regular", 14));
+	Style->Set("ActorNameBig", ActorNameBigTextStyle);
+
+	FTextBlockStyle ComponentNameBigTextStyle(NormalTextStyle);
+	ComponentNameBigTextStyle.SetFont(DEFAULT_FONT("Regular", 12));
+	Style->Set("ComponentNameBig", ComponentNameBigTextStyle);
+
+	FTextBlockStyle SlotLayerInfoTextStyle(NormalTextStyle);
+	SlotLayerInfoTextStyle.SetFont(DEFAULT_FONT("Italic", 8));
+	Style->Set("SlotLayerInfo", SlotLayerInfoTextStyle);
+
+	FSlateFontInfo LayerViewItemFont = DEFAULT_FONT("Bold", 12);
+	LayerViewItemFont.OutlineSettings = HandleFontOutline;
+	Style->Set("LayerView.Row.Font", LayerViewItemFont);
+
+	Style->Set("LayerView.Row.HandleFont", RegularTextStyle);
+
+	FTextBlockStyle LayerViewItemTextStyle(NormalTextStyle);
+	LayerViewItemTextStyle.SetShadowOffset(FVector2D(1.0f, 1.0f));
+	LayerViewItemTextStyle.SetColorAndOpacity(LayerViewItemTextShadowColor);
+
+	Style->Set("LayerView.Row.HeaderText",
+		FTextBlockStyle(LayerViewItemTextStyle)
+		.SetColorAndOpacity(FStyleColors::Foreground)
+		.SetFont(LayerViewItemFont));
+
+	Style->Set("LayerView.Row.HeaderText.Small",
+		FTextBlockStyle(LayerViewItemTextStyle)
+		.SetColorAndOpacity(FStyleColors::Foreground)
+		.SetFont(RegularTextStyle.Font));
+
+	FTextBlockStyle StagePropertyDetailsTextStyle(NormalTextStyle);
+	StagePropertyDetailsTextStyle.SetFont(DEFAULT_FONT("Regular", 12));
+	Style->Set("Font.Stage.Details", StagePropertyDetailsTextStyle);
+
+	Style->Set("Font.Stage.Details.Bold", BoldTextStyle);
+
+	FTextBlockStyle StagePropertyDetailsSmallTextStyle(NormalTextStyle);
+	StagePropertyDetailsSmallTextStyle.SetFont(IDetailLayoutBuilder::GetDetailFont());
+	Style->Set("Font.Stage.Details.Small", StagePropertyDetailsSmallTextStyle);
+
+	FTextBlockStyle StagePropertyDetailsSmallBoldTextStyle(NormalTextStyle);
+	StagePropertyDetailsSmallBoldTextStyle.SetFont(IDetailLayoutBuilder::GetDetailFontBold());
+	Style->Set("Font.Stage.Details.Small.Bold", StagePropertyDetailsSmallBoldTextStyle);
+
+	Style->Set("Font.Stage.Details.Small.Bold", FTextBlockStyle(NormalTextStyle)
+		.SetFont(IDetailLayoutBuilder::GetDetailFontBold()));
+}
+
+#undef IMAGE_BRUSH
+#undef BOX_BRUSH
+#undef BORDER_BRUSH
+
+#undef IMAGE_BRUSH_SVG
+#undef BOX_BRUSH_SVG
+#undef BORDER_BRUSH_SVG
+
+#undef DEFAULT_FONT
+
+#undef IMAGE_PLUGIN_BRUSH
+#undef IMAGE_PLUGIN_BRUSH_SVG
diff --git a/Engine/Source/Editor/MaterialEditor/Private/MaterialEditorStyle.h b/Engine/Source/Editor/MaterialEditor/Private/MaterialEditorStyle.h
new file mode 100644
--- /dev/null
+++ b/Engine/Source/Editor/MaterialEditor/Private/MaterialEditorStyle.h
@@ -0,0 +1,27 @@
+// Copyright Epic Games, Inc. All Rights Reserved.
+
+#pragma once
+
+#include "Styling/ISlateStyle.h"
+#include "Styling/SlateStyle.h"
+
+class FSubstrateMaterialEditorStyle
+{
+public:
+	static void Initialize();
+	static void Shutdown();
+
+	static FName GetStyleSetName();
+	static const ISlateStyle& Get();
+
+	static FLinearColor GetColor(const FName& InName);
+	static const FSlateBrush* GetBrush(const FName& InName);
+
+private:
+	static TSharedPtr<FSlateStyleSet> StyleInstance;
+
+	static TSharedRef<FSlateStyleSet> Create();
+
+	static void SetupLayerViewStyles(const TSharedRef<FSlateStyleSet>& Style);
+	static void SetupTextStyles(const TSharedRef<FSlateStyleSet>& Style);
+};
diff --git a/Engine/Source/Editor/MaterialEditor/Private/MaterialInstanceEditor.cpp b/Engine/Source/Editor/MaterialEditor/Private/MaterialInstanceEditor.cpp
--- a/Engine/Source/Editor/MaterialEditor/Private/MaterialInstanceEditor.cpp
+++ b/Engine/Source/Editor/MaterialEditor/Private/MaterialInstanceEditor.cpp
@@ -840,20 +840,20 @@ void FMaterialInstanceEditor::CreateInternalWidgets()
 
 	auto ValidationLambda = ([](const FRootPropertyNodeList& PropertyNodeList) { return true; });
 	MaterialInstanceDetails->SetCustomValidatePropertyNodesFunction(FOnValidateDetailsViewPropertyNodes::CreateLambda(MoveTemp(ValidationLambda)));
-
-	FOnGetDetailCustomizationInstance LayoutMICDetails = FOnGetDetailCustomizationInstance::CreateStatic( 
-		&FMaterialInstanceParameterDetails::MakeInstance, MaterialEditorInstance.Get(), FGetShowHiddenParameters::CreateSP(this, &FMaterialInstanceEditor::GetShowHiddenParameters) );
-	MaterialInstanceDetails->RegisterInstancedCustomPropertyLayout( UMaterialEditorInstanceConstant::StaticClass(), LayoutMICDetails );
-	MaterialInstanceDetails->SetCustomFilterLabel(LOCTEXT("ShowOverriddenOnly", "Show Only Overridden Parameters"));
-	MaterialInstanceDetails->SetCustomFilterDelegate(FSimpleDelegate::CreateSP(this, &FMaterialInstanceEditor::FilterOverriddenProperties));
-	MaterialEditorInstance->DetailsView = MaterialInstanceDetails;
-
 	if (!bIsFunctionPreviewMaterial)
 	{
 		MaterialLayersFunctionsInstance = SNew(SMaterialLayersFunctionsInstanceWrapper)
 			.InMaterialEditorInstance(MaterialEditorInstance)
 			.InShowHiddenDelegate(FGetShowHiddenParameters::CreateSP(this, &FMaterialInstanceEditor::GetShowHiddenParameters));
 	}
+	FOnGetDetailCustomizationInstance LayoutMICDetails = FOnGetDetailCustomizationInstance::CreateStatic( 
+		&FMaterialInstanceParameterDetails::MakeInstance, MaterialEditorInstance.Get(), MaterialLayersFunctionsInstance.Get(), FGetShowHiddenParameters::CreateSP(this, &FMaterialInstanceEditor::GetShowHiddenParameters) );
+	MaterialInstanceDetails->RegisterInstancedCustomPropertyLayout( UMaterialEditorInstanceConstant::StaticClass(), LayoutMICDetails );
+	MaterialInstanceDetails->SetCustomFilterLabel(LOCTEXT("ShowOverriddenOnly", "Show Only Overridden Parameters"));
+	MaterialInstanceDetails->SetCustomFilterDelegate(FSimpleDelegate::CreateSP(this, &FMaterialInstanceEditor::FilterOverriddenProperties));
+	MaterialEditorInstance->DetailsView = MaterialInstanceDetails;
+
+	
 }
 
 void FMaterialInstanceEditor::FilterOverriddenProperties()
diff --git a/Engine/Source/Editor/MaterialEditor/Private/SMaterialLayersFunctionsTree.cpp b/Engine/Source/Editor/MaterialEditor/Private/SMaterialLayersFunctionsTree.cpp
--- a/Engine/Source/Editor/MaterialEditor/Private/SMaterialLayersFunctionsTree.cpp
+++ b/Engine/Source/Editor/MaterialEditor/Private/SMaterialLayersFunctionsTree.cpp
@@ -1,6 +1,7 @@
 // Copyright Epic Games, Inc. All Rights Reserved.
 
 #include "SMaterialLayersFunctionsTree.h"
+#include "SMaterialSubstrateTree.h"
 #include "MaterialEditor/DEditorFontParameterValue.h"
 #include "MaterialEditor/DEditorMaterialLayersParameterValue.h"
 #include "MaterialEditor/DEditorRuntimeVirtualTextureParameterValue.h"
@@ -1786,8 +1787,12 @@ void SMaterialLayersFunctionsInstanceWrapper::Refresh()
 
 	if (LayerParameter != nullptr)
 	{
+#if ENABLE_MATERIAL_LAYER_PROTOTYPE
+		FOnClicked OnRelinkToParent = FOnClicked::CreateSP(NestedTree.ToSharedRef(), &SMaterialSubstrateTree::RelinkLayersToParent);
+#else
 		FOnClicked OnRelinkToParent = FOnClicked::CreateSP(NestedTree.ToSharedRef(), &SMaterialLayersFunctionsInstanceTree::RelinkLayersToParent);
-
+#endif
+		
 		this->ChildSlot
 			[
 				SNew(SVerticalBox)
@@ -1819,9 +1824,15 @@ void SMaterialLayersFunctionsInstanceWrapper::Refresh()
 				.AutoWidth()
 				.VAlign(VAlign_Center)
 				[
+#if ENABLE_MATERIAL_LAYER_PROTOTYPE
+				PropertyCustomizationHelpers::MakeAddButton(FSimpleDelegate::CreateSP(NestedTree.Get(), &SMaterialSubstrateTree::AddLayer))
+#else
 					PropertyCustomizationHelpers::MakeAddButton(FSimpleDelegate::CreateSP(NestedTree.Get(), &SMaterialLayersFunctionsInstanceTree::AddLayer))
+#endif
 				];
 		}
+
+#ifndef ENABLE_MATERIAL_LAYER_PROTOTYPE
 		HeaderBox->AddSlot()
 			.FillWidth(1.0f)
 			[
@@ -1858,6 +1869,7 @@ void SMaterialLayersFunctionsInstanceWrapper::Refresh()
 				.OnClicked(OnChildButtonClicked)
 				.ToolTipText(LOCTEXT("SaveToChildInstance", "Save To Child Instance"))
 			];
+#endif
 	}
 	else
 	{
@@ -1877,10 +1889,18 @@ void SMaterialLayersFunctionsInstanceWrapper::Refresh()
 
 void SMaterialLayersFunctionsInstanceWrapper::Construct(const FArguments& InArgs)
 {
+#if ENABLE_MATERIAL_LAYER_PROTOTYPE
+	NestedTree = SNew(SMaterialSubstrateTree)
+		.InMaterialEditorInstance(InArgs._InMaterialEditorInstance)
+		.InWrapper(this)
+		.InShowHiddenDelegate(InArgs._InShowHiddenDelegate);
+#else
 	NestedTree = SNew(SMaterialLayersFunctionsInstanceTree)
 		.InMaterialEditorInstance(InArgs._InMaterialEditorInstance)
 		.InWrapper(this)
 		.InShowHiddenDelegate(InArgs._InShowHiddenDelegate);
+#endif
+	
 
 	LayerParameter = NestedTree->FunctionParameter;
 
diff --git a/Engine/Source/Editor/MaterialEditor/Private/SMaterialLayersFunctionsTree.h b/Engine/Source/Editor/MaterialEditor/Private/SMaterialLayersFunctionsTree.h
--- a/Engine/Source/Editor/MaterialEditor/Private/SMaterialLayersFunctionsTree.h
+++ b/Engine/Source/Editor/MaterialEditor/Private/SMaterialLayersFunctionsTree.h
@@ -17,6 +17,7 @@
 #include "Widgets/DeclarativeSyntaxSupport.h"
 #include "Widgets/SCompoundWidget.h"
 #include "Widgets/Views/STreeView.h"
+#include "Materials/MaterialExpressionMaterialSample.h"
 
 class IPropertyHandle;
 class SMaterialLayersFunctionsInstanceTree;
@@ -115,7 +116,11 @@ public:
 	TAttribute<ECheckBoxState> IsParamChecked;
 	TWeakObjectPtr<class UDEditorParameterValue> LayerParameter;
 	class UMaterialEditorInstanceConstant* MaterialEditorInstance;
+#if ENABLE_MATERIAL_LAYER_PROTOTYPE
+	TSharedPtr<class SMaterialSubstrateTree> NestedTree;
+#else
 	TSharedPtr<class SMaterialLayersFunctionsInstanceTree> NestedTree;
+#endif
 	FSimpleDelegate OnLayerPropertyChanged;
 };
 
diff --git a/Engine/Source/Editor/MaterialEditor/Private/SMaterialSubstrateTree.cpp b/Engine/Source/Editor/MaterialEditor/Private/SMaterialSubstrateTree.cpp
new file mode 100644
--- /dev/null
+++ b/Engine/Source/Editor/MaterialEditor/Private/SMaterialSubstrateTree.cpp
@@ -0,0 +1,1184 @@
+// Copyright Epic Games, Inc. All Rights Reserved.
+
+#include "SMaterialSubstrateTree.h"
+#include "SMaterialLayersFunctionsTree.h"
+
+#include "SMaterialLayersFunctionsTree.h"
+#include "MaterialEditor/DEditorFontParameterValue.h"
+#include "MaterialEditor/DEditorMaterialLayersParameterValue.h"
+#include "MaterialEditor/DEditorRuntimeVirtualTextureParameterValue.h"
+#include "MaterialEditor/DEditorScalarParameterValue.h"
+#include "MaterialEditor/DEditorStaticComponentMaskParameterValue.h"
+#include "MaterialEditor/DEditorStaticSwitchParameterValue.h"
+#include "MaterialEditor/DEditorTextureParameterValue.h"
+#include "MaterialEditor/DEditorVectorParameterValue.h"
+#include "MaterialEditor/MaterialEditorInstanceConstant.h"
+#include "Materials/Material.h"
+#include "PropertyHandle.h"
+#include "Modules/ModuleManager.h"
+#include "PropertyEditorModule.h"
+#include "ISinglePropertyView.h"
+#include "Materials/MaterialInstanceConstant.h"
+#include "Styling/AppStyle.h"
+#include "PropertyCustomizationHelpers.h"
+#include "ScopedTransaction.h"
+#include "IPropertyRowGenerator.h"
+#include "Widgets/Views/STreeView.h"
+#include "IDetailTreeNode.h"
+#include "AssetThumbnail.h"
+#include "MaterialEditorInstanceDetailCustomization.h"
+#include "MaterialPropertyHelpers.h"
+#include "DetailWidgetRow.h"
+#include "Widgets/Input/SCheckBox.h"
+#include "EditorSupportDelegates.h"
+#include "Widgets/Images/SImage.h"
+#include "MaterialEditor/MaterialEditorPreviewParameters.h"
+#include "Widgets/Input/SButton.h"
+#include "Widgets/Text/SInlineEditableTextBlock.h"
+#include "Materials/MaterialFunctionInstance.h"
+#include "Materials/MaterialExpressionTextureSampleParameter.h"
+#include "Framework/Application/SlateApplication.h"
+#include "Widgets/Input/SEditableTextBox.h"
+#include "Curves/CurveLinearColor.h"
+#include "Curves/CurveLinearColorAtlas.h"
+#include "Subsystems/AssetEditorSubsystem.h"
+#include "Editor.h"
+#include "EditorFontGlyphs.h"
+#include "SResetToDefaultPropertyEditor.h"
+#include "DragAndDrop/AssetDragDropOp.h"
+#include "ThumbnailRendering/ThumbnailManager.h"
+#include "Styling/AppStyle.h"
+#include "Styling/StyleColors.h"
+#include "MaterialEditorStyle.h"
+
+#define LOCTEXT_NAMESPACE "MaterialSubstrateTree"
+
+
+/// ===========================================================================================================
+/// SMaterialSubstrateItem Methods
+/// ===========================================================================================================
+FString SMaterialSubstrateTreeItem::GetCurvePath(UDEditorScalarParameterValue* Parameter) const
+{
+	FString Path = Parameter->AtlasData.Curve->GetPathName();
+	return Path;
+}
+
+const FSlateBrush* SMaterialSubstrateTreeItem::GetBorderImage() const
+{
+	return FAppStyle::Get().GetBrush("DetailsView.CategoryMiddle");
+}
+
+FSlateColor SMaterialSubstrateTreeItem::GetOuterBackgroundColor(TSharedPtr<FSortedParamData> InParamData) const
+{
+	if (InParamData->StackDataType == EStackDataType::Stack)
+	{
+		if (bIsBeingDragged)
+		{
+			return FAppStyle::Get().GetSlateColor("Colors.Recessed");
+		}
+		else if (bIsHoveredDragTarget)
+		{
+			return FAppStyle::Get().GetSlateColor("Colors.Highlight");
+		}
+		else
+		{
+			return FAppStyle::Get().GetSlateColor("Colors.Header");
+		}
+	}
+	else if (IsHovered() || InParamData->StackDataType == EStackDataType::Group)
+	{
+		return FAppStyle::Get().GetSlateColor("Colors.Header");
+	}
+
+	return FAppStyle::Get().GetSlateColor("Colors.Panel");
+}
+
+void SMaterialSubstrateTreeItem::RefreshOnRowChange(const FAssetData& AssetData, SMaterialSubstrateTree* InTree)
+{
+	if (SMaterialLayersFunctionsInstanceWrapper* Wrapper = InTree->GetWrapper())
+	{
+		if (Wrapper->OnLayerPropertyChanged.IsBound())
+		{
+			Wrapper->OnLayerPropertyChanged.Execute();
+		}
+		else
+		{
+			InTree->CreateGroupsWidget();
+		}
+	}
+}
+
+bool SMaterialSubstrateTreeItem::GetFilterState(SMaterialSubstrateTree* InTree, TSharedPtr<FSortedParamData> InStackData) const
+{
+	if (InStackData->ParameterInfo.Association == EMaterialParameterAssociation::LayerParameter)
+	{
+		return InTree->FunctionInstance->EditorOnly.RestrictToLayerRelatives[InStackData->ParameterInfo.Index];
+	}
+	if (InStackData->ParameterInfo.Association == EMaterialParameterAssociation::BlendParameter)
+	{
+		return InTree->FunctionInstance->EditorOnly.RestrictToBlendRelatives[InStackData->ParameterInfo.Index];
+	}
+	return false;
+}
+
+void SMaterialSubstrateTreeItem::FilterClicked(const ECheckBoxState NewCheckedState, SMaterialSubstrateTree* InTree, TSharedPtr<FSortedParamData> InStackData)
+{
+	if (InStackData->ParameterInfo.Association == EMaterialParameterAssociation::LayerParameter)
+	{
+		InTree->FunctionInstance->EditorOnly.RestrictToLayerRelatives[InStackData->ParameterInfo.Index] = !InTree->FunctionInstance->EditorOnly.RestrictToLayerRelatives[InStackData->ParameterInfo.Index];
+	}
+	if (InStackData->ParameterInfo.Association == EMaterialParameterAssociation::BlendParameter)
+	{
+		InTree->FunctionInstance->EditorOnly.RestrictToBlendRelatives[InStackData->ParameterInfo.Index] = !InTree->FunctionInstance->EditorOnly.RestrictToBlendRelatives[InStackData->ParameterInfo.Index];
+	}
+}
+
+ECheckBoxState SMaterialSubstrateTreeItem::GetFilterChecked(SMaterialSubstrateTree* InTree, TSharedPtr<FSortedParamData> InStackData) const
+{
+	return GetFilterState(InTree, InStackData) ? ECheckBoxState::Checked : ECheckBoxState::Unchecked;
+}
+
+FText SMaterialSubstrateTreeItem::GetLayerName(SMaterialSubstrateTree* InTree, int32 Counter) const
+{
+	return InTree->FunctionInstance->GetLayerName(Counter);
+}
+
+void SMaterialSubstrateTreeItem::OnNameChanged(const FText& InText, ETextCommit::Type CommitInfo, SMaterialSubstrateTree* InTree, int32 Counter)
+{
+	const FScopedTransaction Transaction(LOCTEXT("RenamedSection", "Renamed layer and blend section"));
+	InTree->FunctionInstanceHandle->NotifyPreChange();
+	InTree->FunctionInstance->EditorOnly.LayerNames[Counter] = InText;
+	InTree->FunctionInstance->UnlinkLayerFromParent(Counter);
+	InTree->MaterialEditorInstance->CopyToSourceInstance(true);
+	InTree->FunctionInstanceHandle->NotifyPostChange(EPropertyChangeType::ValueSet);
+}
+
+FReply SMaterialSubstrateTreeItem::OnLayerDrop(const FDragDropEvent& DragDropEvent)
+{
+	if (!bIsHoveredDragTarget)
+	{
+		return FReply::Unhandled();
+	}
+	FScopedTransaction Transaction(NSLOCTEXT("UnrealEd", "MoveLayer", "Move Layer"));
+	Tree->FunctionInstanceHandle->NotifyPreChange();
+	bIsHoveredDragTarget = false;
+	TSharedPtr<FSubstrateLayerDragDropOp> ArrayDropOp = DragDropEvent.GetOperationAs< FSubstrateLayerDragDropOp >();
+	TSharedPtr<SMaterialSubstrateTreeItem> LayerPtr = nullptr;
+	if (ArrayDropOp.IsValid() && ArrayDropOp->OwningStack.IsValid())
+	{
+		LayerPtr = ArrayDropOp->OwningStack.Pin();
+		LayerPtr->bIsBeingDragged = false;
+	}
+	else
+	{
+		// see if it is an accepted asset drop
+		TSharedPtr<FAssetDragDropOp> AssetDropOp = DragDropEvent.GetOperationAs<FAssetDragDropOp>();
+
+		if (AssetDropOp.IsValid())
+		{
+			bool DidModifyTree = false;
+			for (const FAssetData& AssetData : AssetDropOp->GetAssets())
+			{
+				EMaterialParameterAssociation InAssociation = EMaterialParameterAssociation::GlobalParameter;
+			
+				if (AssetData.AssetClassPath.GetAssetName() == TEXT("MaterialFunctionMaterialLayer"))
+				{
+					InAssociation = EMaterialParameterAssociation::LayerParameter;
+					Tree->InsertLayerFromAsset(AssetData, StackParameterData->ParameterInfo.Index, InAssociation);
+					DidModifyTree = true;
+				}
+				else if (AssetData.AssetClassPath.GetAssetName() == TEXT("MaterialFunctionMaterialLayerBlend"))
+				{
+					InAssociation = EMaterialParameterAssociation::BlendParameter;
+					Tree->InsertLayerFromAsset(AssetData, StackParameterData->ParameterInfo.Index, InAssociation);
+					DidModifyTree = true;
+				}
+			}
+
+			if (DidModifyTree)
+			{
+				return FReply::Handled();
+			}
+			
+		}
+		return FReply::Unhandled();
+	}
+	
+	if (!LayerPtr.IsValid())
+	{
+		return FReply::Unhandled();
+	}
+	TSharedPtr<FSortedParamData> SwappingPropertyData = LayerPtr->StackParameterData;
+	TSharedPtr<FSortedParamData> SwappablePropertyData = StackParameterData;
+	if (SwappingPropertyData.IsValid() && SwappablePropertyData.IsValid())
+	{
+		if (SwappingPropertyData != SwappablePropertyData)
+		{
+			int32 OriginalIndex = SwappingPropertyData->ParameterInfo.Index;
+			if (SwappingPropertyData->ParameterInfo.Association == EMaterialParameterAssociation::BlendParameter)
+			{
+				OriginalIndex++;
+			}
+
+			int32 NewIndex = SwappablePropertyData->ParameterInfo.Index;
+			if (SwappablePropertyData->ParameterInfo.Association == EMaterialParameterAssociation::BlendParameter)
+			{
+				NewIndex++;
+			}
+
+			if (OriginalIndex != NewIndex)
+			{
+				Tree->MaterialEditorInstance->SourceInstance->SwapLayerParameterIndices(OriginalIndex, NewIndex);
+
+				// Need to save the moving and target expansion states before swapping
+				const bool bOriginalSwappableExpansion = IsItemExpanded();
+				const bool bOriginalSwappingExpansion = LayerPtr->IsItemExpanded();
+
+				TArray<void*> StructPtrs;
+				Tree->FunctionInstanceHandle->AccessRawData(StructPtrs);
+				FMaterialLayersFunctions* MaterialLayersFunctions = static_cast<FMaterialLayersFunctions*>(StructPtrs[0]);
+				MaterialLayersFunctions->MoveBlendedLayer(OriginalIndex, NewIndex);
+
+				Tree->OnExpansionChanged(SwappablePropertyData, bOriginalSwappingExpansion);
+				Tree->OnExpansionChanged(SwappingPropertyData, bOriginalSwappableExpansion);
+				Tree->FunctionInstanceHandle->NotifyPostChange(EPropertyChangeType::ValueSet);
+				Tree->CreateGroupsWidget();
+				Tree->RequestTreeRefresh();
+				Tree->SetParentsExpansionState();
+			}
+		}
+	}
+
+	return FReply::Handled();
+}
+
+
+bool SMaterialSubstrateTree::IsOverriddenExpression(class UDEditorParameterValue* Parameter, int32 InIndex)
+{
+	return FMaterialPropertyHelpers::IsOverriddenExpression(Parameter) && FunctionInstance->EditorOnly.LayerStates[InIndex];
+}
+
+bool SMaterialSubstrateTree::IsOverriddenExpression(TObjectPtr<UDEditorParameterValue> Parameter, int32 InIndex)
+{
+	return IsOverriddenExpression(Parameter.Get(), InIndex);
+}
+
+FGetShowHiddenParameters SMaterialSubstrateTree::GetShowHiddenDelegate() const
+{
+	return ShowHiddenDelegate;
+}
+
+void  SMaterialSubstrateTreeItem::OnOverrideParameter(bool NewValue, class UDEditorParameterValue* Parameter)
+{
+	FMaterialPropertyHelpers::OnOverrideParameter(NewValue, Parameter, MaterialEditorInstance);
+}
+
+void  SMaterialSubstrateTreeItem::OnOverrideParameter(bool NewValue, TObjectPtr<UDEditorParameterValue> Parameter)
+{
+	OnOverrideParameter(NewValue, Parameter.Get());
+}
+
+void SMaterialSubstrateTreeItem::Construct(const FArguments& InArgs, const TSharedRef<STableViewBase>& InOwnerTableView)
+{
+	StackParameterData = InArgs._StackParameterData;
+	MaterialEditorInstance = InArgs._MaterialEditorInstance;
+	Tree = InArgs._InTree;
+
+	TSharedRef<SWidget> LeftSideWidget = SNullWidget::NullWidget;
+	TSharedRef<SWidget> RightSideWidget = SNullWidget::NullWidget;
+	TSharedRef<SWidget> ResetWidget = SNullWidget::NullWidget;
+	FText NameOverride;
+	TSharedRef<SVerticalBox> WrapperWidget = SNew(SVerticalBox);
+	EHorizontalAlignment ValueAlignment = HAlign_Left;
+// STACK --------------------------------------------------
+	if (StackParameterData->StackDataType == EStackDataType::Stack)
+	{
+		// WrapperWidget->AddSlot()
+		// 	.Padding(5.0f)
+		// 	.AutoHeight()
+		// 	[
+		// 		SNullWidget::NullWidget
+		// 	];
+#if WITH_EDITOR
+		NameOverride = Tree->FunctionInstance->GetLayerName(StackParameterData->ParameterInfo.Index);
+#endif
+		TSharedRef<SHorizontalBox> HeaderRowWidget = SNew(SHorizontalBox);
+
+		if (StackParameterData->ParameterInfo.Index != 0)
+		{
+			TAttribute<bool>::FGetter IsEnabledGetter = TAttribute<bool>::FGetter::CreateSP(Tree, &SMaterialSubstrateTree::IsLayerVisible, StackParameterData->ParameterInfo.Index);
+			TAttribute<bool> IsEnabledAttribute = TAttribute<bool>::Create(IsEnabledGetter);
+
+			FOnClicked VisibilityClickedDelegate = FOnClicked::CreateSP(Tree, &SMaterialSubstrateTree::ToggleLayerVisibility, StackParameterData->ParameterInfo.Index);
+
+			HeaderRowWidget->AddSlot()
+				.AutoWidth()
+				.VAlign(VAlign_Center)
+				[
+					PropertyCustomizationHelpers::MakeVisibilityButton(VisibilityClickedDelegate, FText(), IsEnabledAttribute)
+				];
+		}
+		const float ThumbnailSize = 64.0f;
+		TArray<TSharedPtr<FSortedParamData>> AssetChildren = StackParameterData->Children;
+		if (AssetChildren.Num() > 0)
+		{
+			HeaderRowWidget->AddSlot()
+				.HAlign(HAlign_Center)
+				.VAlign(VAlign_Center)
+				.Padding(0.0f)
+				.AutoWidth()
+				[
+					SNullWidget::NullWidget
+				];
+		}
+		for (TSharedPtr<FSortedParamData> AssetChild : AssetChildren)
+		{
+			TSharedPtr<SBox> ThumbnailBox;
+			UObject* AssetObject = nullptr;
+			AssetChild->ParameterHandle->GetValue(AssetObject);
+			int32 PreviewIndex = INDEX_NONE;
+			int32 ThumbnailIndex = INDEX_NONE;
+			EMaterialParameterAssociation PreviewAssociation = EMaterialParameterAssociation::GlobalParameter;
+			if (AssetObject)
+			{
+				if (Cast<UMaterialFunctionInterface>(AssetObject)->GetMaterialFunctionUsage() == EMaterialFunctionUsage::MaterialLayer)
+				{
+					PreviewIndex = StackParameterData->ParameterInfo.Index;
+					PreviewAssociation = EMaterialParameterAssociation::LayerParameter;
+					Tree->UpdateThumbnailMaterial(PreviewAssociation, PreviewIndex);
+					ThumbnailIndex = PreviewIndex;
+				}
+				if (Cast<UMaterialFunctionInterface>(AssetObject)->GetMaterialFunctionUsage() == EMaterialFunctionUsage::MaterialLayerBlend)
+				{
+					PreviewIndex = StackParameterData->ParameterInfo.Index;
+					PreviewAssociation = EMaterialParameterAssociation::BlendParameter;
+					Tree->UpdateThumbnailMaterial(PreviewAssociation, PreviewIndex, true);
+					ThumbnailIndex = PreviewIndex - 1;
+				}
+			}
+			HeaderRowWidget->AddSlot()
+				.AutoWidth()
+				.HAlign(HAlign_Center)
+				.VAlign(VAlign_Center)
+				.Padding(4.0f)
+				.MaxWidth(ThumbnailSize)
+				[
+					SAssignNew(ThumbnailBox, SBox)
+					[
+						Tree->CreateThumbnailWidget(PreviewAssociation, ThumbnailIndex, ThumbnailSize)
+					]
+				];
+			ThumbnailBox->SetMaxDesiredHeight(ThumbnailSize);
+			ThumbnailBox->SetMinDesiredHeight(ThumbnailSize);
+			ThumbnailBox->SetMinDesiredWidth(ThumbnailSize);
+			ThumbnailBox->SetMaxDesiredWidth(ThumbnailSize);
+		}
+
+		
+
+		if (StackParameterData->ParameterInfo.Index != 0)
+		{
+			HeaderRowWidget->AddSlot()
+				.VAlign(VAlign_Center)
+				.AutoWidth()
+				.Padding(5.0f)
+				[
+					SNew(SEditableTextBox)
+					.Text(TAttribute<FText>::Create(TAttribute<FText>::FGetter::CreateSP(this, &SMaterialSubstrateTreeItem::GetLayerName, InArgs._InTree, StackParameterData->ParameterInfo.Index)))
+					.OnTextCommitted(FOnTextCommitted::CreateSP(this, &SMaterialSubstrateTreeItem::OnNameChanged, InArgs._InTree, StackParameterData->ParameterInfo.Index))
+				];
+		}
+		else
+		{
+			HeaderRowWidget->AddSlot()
+				.VAlign(VAlign_Center)
+				.AutoWidth()
+				.Padding(5.0f)
+				[
+					SNew(STextBlock)
+					.Text(NameOverride)
+				];
+		}
+
+		// Unlink UI
+		HeaderRowWidget->AddSlot()
+			.FillWidth(1.0f)
+			.VAlign(VAlign_Center)
+			[
+				SNullWidget::NullWidget
+			];
+		HeaderRowWidget->AddSlot()
+			.AutoWidth()
+			.VAlign(VAlign_Center)
+			.Padding(0.0f, 0.0f, 0.0f, 0.0f)
+			[
+				SNew(SButton)
+				.Text(LOCTEXT("Unlink", "Unlink"))
+				.HAlign(HAlign_Center)
+				.OnClicked(Tree, &SMaterialSubstrateTree::UnlinkLayer, StackParameterData->ParameterInfo.Index)
+				.ToolTipText(LOCTEXT("UnlinkLayer", "Whether or not to unlink this layer/blend combination from the parent."))
+				.Visibility(Tree, &SMaterialSubstrateTree::GetUnlinkLayerVisibility, StackParameterData->ParameterInfo.Index)
+			];
+
+		// Can only remove layers that aren't the base layer.
+		if (StackParameterData->ParameterInfo.Index != 0)
+		{
+			HeaderRowWidget->AddSlot()
+				.AutoWidth()
+				.VAlign(VAlign_Center)
+				.Padding(0.0f, 0.0f, 5.0f, 0.0f)
+				[
+					PropertyCustomizationHelpers::MakeClearButton(FSimpleDelegate::CreateSP(InArgs._InTree, &SMaterialSubstrateTree::RemoveLayer, StackParameterData->ParameterInfo.Index))
+				];
+		}
+		LeftSideWidget = HeaderRowWidget;
+	}
+// END STACK
+
+
+// FINAL WRAPPER
+	if (StackParameterData->StackDataType == EStackDataType::Stack)
+	{
+		TSharedPtr<SHorizontalBox> FinalStack;
+		WrapperWidget->AddSlot()
+			.Padding(15.0f)/*.AutoHeight()*/
+			[
+						SAssignNew(FinalStack, SHorizontalBox)
+			];
+		// if (StackParameterData->ParameterInfo.Index != 0)
+		// {
+		// 	FinalStack->AddSlot()
+		// 		.HAlign(HAlign_Center)
+		// 		.VAlign(VAlign_Center)
+		// 		.Padding(2.5f, 0)
+		// 		.AutoWidth()
+		// 		[
+		// 			FMaterialPropertyHelpers::MakeStackReorderHandle(SharedThis(this))
+		// 		];
+		// }
+		// FinalStack->AddSlot()
+		// 	.AutoWidth()
+		// 	.VAlign(VAlign_Center)
+		// 	.Padding(FMargin(2.0f))
+		// 	[
+		// 		SNew(SExpanderArrow, SharedThis(this))
+		// 	];
+		FinalStack->AddSlot()
+			.Padding(FMargin(2.0f))
+			.VAlign(VAlign_Center)
+			[
+				LeftSideWidget
+			];
+	}
+
+	this->ChildSlot
+		[
+			WrapperWidget
+		];
+
+	FOnTableRowDragEnter LayerDragDelegate = FOnTableRowDragEnter::CreateSP(this, &SMaterialSubstrateTreeItem::OnLayerDragEnter);
+	FOnTableRowDragLeave LayerDragLeaveDelegate = FOnTableRowDragLeave::CreateSP(this, &SMaterialSubstrateTreeItem::OnLayerDragLeave);
+	FOnTableRowDrop LayerDropDelegate = FOnTableRowDrop::CreateSP(this, &SMaterialSubstrateTreeItem::OnLayerDrop);
+
+	STableRow< TSharedPtr<FSortedParamData> >::ConstructInternal(
+		STableRow< TSharedPtr<FSortedParamData> >::FArguments()
+		.Padding(20.0f)
+		.OnPaintDropIndicator(this, &SMaterialSubstrateTreeItem::OnLayerItemPaintDropIndicator)
+		.Style(FSubstrateMaterialEditorStyle::Get(), "LayerView.Row")
+		.ShowSelection(true)
+		.OnDragEnter(LayerDragDelegate)
+		.OnDragLeave(LayerDragLeaveDelegate)
+		.OnDrop(LayerDropDelegate),
+		InOwnerTableView
+	);
+}
+int32 SMaterialSubstrateTreeItem::OnLayerItemPaintDropIndicator(EItemDropZone InItemDropZone, const FPaintArgs& Args, const FGeometry& AllottedGeometry, const FSlateRect& MyCullingRect, FSlateWindowElementList& OutDrawElements, int32 LayerId, const FWidgetStyle& InWidgetStyle, bool bParentEnabled) const
+{
+	const FSlateBrush* DropIndicatorBrush = GetDropIndicatorBrush(InItemDropZone);
+	static float OffsetX = 10.0f;
+	FVector2D Offset(OffsetX * GetIndentLevel(), 0.0f);
+
+	FSlateDrawElement::MakeBox(
+		OutDrawElements,
+		LayerId++,
+		AllottedGeometry.ToPaintGeometry(FVector2D(AllottedGeometry.GetLocalSize() - Offset), FSlateLayoutTransform(Offset)),
+		DropIndicatorBrush,
+		ESlateDrawEffect::None,
+		DropIndicatorBrush->GetTint(InWidgetStyle) * InWidgetStyle.GetColorAndOpacityTint()
+	);
+
+	return LayerId;
+}
+FString SMaterialSubstrateTreeItem::GetInstancePath(SMaterialSubstrateTree* InTree) const
+{
+	FString InstancePath;
+	if (StackParameterData->ParameterInfo.Association == EMaterialParameterAssociation::BlendParameter && InTree->FunctionInstance->Blends.IsValidIndex(StackParameterData->ParameterInfo.Index))
+	{
+		InstancePath = InTree->FunctionInstance->Blends[StackParameterData->ParameterInfo.Index]->GetPathName();
+	}
+	else if (StackParameterData->ParameterInfo.Association == EMaterialParameterAssociation::LayerParameter && InTree->FunctionInstance->Layers.IsValidIndex(StackParameterData->ParameterInfo.Index))
+	{
+		InstancePath = InTree->FunctionInstance->Layers[StackParameterData->ParameterInfo.Index]->GetPathName();
+	}
+	return InstancePath;
+}
+
+void SMaterialSubstrateTree::Construct(const FArguments& InArgs)
+{
+	ColumnSizeData.SetValueColumnWidth(0.5f);
+
+	MaterialEditorInstance = InArgs._InMaterialEditorInstance;
+	Wrapper = InArgs._InWrapper;
+	ShowHiddenDelegate = InArgs._InShowHiddenDelegate;
+	CreateGroupsWidget();
+
+#if WITH_EDITOR
+	//Fixup for adding new bool arrays to the class
+	if (FunctionInstance)
+	{
+		if (FunctionInstance->Layers.Num() != FunctionInstance->EditorOnly.RestrictToLayerRelatives.Num())
+		{
+			int32 OriginalSize = FunctionInstance->EditorOnly.RestrictToLayerRelatives.Num();
+			for (int32 LayerIt = 0; LayerIt < FunctionInstance->Layers.Num() - OriginalSize; LayerIt++)
+			{
+				FunctionInstance->EditorOnly.RestrictToLayerRelatives.Add(false);
+			}
+		}
+		if (FunctionInstance->Blends.Num() != FunctionInstance->EditorOnly.RestrictToBlendRelatives.Num())
+		{
+			int32 OriginalSize = FunctionInstance->EditorOnly.RestrictToBlendRelatives.Num();
+			for (int32 BlendIt = 0; BlendIt < FunctionInstance->Blends.Num() - OriginalSize; BlendIt++)
+			{
+				FunctionInstance->EditorOnly.RestrictToBlendRelatives.Add(false);
+			}
+		}
+	}
+#endif
+
+	STreeView<TSharedPtr<FSortedParamData>>::Construct(
+		STreeView::FArguments()
+		.TreeItemsSource(&LayerProperties)
+		.SelectionMode(ESelectionMode::Single)
+		.OnSelectionChanged(this, &SMaterialSubstrateTree::OnSelectionChangedMaterialSubstrateView)
+		.OnGenerateRow(this, &SMaterialSubstrateTree::OnGenerateRowMaterialLayersFunctionsTreeView)
+		.OnGetChildren(this, &SMaterialSubstrateTree::OnGetChildrenMaterialLayersFunctionsTreeView)
+		.OnExpansionChanged(this, &SMaterialSubstrateTree::OnExpansionChanged)
+	);
+
+	SetParentsExpansionState();
+}
+void SMaterialSubstrateTree::OnSelectionChangedMaterialSubstrateView(TSharedPtr<FSortedParamData> InSelectedItem, ESelectInfo::Type SelectInfo)
+{
+	MaterialEditorInstance->DetailsView.Pin()->ForceRefresh();
+}
+TSharedRef< ITableRow > SMaterialSubstrateTree::OnGenerateRowMaterialLayersFunctionsTreeView(TSharedPtr<FSortedParamData> Item, const TSharedRef< STableViewBase >& OwnerTable)
+{
+	TSharedRef< SMaterialSubstrateTreeItem > ReturnRow = SNew(SMaterialSubstrateTreeItem, OwnerTable)
+		.StackParameterData(Item)
+		.MaterialEditorInstance(MaterialEditorInstance)
+		.InTree(this);
+	return ReturnRow;
+}
+
+void SMaterialSubstrateTree::OnGetChildrenMaterialLayersFunctionsTreeView(TSharedPtr<FSortedParamData> InParent, TArray< TSharedPtr<FSortedParamData> >& OutChildren)
+{
+	OutChildren = InParent->Children;
+}
+
+
+void SMaterialSubstrateTree::OnExpansionChanged(TSharedPtr<FSortedParamData> Item, bool bIsExpanded)
+{
+	bool* ExpansionValue = MaterialEditorInstance->SourceInstance->LayerParameterExpansion.Find(Item->NodeKey);
+	if (ExpansionValue == nullptr)
+	{
+		MaterialEditorInstance->SourceInstance->LayerParameterExpansion.Add(Item->NodeKey, bIsExpanded);
+	}
+	else if (*ExpansionValue != bIsExpanded)
+	{
+		MaterialEditorInstance->SourceInstance->LayerParameterExpansion.Emplace(Item->NodeKey, bIsExpanded);
+	}
+	// Expand any children that are also expanded
+	for (auto Child : Item->Children)
+	{
+		bool* ChildExpansionValue = MaterialEditorInstance->SourceInstance->LayerParameterExpansion.Find(Child->NodeKey);
+		if (ChildExpansionValue != nullptr && *ChildExpansionValue == true)
+		{
+			SetItemExpansion(Child, true);
+		}
+	}
+}
+
+void SMaterialSubstrateTree::SetParentsExpansionState()
+{
+	for (const auto& Pair : LayerProperties)
+	{
+		if (Pair->Children.Num())
+		{
+			bool* bIsExpanded = MaterialEditorInstance->SourceInstance->LayerParameterExpansion.Find(Pair->NodeKey);
+			if (bIsExpanded)
+			{
+				SetItemExpansion(Pair, *bIsExpanded);
+			}
+		}
+	}
+}
+
+void SMaterialSubstrateTree::RefreshOnAssetChange(const struct FAssetData& InAssetData, int32 Index, EMaterialParameterAssociation MaterialType)
+{
+	FMaterialPropertyHelpers::OnMaterialLayerAssetChanged(InAssetData, Index, MaterialType, FunctionInstanceHandle, FunctionInstance);
+	//set their overrides back to 0
+	MaterialEditorInstance->CleanParameterStack(Index, MaterialType);
+	CreateGroupsWidget();
+	MaterialEditorInstance->ResetOverrides(Index, MaterialType);
+	RequestTreeRefresh();
+}
+
+void SMaterialSubstrateTree::ResetAssetToDefault(TSharedPtr<FSortedParamData> InData)
+{
+	FMaterialPropertyHelpers::ResetLayerAssetToDefault(InData->Parameter, InData->ParameterInfo.Association, InData->ParameterInfo.Index, MaterialEditorInstance);
+	UpdateThumbnailMaterial(InData->ParameterInfo.Association, InData->ParameterInfo.Index, false);
+	CreateGroupsWidget();
+	RequestTreeRefresh();
+}
+void SMaterialSubstrateTree::InsertLayerFromAsset(const struct FAssetData& InAssetData, int32 Index, EMaterialParameterAssociation MaterialType)
+{
+	const FScopedTransaction Transaction(LOCTEXT("InsertLayerAndBlend", "Insert a new Layer and a Blend into it"));
+	FunctionInstanceHandle->NotifyPreChange();
+	//TODO: Need to implement insert at index position instead of this
+	int32 LayerIndex = FunctionInstance->AppendBlendedLayer();
+
+	if (MaterialType == LayerParameter)
+	{
+		// Blend indices are offset by 1, no blend for base layer, so we only update the index for Layer
+		Index = LayerIndex;
+	}
+	FMaterialPropertyHelpers::OnMaterialLayerAssetChanged(InAssetData, Index, MaterialType, FunctionInstanceHandle, FunctionInstance);
+
+	FunctionInstanceHandle->NotifyFinishedChangingProperties();
+	MaterialEditorInstance->CleanParameterStack(Index, MaterialType);
+	CreateGroupsWidget();
+	MaterialEditorInstance->ResetOverrides(Index, MaterialType);
+	RequestTreeRefresh();
+}
+void SMaterialSubstrateTree::AddLayer()
+{
+	const FScopedTransaction Transaction(LOCTEXT("AddLayerAndBlend", "Add a new Layer and a Blend into it"));
+	FunctionInstanceHandle->NotifyPreChange();
+	FunctionInstance->AppendBlendedLayer();
+	FunctionInstanceHandle->NotifyPostChange(EPropertyChangeType::ArrayAdd);
+	CreateGroupsWidget();
+	RequestTreeRefresh();
+}
+
+void SMaterialSubstrateTree::RemoveLayer(int32 Index)
+{
+	const FScopedTransaction Transaction(LOCTEXT("RemoveLayerAndBlend", "Remove a Layer and the attached Blend"));
+	FunctionInstanceHandle->NotifyPreChange();
+	FunctionInstance->RemoveBlendedLayerAt(Index);
+	MaterialEditorInstance->SourceInstance->RemoveLayerParameterIndex(Index);
+	FunctionInstanceHandle->NotifyPostChange(EPropertyChangeType::ArrayRemove);
+	CreateGroupsWidget();
+	RequestTreeRefresh();
+}
+
+FReply SMaterialSubstrateTree::UnlinkLayer(int32 Index)
+{
+	const FScopedTransaction Transaction(LOCTEXT("UnlinkLayerFromParent", "Unlink a layer from the parent"));
+	FunctionInstanceHandle->NotifyPreChange();
+	FunctionInstance->UnlinkLayerFromParent(Index);
+	FunctionInstanceHandle->NotifyPostChange(EPropertyChangeType::ValueSet);
+	CreateGroupsWidget();
+	return FReply::Handled();
+}
+
+FReply SMaterialSubstrateTree::RelinkLayersToParent()
+{
+	const FScopedTransaction Transaction(LOCTEXT("RelinkLayersToParent", "Relink layers to parent"));
+	FunctionInstanceHandle->NotifyPreChange();
+	FunctionInstance->RelinkLayersToParent();
+	FunctionInstanceHandle->NotifyPostChange(EPropertyChangeType::ValueSet);
+	MaterialEditorInstance->RegenerateArrays();
+	CreateGroupsWidget();
+	return FReply::Handled();
+}
+
+EVisibility SMaterialSubstrateTree::GetUnlinkLayerVisibility(int32 Index) const
+{
+	if (FunctionInstance->IsLayerLinkedToParent(Index))
+	{
+		return EVisibility::Visible;
+	}
+	return EVisibility::Collapsed;
+}
+
+EVisibility SMaterialSubstrateTree::GetRelinkLayersToParentVisibility() const
+{
+	if (FunctionInstance->HasAnyUnlinkedLayers())
+	{
+		return EVisibility::Visible;
+	}
+	return EVisibility::Collapsed;
+}
+
+FReply SMaterialSubstrateTree::ToggleLayerVisibility(int32 Index)
+{
+	if (!FSlateApplication::Get().GetModifierKeys().AreModifersDown(EModifierKey::Alt))
+	{
+		bLayerIsolated = false;
+		const FScopedTransaction Transaction(LOCTEXT("ToggleLayerAndBlendVisibility", "Toggles visibility for a blended layer"));
+		FunctionInstanceHandle->NotifyPreChange();
+		FunctionInstance->ToggleBlendedLayerVisibility(Index);
+		FunctionInstanceHandle->NotifyPostChange(EPropertyChangeType::ValueSet);
+		CreateGroupsWidget();
+		return FReply::Handled();
+	}
+	else
+	{
+		const FScopedTransaction Transaction(LOCTEXT("ToggleLayerAndBlendVisibility", "Toggles visibility for a blended layer"));
+		FunctionInstanceHandle->NotifyPreChange();
+		if (FunctionInstance->GetLayerVisibility(Index) == false)
+		{
+			// Reset if clicking on a disabled layer
+			FunctionInstance->SetBlendedLayerVisibility(Index, true);
+			bLayerIsolated = false;
+		}
+		for (int32 LayerIt = 1; LayerIt < FunctionInstance->EditorOnly.LayerStates.Num(); LayerIt++)
+		{
+			if (LayerIt != Index)
+			{
+				FunctionInstance->SetBlendedLayerVisibility(LayerIt, bLayerIsolated);
+			}
+		}
+
+		bLayerIsolated = !bLayerIsolated;
+		FunctionInstanceHandle->NotifyPostChange(EPropertyChangeType::ValueSet);
+		CreateGroupsWidget();
+		return FReply::Handled();
+	}
+
+}
+
+TSharedPtr<class FAssetThumbnailPool> SMaterialSubstrateTree::GetTreeThumbnailPool()
+{
+	return UThumbnailManager::Get().GetSharedThumbnailPool();
+}
+
+TSharedPtr<IDetailTreeNode> SMaterialSubstrateTree::FindParameterGroupsNode(TSharedPtr<IPropertyRowGenerator> PropertyRowGenerator)
+{
+	const TArray<TSharedRef<IDetailTreeNode>> RootNodes = PropertyRowGenerator->GetRootTreeNodes();
+	if (RootNodes.Num() > 0)
+	{
+		TSharedPtr<IDetailTreeNode> Category = RootNodes[0];
+		TArray<TSharedRef<IDetailTreeNode>> Children;
+		Category->GetChildren(Children);
+
+		for (int32 ChildIdx = 0; ChildIdx < Children.Num(); ChildIdx++)
+		{
+			TSharedPtr<IPropertyHandle> PropertyHandle = Children[ChildIdx]->CreatePropertyHandle();
+			if (PropertyHandle.IsValid() && PropertyHandle->GetProperty() && PropertyHandle->GetProperty()->GetName() == "ParameterGroups")
+			{
+				return Children[ChildIdx];
+			}
+		}
+	}
+	return nullptr;
+}
+
+void SMaterialSubstrateTree::CreateGroupsWidget()
+{
+	check(MaterialEditorInstance);
+	MaterialEditorInstance->RegenerateArrays();
+	NonLayerProperties.Empty();
+	LayerProperties.Empty();
+	FunctionParameter = nullptr;
+	FPropertyEditorModule& Module = FModuleManager::LoadModuleChecked<FPropertyEditorModule>("PropertyEditor");
+	if (!Generator.IsValid())
+	{
+		FPropertyRowGeneratorArgs Args;
+		Generator = Module.CreatePropertyRowGenerator(Args);
+		// the sizes of the parameter lists are only based on the parent material and not changed out from under the details panel 
+		// When a parameter is added open MI editors are refreshed
+		// the tree should also refresh if one of the layer or blend assets is swapped
+
+		auto ValidationLambda = ([](const FRootPropertyNodeList& PropertyNodeList) { return true; });
+		Generator->SetCustomValidatePropertyNodesFunction(FOnValidatePropertyRowGeneratorNodes::CreateLambda(MoveTemp(ValidationLambda)));
+
+		TArray<UObject*> Objects;
+		Objects.Add(MaterialEditorInstance);
+		Generator->SetObjects(Objects);
+	}
+	else
+	{
+		TArray<UObject*> Objects;
+		Objects.Add(MaterialEditorInstance);
+		Generator->SetObjects(Objects);
+	}
+
+	TSharedPtr<IDetailTreeNode> ParameterGroups = FindParameterGroupsNode(Generator);
+	if (ParameterGroups.IsValid())
+	{
+		TArray<TSharedRef<IDetailTreeNode>> Children;
+		ParameterGroups->GetChildren(Children);
+		// the order of DeferredSearches should correspond to NonLayerProperties exactly
+		TArray<TSharedPtr<IPropertyHandle>> DeferredSearches;
+		for (int32 GroupIdx = 0; GroupIdx < Children.Num(); ++GroupIdx)
+		{
+			TArray<void*> GroupPtrs;
+			TSharedPtr<IPropertyHandle> ChildHandle = Children[GroupIdx]->CreatePropertyHandle();
+			ChildHandle->AccessRawData(GroupPtrs);
+			auto GroupIt = GroupPtrs.CreateConstIterator();
+			const FEditorParameterGroup* ParameterGroupPtr = reinterpret_cast<FEditorParameterGroup*>(*GroupIt);
+			const FEditorParameterGroup& ParameterGroup = *ParameterGroupPtr;
+
+			for (int32 ParamIdx = 0; ParamIdx < ParameterGroup.Parameters.Num(); ParamIdx++)
+			{
+				UDEditorParameterValue* Parameter = ParameterGroup.Parameters[ParamIdx];
+
+				TSharedPtr<IPropertyHandle> ParametersArrayProperty = ChildHandle->GetChildHandle("Parameters");
+				TSharedPtr<IPropertyHandle> ParameterProperty = ParametersArrayProperty->GetChildHandle(ParamIdx);
+				TSharedPtr<IPropertyHandle> ParameterValueProperty = ParameterProperty->GetChildHandle("ParameterValue");
+
+				if (Cast<UDEditorMaterialLayersParameterValue>(Parameter))
+				{
+					if (FunctionParameter == nullptr)
+					{
+						FunctionParameter = Parameter;
+					}
+					TArray<void*> StructPtrs;
+					ParameterValueProperty->AccessRawData(StructPtrs);
+					auto It = StructPtrs.CreateConstIterator();
+					FunctionInstance = reinterpret_cast<FMaterialLayersFunctions*>(*It);
+					FunctionInstanceHandle = ParameterValueProperty;
+
+					TSharedPtr<IPropertyHandle>	LayerHandle = ChildHandle->GetChildHandle("Layers").ToSharedRef();
+					TSharedPtr<IPropertyHandle> BlendHandle = ChildHandle->GetChildHandle("Blends").ToSharedRef();
+					uint32 LayerChildren;
+					LayerHandle->GetNumChildren(LayerChildren);
+					uint32 BlendChildren;
+					BlendHandle->GetNumChildren(BlendChildren);
+					if (MaterialEditorInstance->StoredLayerPreviews.Num() != LayerChildren)
+					{
+						MaterialEditorInstance->StoredLayerPreviews.Empty();
+						MaterialEditorInstance->StoredLayerPreviews.AddDefaulted(LayerChildren);
+					}
+					if (MaterialEditorInstance->StoredBlendPreviews.Num() != BlendChildren)
+					{
+						MaterialEditorInstance->StoredBlendPreviews.Empty();
+						MaterialEditorInstance->StoredBlendPreviews.AddDefaulted(BlendChildren);
+					}
+
+					TSharedRef<FSortedParamData> StackProperty = MakeShared<FSortedParamData>();
+					StackProperty->StackDataType = EStackDataType::Stack;
+					StackProperty->Parameter = Parameter;
+					StackProperty->ParameterInfo.Index = LayerChildren - 1;
+					StackProperty->NodeKey = FString::FromInt(StackProperty->ParameterInfo.Index);
+
+
+					TSharedRef<FSortedParamData> ChildProperty = MakeShared<FSortedParamData>();
+					ChildProperty->StackDataType = EStackDataType::Asset;
+					ChildProperty->Parameter = Parameter;
+					ChildProperty->ParameterHandle = LayerHandle->AsArray()->GetElement(LayerChildren - 1);
+					ChildProperty->ParameterNode = Generator->FindTreeNode(ChildProperty->ParameterHandle);
+					ChildProperty->ParameterInfo.Index = LayerChildren - 1;
+					ChildProperty->ParameterInfo.Association = EMaterialParameterAssociation::LayerParameter;
+					ChildProperty->NodeKey = FString::FromInt(ChildProperty->ParameterInfo.Index) + FString::FromInt(ChildProperty->ParameterInfo.Association);
+
+					{
+						UObject* AssetObject = nullptr;
+						ChildProperty->ParameterHandle->GetValue(AssetObject);
+						if (AssetObject)
+						{
+							if (MaterialEditorInstance->StoredLayerPreviews[LayerChildren - 1] == nullptr)
+							{
+								MaterialEditorInstance->StoredLayerPreviews[LayerChildren - 1] = (NewObject<UMaterialInstanceConstant>(MaterialEditorInstance, NAME_None));
+							}
+							UMaterialInterface* EditedMaterial = Cast<UMaterialFunctionInterface>(AssetObject)->GetPreviewMaterial();
+							if (MaterialEditorInstance->StoredLayerPreviews[LayerChildren - 1] && MaterialEditorInstance->StoredLayerPreviews[LayerChildren - 1]->Parent != EditedMaterial)
+							{
+								MaterialEditorInstance->StoredLayerPreviews[LayerChildren - 1]->SetParentEditorOnly(EditedMaterial);
+							}
+						}
+					}
+
+					StackProperty->Children.Add(ChildProperty);
+					LayerProperties.Add(StackProperty);
+
+					if (BlendChildren > 0 && LayerChildren > BlendChildren)
+					{
+						for (int32 Counter = BlendChildren - 1; Counter >= 0; Counter--)
+						{
+							ChildProperty = MakeShared<FSortedParamData>();
+							ChildProperty->StackDataType = EStackDataType::Asset;
+							ChildProperty->Parameter = Parameter;
+							ChildProperty->ParameterHandle = BlendHandle->AsArray()->GetElement(Counter);
+							ChildProperty->ParameterNode = Generator->FindTreeNode(ChildProperty->ParameterHandle);
+							ChildProperty->ParameterInfo.Index = Counter;
+							ChildProperty->ParameterInfo.Association = EMaterialParameterAssociation::BlendParameter;
+							ChildProperty->NodeKey = FString::FromInt(ChildProperty->ParameterInfo.Index) + FString::FromInt(ChildProperty->ParameterInfo.Association);
+							{
+								UObject* AssetObject = nullptr;
+								ChildProperty->ParameterHandle->GetValue(AssetObject);
+								if (AssetObject)
+								{
+									if (MaterialEditorInstance->StoredBlendPreviews[Counter] == nullptr)
+									{
+										MaterialEditorInstance->StoredBlendPreviews[Counter] = (NewObject<UMaterialInstanceConstant>(MaterialEditorInstance, NAME_None));
+									}
+									UMaterialInterface* EditedMaterial = Cast<UMaterialFunctionInterface>(AssetObject)->GetPreviewMaterial();
+									if (MaterialEditorInstance->StoredBlendPreviews[Counter] && MaterialEditorInstance->StoredBlendPreviews[Counter]->Parent != EditedMaterial)
+									{
+										MaterialEditorInstance->StoredBlendPreviews[Counter]->SetParentEditorOnly(EditedMaterial);
+									}
+								}
+							}
+							LayerProperties.Last()->Children.Add(ChildProperty);
+
+							StackProperty = MakeShared<FSortedParamData>();
+							StackProperty->StackDataType = EStackDataType::Stack;
+							StackProperty->Parameter = Parameter;
+							StackProperty->ParameterInfo.Index = Counter;
+							StackProperty->NodeKey = FString::FromInt(StackProperty->ParameterInfo.Index);
+							LayerProperties.Add(StackProperty);
+
+							ChildProperty = MakeShared<FSortedParamData>();
+							ChildProperty->StackDataType = EStackDataType::Asset;
+							ChildProperty->Parameter = Parameter;
+							ChildProperty->ParameterHandle = LayerHandle->AsArray()->GetElement(Counter);
+							ChildProperty->ParameterNode = Generator->FindTreeNode(ChildProperty->ParameterHandle);
+							ChildProperty->ParameterInfo.Index = Counter;
+							ChildProperty->ParameterInfo.Association = EMaterialParameterAssociation::LayerParameter;
+							ChildProperty->NodeKey = FString::FromInt(ChildProperty->ParameterInfo.Index) + FString::FromInt(ChildProperty->ParameterInfo.Association);
+							{
+								UObject* AssetObject = nullptr;
+								ChildProperty->ParameterHandle->GetValue(AssetObject);
+								if (AssetObject)
+								{
+									if (MaterialEditorInstance->StoredLayerPreviews[Counter] == nullptr)
+									{
+										MaterialEditorInstance->StoredLayerPreviews[Counter] = (NewObject<UMaterialInstanceConstant>(MaterialEditorInstance, NAME_None));
+									}
+									UMaterialInterface* EditedMaterial = Cast<UMaterialFunctionInterface>(AssetObject)->GetPreviewMaterial();
+									if (MaterialEditorInstance->StoredLayerPreviews[Counter] && MaterialEditorInstance->StoredLayerPreviews[Counter]->Parent != EditedMaterial)
+									{
+										MaterialEditorInstance->StoredLayerPreviews[Counter]->SetParentEditorOnly(EditedMaterial);
+									}
+								}
+							}
+							LayerProperties.Last()->Children.Add(ChildProperty);
+						}
+					}
+				}
+				else
+				{
+					FUnsortedParamData NonLayerProperty;
+					UDEditorScalarParameterValue* ScalarParam = Cast<UDEditorScalarParameterValue>(Parameter);
+
+					if (ScalarParam && ScalarParam->SliderMax > ScalarParam->SliderMin)
+					{
+						ParameterValueProperty->SetInstanceMetaData("UIMin", FString::Printf(TEXT("%f"), ScalarParam->SliderMin));
+						ParameterValueProperty->SetInstanceMetaData("UIMax", FString::Printf(TEXT("%f"), ScalarParam->SliderMax));
+					}
+
+					NonLayerProperty.Parameter = Parameter;
+					NonLayerProperty.ParameterGroup = ParameterGroup;
+
+					DeferredSearches.Add(ParameterValueProperty);
+					NonLayerProperty.UnsortedName = Parameter->ParameterInfo.Name;
+
+					NonLayerProperties.Add(NonLayerProperty);
+				}
+			}
+		}
+
+		checkf(NonLayerProperties.Num() == DeferredSearches.Num(), TEXT("Internal inconsistency: number of node searches does not match the number of properties"));
+		TArray<TSharedPtr<IDetailTreeNode>> DeferredResults = Generator->FindTreeNodes(DeferredSearches);
+		checkf(NonLayerProperties.Num() == DeferredResults.Num(), TEXT("Internal inconsistency: number of node search results does not match the number of properties"));
+
+		for (int Idx = 0, NumUnsorted = NonLayerProperties.Num(); Idx < NumUnsorted; ++Idx)
+		{
+			FUnsortedParamData& NonLayerProperty = NonLayerProperties[Idx];
+			NonLayerProperty.ParameterNode = DeferredResults[Idx];
+			NonLayerProperty.ParameterHandle = NonLayerProperty.ParameterNode->CreatePropertyHandle();
+		}
+
+		DeferredResults.Empty();
+		DeferredSearches.Empty();
+
+		for (int32 LayerIdx = 0; LayerIdx < LayerProperties.Num(); LayerIdx++)
+		{
+			for (int32 ChildIdx = 0; ChildIdx < LayerProperties[LayerIdx]->Children.Num(); ChildIdx++)
+			{
+				ShowSubParameters(LayerProperties[LayerIdx]->Children[ChildIdx]);
+			}
+		}
+	}
+
+	SetParentsExpansionState();
+}
+
+bool SMaterialSubstrateTree::IsLayerVisible(int32 Index) const
+{
+	if (FunctionParameter.IsValid())
+	{
+		return FunctionInstance->GetLayerVisibility(Index);
+	}
+	return false;
+}
+
+TSharedRef<SWidget> SMaterialSubstrateTree::CreateThumbnailWidget(EMaterialParameterAssociation InAssociation, int32 InIndex, float InThumbnailSize)
+{
+	UObject* ThumbnailObject = nullptr;
+	if (InAssociation == EMaterialParameterAssociation::LayerParameter)
+	{
+		ThumbnailObject = MaterialEditorInstance->StoredLayerPreviews[InIndex];
+	}
+	else if (InAssociation == EMaterialParameterAssociation::BlendParameter)
+	{
+		ThumbnailObject = MaterialEditorInstance->StoredBlendPreviews[InIndex];
+	}
+	const TSharedPtr<FAssetThumbnail> AssetThumbnail = MakeShareable(new FAssetThumbnail(ThumbnailObject, InThumbnailSize, InThumbnailSize, GetTreeThumbnailPool()));
+	TSharedRef<SWidget> ThumbnailWidget = AssetThumbnail->MakeThumbnailWidget();
+	AssetThumbnail->SetRealTime(true);
+	ThumbnailWidget->SetOnMouseDoubleClick(FPointerEventHandler::CreateSP(this, &SMaterialSubstrateTree::OnThumbnailDoubleClick, InAssociation, InIndex));
+	return ThumbnailWidget;
+}
+
+void SMaterialSubstrateTree::UpdateThumbnailMaterial(TEnumAsByte<EMaterialParameterAssociation> InAssociation, int32 InIndex, bool bAlterBlendIndex)
+{
+	// Need to invert index b/c layer properties is generated in reverse order
+	TArray<TSharedPtr<FSortedParamData>> AssetChildren = LayerProperties[LayerProperties.Num() - 1 - InIndex]->Children;
+	UMaterialInstanceConstant* MaterialToUpdate = nullptr;
+	int32 ParameterIndex = InIndex;
+	if (InAssociation == EMaterialParameterAssociation::LayerParameter)
+	{
+		MaterialToUpdate = MaterialEditorInstance->StoredLayerPreviews[ParameterIndex];
+	}
+	if (InAssociation == EMaterialParameterAssociation::BlendParameter)
+	{
+		if (bAlterBlendIndex)
+		{
+			ParameterIndex--;
+		}
+		MaterialToUpdate = MaterialEditorInstance->StoredBlendPreviews[ParameterIndex];
+	}
+
+	TArray<FEditorParameterGroup> ParameterGroups;
+	for (TSharedPtr<FSortedParamData> AssetChild : AssetChildren)
+	{
+		for (TSharedPtr<FSortedParamData> Group : AssetChild->Children)
+		{
+			if (Group->ParameterInfo.Association == InAssociation)
+			{
+				FEditorParameterGroup DuplicatedGroup = FEditorParameterGroup();
+				DuplicatedGroup.GroupAssociation = Group->Group.GroupAssociation;
+				DuplicatedGroup.GroupName = Group->Group.GroupName;
+				DuplicatedGroup.GroupSortPriority = Group->Group.GroupSortPriority;
+				for (UDEditorParameterValue* Parameter : Group->Group.Parameters)
+				{
+					if (Parameter->ParameterInfo.Index == ParameterIndex)
+					{
+						DuplicatedGroup.Parameters.Add(Parameter);
+					}
+				}
+				ParameterGroups.Add(DuplicatedGroup);
+			}
+		}
+
+	
+
+	}
+	if (MaterialToUpdate != nullptr)
+	{
+		FMaterialPropertyHelpers::TransitionAndCopyParameters(MaterialToUpdate, ParameterGroups, true);
+	}
+}
+
+FReply SMaterialSubstrateTree::OnThumbnailDoubleClick(const FGeometry& Geometry, const FPointerEvent& MouseEvent, EMaterialParameterAssociation InAssociation, int32 InIndex)
+{
+	UMaterialFunctionInterface* AssetToOpen = nullptr;
+	if (InAssociation == EMaterialParameterAssociation::BlendParameter)
+	{
+		AssetToOpen = FunctionInstance->Blends[InIndex];
+	}
+	else if (InAssociation == EMaterialParameterAssociation::LayerParameter)
+	{
+		AssetToOpen = FunctionInstance->Layers[InIndex];
+	}
+	if (AssetToOpen != nullptr)
+	{
+		GEditor->GetEditorSubsystem<UAssetEditorSubsystem>()->OpenEditorForAsset(AssetToOpen);
+		return FReply::Handled();
+	}
+	return FReply::Unhandled();
+}
+
+void SMaterialSubstrateTree::ShowSubParameters(TSharedPtr<FSortedParamData> ParentParameter)
+{
+	for (FUnsortedParamData Property : NonLayerProperties)
+	{
+		UDEditorParameterValue* Parameter = Property.Parameter;
+		if (Parameter->ParameterInfo.Index == ParentParameter->ParameterInfo.Index
+			&& Parameter->ParameterInfo.Association == ParentParameter->ParameterInfo.Association)
+		{
+			TSharedPtr<FSortedParamData> GroupProperty(new FSortedParamData());
+			GroupProperty->StackDataType = EStackDataType::Group;
+			GroupProperty->ParameterInfo.Index = Parameter->ParameterInfo.Index;
+			GroupProperty->ParameterInfo.Association = Parameter->ParameterInfo.Association;
+			GroupProperty->Group = Property.ParameterGroup;
+			GroupProperty->NodeKey = FString::FromInt(GroupProperty->ParameterInfo.Index) + FString::FromInt(GroupProperty->ParameterInfo.Association) + Property.ParameterGroup.GroupName.ToString();
+
+			bool bAddNewGroup = true;
+			for (TSharedPtr<struct FSortedParamData> GroupChild : ParentParameter->Children)
+			{
+				if (GroupChild->NodeKey == GroupProperty->NodeKey)
+				{
+					bAddNewGroup = false;
+				}
+			}
+			if (bAddNewGroup)
+			{
+				ParentParameter->Children.Add(GroupProperty);
+			}
+
+			TSharedPtr<FSortedParamData> ChildProperty(new FSortedParamData());
+			ChildProperty->StackDataType = EStackDataType::Property;
+			ChildProperty->Parameter = Parameter;
+			ChildProperty->ParameterInfo.Index = Parameter->ParameterInfo.Index;
+			ChildProperty->ParameterInfo.Association = Parameter->ParameterInfo.Association;
+			ChildProperty->ParameterNode = Property.ParameterNode;
+			ChildProperty->PropertyName = Property.UnsortedName;
+			ChildProperty->NodeKey = FString::FromInt(ChildProperty->ParameterInfo.Index) + FString::FromInt(ChildProperty->ParameterInfo.Association) +  Property.ParameterGroup.GroupName.ToString() + Property.UnsortedName.ToString();
+
+
+			UDEditorStaticComponentMaskParameterValue* CompMaskParam = Cast<UDEditorStaticComponentMaskParameterValue>(Parameter);
+			if (!CompMaskParam)
+			{
+				TArray<TSharedRef<IDetailTreeNode>> ParamChildren;
+				Property.ParameterNode->GetChildren(ParamChildren);
+				for (int32 ParamChildIdx = 0; ParamChildIdx < ParamChildren.Num(); ParamChildIdx++)
+				{
+					TSharedPtr<FSortedParamData> ParamChildProperty(new FSortedParamData());
+					ParamChildProperty->StackDataType = EStackDataType::PropertyChild;
+					ParamChildProperty->ParameterNode = ParamChildren[ParamChildIdx];
+					ParamChildProperty->ParameterHandle = ParamChildProperty->ParameterNode->CreatePropertyHandle();
+					ParamChildProperty->ParameterInfo.Index = Parameter->ParameterInfo.Index;
+					ParamChildProperty->ParameterInfo.Association = Parameter->ParameterInfo.Association;
+					ParamChildProperty->Parameter = ChildProperty->Parameter;
+					ChildProperty->Children.Add(ParamChildProperty);
+				}
+			}
+			for (TSharedPtr<struct FSortedParamData> GroupChild : ParentParameter->Children)
+			{
+				if (GroupChild->Group.GroupName == Property.ParameterGroup.GroupName
+					&& GroupChild->ParameterInfo.Association == ChildProperty->ParameterInfo.Association
+					&&  GroupChild->ParameterInfo.Index == ChildProperty->ParameterInfo.Index)
+				{
+					GroupChild->Children.Add(ChildProperty);
+				}
+			}
+
+		}
+	}
+}
+
+#undef LOCTEXT_NAMESPACE
\ No newline at end of file
diff --git a/Engine/Source/Editor/MaterialEditor/Private/SMaterialSubstrateTree.h b/Engine/Source/Editor/MaterialEditor/Private/SMaterialSubstrateTree.h
new file mode 100644
--- /dev/null
+++ b/Engine/Source/Editor/MaterialEditor/Private/SMaterialSubstrateTree.h
@@ -0,0 +1,215 @@
+// Copyright Epic Games, Inc. All Rights Reserved.
+
+#pragma once
+
+#include "CoreMinimal.h"
+#include "DetailColumnSizeData.h"
+#include "Engine/EngineTypes.h"
+#include "IDetailPropertyRow.h"
+#include "IDetailTreeNode.h"
+#include "MaterialEditor/MaterialEditorInstanceConstant.h"
+#include "MaterialPropertyHelpers.h"
+#include "Materials/Material.h"
+#include "PropertyCustomizationHelpers.h"
+#include "UObject/Object.h"
+#include "UObject/ObjectMacros.h"
+#include "UObject/WeakObjectPtrTemplates.h"
+#include "Widgets/DeclarativeSyntaxSupport.h"
+#include "Widgets/SCompoundWidget.h"
+#include "Widgets/Views/STreeView.h"
+
+class IPropertyHandle;
+class SMaterialSubstrateTree;
+class SMaterialLayersFunctionsInstanceWrapper;
+class UDEditorParameterValue;
+class UMaterialEditorInstanceConstant;
+
+typedef TSharedPtr<FSortedParamData> FSortedParamDataPtr;
+class SMaterialSubstrateTree : public STreeView<FSortedParamDataPtr>
+{
+	friend class SMaterialSubstrateTreeItem;
+public:
+	SLATE_BEGIN_ARGS(SMaterialSubstrateTree)
+		: _InMaterialEditorInstance(nullptr)
+	{}
+
+	SLATE_ARGUMENT(UMaterialEditorInstanceConstant*, InMaterialEditorInstance)
+	SLATE_ARGUMENT(SMaterialLayersFunctionsInstanceWrapper*, InWrapper)
+	SLATE_ARGUMENT(FGetShowHiddenParameters, InShowHiddenDelegate)
+	SLATE_END_ARGS()
+
+	/** Constructs this widget with InArgs */
+	void Construct(const FArguments& InArgs);
+	TSharedRef< ITableRow > OnGenerateRowMaterialLayersFunctionsTreeView(TSharedPtr<FSortedParamData> Item, const TSharedRef< STableViewBase >& OwnerTable);
+	void OnGetChildrenMaterialLayersFunctionsTreeView(TSharedPtr<FSortedParamData> InParent, TArray< TSharedPtr<FSortedParamData> >& OutChildren);
+	void OnExpansionChanged(TSharedPtr<FSortedParamData> Item, bool bIsExpanded);
+	void OnSelectionChangedMaterialSubstrateView(TSharedPtr<FSortedParamData> InSelectedItem, ESelectInfo::Type SelectInfo);
+	void SetParentsExpansionState();
+
+	void ShowHiddenValues(bool& bShowHiddenParameters) { bShowHiddenParameters = true; }
+	TWeakObjectPtr<class UDEditorParameterValue> FunctionParameter;
+	struct FMaterialLayersFunctions* FunctionInstance;
+	TSharedPtr<IPropertyHandle> FunctionInstanceHandle;
+	void RefreshOnAssetChange(const struct FAssetData& InAssetData, int32 Index, EMaterialParameterAssociation MaterialType);
+	void ResetAssetToDefault(TSharedPtr<FSortedParamData> InData);
+	void AddLayer();
+	void InsertLayerFromAsset(const struct FAssetData& InAssetData, int32 Index, EMaterialParameterAssociation MaterialType);
+	void RemoveLayer(int32 Index);
+	FReply UnlinkLayer(int32 Index);
+	FReply RelinkLayersToParent();
+	EVisibility GetUnlinkLayerVisibility(int32 Index) const;
+	EVisibility GetRelinkLayersToParentVisibility() const;
+	FReply ToggleLayerVisibility(int32 Index);
+	TSharedPtr<class FAssetThumbnailPool> GetTreeThumbnailPool();
+
+	/** Object that stores all of the possible parameters we can edit */
+	UMaterialEditorInstanceConstant* MaterialEditorInstance;
+
+	/** Builds the custom parameter groups category */
+	void CreateGroupsWidget();
+
+	bool IsLayerVisible(int32 Index) const;
+
+	SMaterialLayersFunctionsInstanceWrapper* GetWrapper() { return Wrapper; }
+
+	TSharedPtr<IDetailTreeNode> FindParameterGroupsNode(TSharedPtr<IPropertyRowGenerator> PropertyRowGenerator);
+
+	TSharedRef<SWidget> CreateThumbnailWidget(EMaterialParameterAssociation InAssociation, int32 InIndex, float InThumbnailSize);
+	void UpdateThumbnailMaterial(TEnumAsByte<EMaterialParameterAssociation> InAssociation, int32 InIndex, bool bAlterBlendIndex = false);
+	FReply OnThumbnailDoubleClick(const FGeometry& Geometry, const FPointerEvent& MouseEvent, EMaterialParameterAssociation InAssociation, int32 InIndex);
+	bool IsOverriddenExpression(class UDEditorParameterValue* Parameter, int32 InIndex);
+	bool IsOverriddenExpression(TObjectPtr<UDEditorParameterValue> Parameter, int32 InIndex);
+	
+	FGetShowHiddenParameters GetShowHiddenDelegate() const;
+protected:
+
+	void ShowSubParameters(TSharedPtr<FSortedParamData> ParentParameter);
+
+private:
+	TArray<TSharedPtr<FSortedParamData>> LayerProperties;
+	
+	TArray<FUnsortedParamData> NonLayerProperties;
+	
+	FDetailColumnSizeData ColumnSizeData;
+	
+	SMaterialLayersFunctionsInstanceWrapper* Wrapper;
+	
+	TSharedPtr<class IPropertyRowGenerator> Generator;
+	
+	bool bLayerIsolated;
+	
+	/** Delegate to call to determine if hidden parameters should be shown */
+	FGetShowHiddenParameters ShowHiddenDelegate;
+};
+
+class SMaterialSubstrateTreeItem : public STableRow< FSortedParamDataPtr >
+{
+public:
+
+	SLATE_BEGIN_ARGS(SMaterialSubstrateTreeItem)
+		: _StackParameterData(nullptr),
+		_MaterialEditorInstance(nullptr),
+		_InTree(nullptr),
+		_Padding( FMargin(0) )
+	{}
+
+	/** The item content. */
+	SLATE_ARGUMENT(FSortedParamDataPtr, StackParameterData)
+	SLATE_ARGUMENT(UMaterialEditorInstanceConstant*, MaterialEditorInstance)
+	SLATE_ARGUMENT(SMaterialSubstrateTree*, InTree)
+	SLATE_ATTRIBUTE( FMargin, Padding )
+	SLATE_END_ARGS()
+
+	bool bIsBeingDragged = false;
+
+private:
+	bool bIsHoveredDragTarget = false;
+
+
+	FString GetCurvePath(class UDEditorScalarParameterValue* Parameter) const;
+	const FSlateBrush* GetBorderImage() const;
+
+	FSlateColor GetOuterBackgroundColor(TSharedPtr<FSortedParamData> InParamData) const;
+public:
+
+	void RefreshOnRowChange(const FAssetData& AssetData, SMaterialSubstrateTree* InTree);
+	bool GetFilterState(SMaterialSubstrateTree* InTree, TSharedPtr<FSortedParamData> InStackData) const;
+	void FilterClicked(const ECheckBoxState NewCheckedState, SMaterialSubstrateTree* InTree, TSharedPtr<FSortedParamData> InStackData);
+	ECheckBoxState GetFilterChecked(SMaterialSubstrateTree* InTree, TSharedPtr<FSortedParamData> InStackData) const;
+	FText GetLayerName(SMaterialSubstrateTree* InTree, int32 Counter) const;
+	void OnNameChanged(const FText& InText, ETextCommit::Type CommitInfo, SMaterialSubstrateTree* InTree, int32 Counter);
+
+
+	void OnLayerDragEnter(const FDragDropEvent& DragDropEvent)
+	{
+		if (StackParameterData->ParameterInfo.Index != 0)
+		{
+			bIsHoveredDragTarget = true;
+		}
+	}
+
+	void OnLayerDragLeave(const FDragDropEvent& DragDropEvent)
+	{
+		bIsHoveredDragTarget = false;
+	}
+
+	void OnLayerDragDetected()
+	{
+		bIsBeingDragged = true;
+	}
+
+	FReply OnLayerDrop(const FDragDropEvent& DragDropEvent);
+	void OnOverrideParameter(bool NewValue, class UDEditorParameterValue* Parameter);
+	void OnOverrideParameter(bool NewValue, TObjectPtr<UDEditorParameterValue> Parameter);
+	/**
+	* Construct the widget
+	*
+	* @param InArgs   A declaration from which to construct the widget
+	*/
+	void Construct(const FArguments& InArgs, const TSharedRef<STableViewBase>& InOwnerTableView);
+	int32 OnLayerItemPaintDropIndicator(EItemDropZone InItemDropZone, const FPaintArgs& Args, const FGeometry& AllottedGeometry, const FSlateRect& MyCullingRect, FSlateWindowElementList& OutDrawElements, int32 LayerId, const FWidgetStyle& InWidgetStyle, bool bParentEnabled) const;
+
+	/** The node info to build the tree view row from. */
+	FSortedParamDataPtr StackParameterData;
+
+	SMaterialSubstrateTree* Tree;
+
+	UMaterialEditorInstanceConstant* MaterialEditorInstance;
+
+	FString GetInstancePath(SMaterialSubstrateTree* InTree) const;
+};
+
+class FSubstrateLayerDragDropOp final : public FDecoratedDragDropOp
+{
+public:
+	DRAG_DROP_OPERATOR_TYPE(FSubstrateLayerDragDropOp, FDecoratedDragDropOp)
+
+	FSubstrateLayerDragDropOp(TSharedPtr<SMaterialSubstrateTreeItem> InOwningStack)
+	{
+		OwningStack = InOwningStack;
+		DecoratorWidget = SNew(SBorder)
+			.BorderImage(FAppStyle::GetBrush("Graph.ConnectorFeedback.Border"))
+			.Content()
+			[
+				SNew(SHorizontalBox)
+				+ SHorizontalBox::Slot()
+				.AutoWidth()
+				.VAlign(VAlign_Center)
+				[
+					SNew(STextBlock)
+					.Text(NSLOCTEXT("LayerDragDrop", "PlaceLayerHere", "Place Layer and Blend Here"))
+				]
+			];
+
+		Construct();
+	};
+
+	TSharedPtr<SWidget> DecoratorWidget;
+
+	virtual TSharedPtr<SWidget> GetDefaultDecorator() const override
+	{
+		return DecoratorWidget;
+	}
+
+	TWeakPtr<class SMaterialSubstrateTreeItem> OwningStack;
+};
