- Updated default layout for Material Instance when showing the new Substrate Tree UI
- Fixed null check issue with detail customization

diff --git a/Engine/Source/Editor/MaterialEditor/Private/MaterialEditorInstanceDetailCustomization.cpp b/Engine/Source/Editor/MaterialEditor/Private/MaterialEditorInstanceDetailCustomization.cpp
--- a/Engine/Source/Editor/MaterialEditor/Private/MaterialEditorInstanceDetailCustomization.cpp
+++ b/Engine/Source/Editor/MaterialEditor/Private/MaterialEditorInstanceDetailCustomization.cpp
@@ -107,7 +107,7 @@ void FMaterialInstanceParameterDetails::CustomizeDetails(IDetailLayoutBuilder& D
 	TSharedRef<IPropertyHandle> ParameterGroupsProperty = DetailLayout.GetProperty("ParameterGroups");
 
 	// check if tree has any selection, we show parameter properties for selected layer only
-	if (this->MaterialLayersFunctionsInstance->NestedTree->GetNumItemsSelected() > 0)
+	if (MaterialLayersFunctionsInstance != nullptr && MaterialLayersFunctionsInstance->NestedTree->GetNumItemsSelected() > 0)
 	{
 		// for each selected FSortedParamData item (type stack)
 		for (TSharedPtr<FSortedParamData> SelectedItem : this->MaterialLayersFunctionsInstance->NestedTree->GetSelectedItems())
diff --git a/Engine/Source/Editor/MaterialEditor/Private/MaterialInstanceEditor.cpp b/Engine/Source/Editor/MaterialEditor/Private/MaterialInstanceEditor.cpp
--- a/Engine/Source/Editor/MaterialEditor/Private/MaterialInstanceEditor.cpp
+++ b/Engine/Source/Editor/MaterialEditor/Private/MaterialInstanceEditor.cpp
@@ -48,6 +48,7 @@
 #include "UnrealEdGlobals.h"
 #include "Editor/UnrealEdEngine.h"
 #include "AdvancedPreviewSceneModule.h"
+#include "ContentBrowserModule.h"
 #include "Misc/MessageDialog.h"
 #include "Framework/Commands/UICommandInfo.h"
 #include "Styling/AppStyle.h"
@@ -55,6 +56,9 @@
 #include "MaterialEditingLibrary.h"
 #include "Widgets/Layout/SScrollBox.h"
 #include "DebugViewModeHelpers.h"
+#include "IContentBrowserSingleton.h"
+#include "Materials/MaterialFunctionMaterialLayer.h"
+#include "Materials/MaterialFunctionMaterialLayerBlend.h"
 #include "VT/RuntimeVirtualTexture.h"
 #include "SparseVolumeTexture/SparseVolumeTexture.h"
 #include "Widgets/Input/SButton.h"
@@ -68,6 +72,7 @@ const FName FMaterialInstanceEditor::PreviewTabId( TEXT( "MaterialInstanceEditor
 const FName FMaterialInstanceEditor::PropertiesTabId( TEXT( "MaterialInstanceEditor_MaterialProperties" ) );
 const FName FMaterialInstanceEditor::LayerPropertiesTabId(TEXT("MaterialInstanceEditor_MaterialLayerProperties"));
 const FName FMaterialInstanceEditor::PreviewSettingsTabId(TEXT("MaterialInstanceEditor_PreviewSettings"));
+const FName FMaterialInstanceEditor::AssetBrowserTabId(TEXT("MaterialInstanceEditor_AssetBrowser"));
 
 extern TAutoConsoleVariable<bool> CVarMaterialEdAllowIgnoringCompilationErrors;
 
@@ -260,6 +265,12 @@ void FMaterialInstanceEditor::RegisterTabSpawners(const TSharedRef<class FTabMan
 		.SetGroup(WorkspaceMenuCategoryRef)
 		.SetIcon(FSlateIcon(FAppStyle::GetAppStyleSetName(), "LevelEditor.Tabs.Details"));
 
+#if ENABLE_MATERIAL_LAYER_PROTOTYPE
+	InTabManager->RegisterTabSpawner(AssetBrowserTabId, FOnSpawnTab::CreateSP(this, &FMaterialInstanceEditor::SpawnTab_AssetBrowser))
+		.SetDisplayName(LOCTEXT("AssetBrowserTab", "Asset Browser"))
+		.SetGroup(WorkspaceMenuCategoryRef)
+		.SetIcon(FSlateIcon(FAppStyle::GetAppStyleSetName(), "LevelEditor.Tabs.Details"));
+#endif
 	MaterialStatsManager->RegisterTabs();
 
 	OnRegisterTabSpawners().Broadcast(InTabManager);
@@ -276,7 +287,9 @@ void FMaterialInstanceEditor::UnregisterTabSpawners(const TSharedRef<class FTabM
 		InTabManager->UnregisterTabSpawner(LayerPropertiesTabId);
 	}
 	InTabManager->UnregisterTabSpawner( PreviewSettingsTabId );
-
+#if ENABLE_MATERIAL_LAYER_PROTOTYPE
+	InTabManager->UnregisterTabSpawner( AssetBrowserTabId );
+#endif
 	MaterialStatsManager->UnregisterTabs();
 
 	OnUnregisterTabSpawners().Broadcast(InTabManager);
@@ -453,7 +466,48 @@ void FMaterialInstanceEditor::InitMaterialInstanceEditor( const EToolkitMode::Ty
 
 	if (!bIsFunctionPreviewMaterial)
 	{
-		StandaloneDefaultLayout = FTabManager::NewLayout("Standalone_MaterialInstanceEditor_Layout_v8")
+		
+#if ENABLE_MATERIAL_LAYER_PROTOTYPE
+		StandaloneDefaultLayout = FTabManager::NewLayout("Standalone_MaterialInstanceEditor_Layout_v9")
+			->AddArea
+			(
+				FTabManager::NewPrimaryArea()->SetOrientation(Orient_Vertical)
+				->Split
+				(
+					FTabManager::NewSplitter()->SetOrientation(Orient_Horizontal)->SetSizeCoefficient(0.9f)
+					->Split
+					(
+						FTabManager::NewStack()->SetSizeCoefficient(0.70f)->SetHideTabWell(true)
+						->AddTab(PreviewTabId, ETabState::OpenedTab)
+						->AddTab(PreviewSettingsTabId, ETabState::ClosedTab)
+					)
+					->Split
+					(
+						FTabManager::NewSplitter()->SetOrientation(Orient_Vertical)->SetSizeCoefficient(0.7f)
+						->Split
+						(
+					FTabManager::NewSplitter()->SetOrientation(Orient_Horizontal)->SetSizeCoefficient(0.3f)
+							->Split
+							(
+								FTabManager::NewStack()/*->SetSizeCoefficient(0.30f)*/
+								->AddTab(LayerPropertiesTabId, ETabState::OpenedTab)
+							)
+							->Split
+							(
+								FTabManager::NewStack()/*->SetSizeCoefficient(0.30f)*/
+								->AddTab(AssetBrowserTabId, ETabState::OpenedTab)
+							)
+						)
+						->Split
+						(
+						FTabManager::NewStack()->SetSizeCoefficient(0.30f)
+								->AddTab(PropertiesTabId, ETabState::OpenedTab)
+						)
+					)
+				)
+			);
+#else
+			StandaloneDefaultLayout = FTabManager::NewLayout("Standalone_MaterialInstanceEditor_Layout_v8")
 			->AddArea
 			(
 				FTabManager::NewPrimaryArea()->SetOrientation(Orient_Vertical)
@@ -475,8 +529,9 @@ void FMaterialInstanceEditor::InitMaterialInstanceEditor( const EToolkitMode::Ty
 					)
 				)
 			);
+#endif
 		}
-
+	
 	const bool bCreateDefaultStandaloneMenu = true;
 	const bool bCreateDefaultToolbar = true;
 	TArray<UObject*> ObjectsToEdit;
@@ -1065,6 +1120,37 @@ TSharedRef<SDockTab> FMaterialInstanceEditor::SpawnTab_PreviewSettings(const FSp
 	return SpawnedTab;
 }
 
+TSharedRef<SDockTab> FMaterialInstanceEditor::SpawnTab_AssetBrowser(const FSpawnTabArgs& Args)
+{
+	check(Args.GetTabId() == AssetBrowserTabId);
+
+	// TSharedRef<SWidget> InWidget = SNullWidget::NullWidget;
+	FContentBrowserModule& ContentBrowserModule = FModuleManager::Get().LoadModuleChecked<FContentBrowserModule>(TEXT("ContentBrowser"));
+
+	// Configure filter for asset picker
+	FAssetPickerConfig Config;
+	Config.SelectionMode = ESelectionMode::Single;
+	Config.Filter.ClassPaths.Add(UMaterialFunctionMaterialLayer::StaticClass()->GetClassPathName());
+	Config.Filter.ClassPaths.Add(UMaterialFunctionMaterialLayerInstance::StaticClass()->GetClassPathName());
+	Config.Filter.ClassPaths.Add(UMaterialFunctionMaterialLayerBlend::StaticClass()->GetClassPathName());
+	Config.Filter.ClassPaths.Add(UMaterialFunctionMaterialLayerBlendInstance::StaticClass()->GetClassPathName());
+	Config.bAddFilterUI = true;
+	Config.ThumbnailScale = 0.4f;
+	Config.InitialThumbnailSize = EThumbnailSize::Small;
+	Config.InitialAssetViewType = EAssetViewType::Tile;
+	
+	TSharedRef<SDockTab> SpawnedTab = SNew(SDockTab)
+		.Label(LOCTEXT("AssetBrowserTab", "Asset Browser"))
+		[
+			SNew(SBox)
+			[
+				ContentBrowserModule.Get().CreateAssetPicker(Config)
+			]
+		];
+
+	return SpawnedTab;
+}
+
 void FMaterialInstanceEditor::AddToSpawnedToolPanels(const FName& TabIdentifier, const TSharedRef<SDockTab>& SpawnedTab)
 {
 	TWeakPtr<SDockTab>* TabSpot = SpawnedToolPanels.Find(TabIdentifier);
diff --git a/Engine/Source/Editor/MaterialEditor/Private/MaterialInstanceEditor.h b/Engine/Source/Editor/MaterialEditor/Private/MaterialInstanceEditor.h
--- a/Engine/Source/Editor/MaterialEditor/Private/MaterialInstanceEditor.h
+++ b/Engine/Source/Editor/MaterialEditor/Private/MaterialInstanceEditor.h
@@ -190,6 +190,9 @@ private:
 
 	/** Spawns the advanced preview settings tab */
 	TSharedRef<SDockTab> SpawnTab_PreviewSettings(const FSpawnTabArgs& Args);
+	
+	/** Spawns the asset browser tab */
+	TSharedRef<SDockTab> SpawnTab_AssetBrowser(const FSpawnTabArgs& Args);
 
 	/**	Caches the specified tab for later retrieval */
 	void AddToSpawnedToolPanels( const FName& TabIdentifier, const TSharedRef<SDockTab>& SpawnedTab );
@@ -270,6 +273,7 @@ private:
 	static const FName PropertiesTabId;	
 	static const FName LayerPropertiesTabId;
 	static const FName PreviewSettingsTabId;
+	static const FName AssetBrowserTabId;
 
 	/** Object used as material statistics manager */
 	TSharedPtr<class FMaterialStats> MaterialStatsManager;
