SMaterialSubstrateTree visual updates

diff --git a/Engine/Source/Editor/MaterialEditor/Private/MaterialEditorStyle.cpp b/Engine/Source/Editor/MaterialEditor/Private/MaterialEditorStyle.cpp
--- a/Engine/Source/Editor/MaterialEditor/Private/MaterialEditorStyle.cpp
+++ b/Engine/Source/Editor/MaterialEditor/Private/MaterialEditorStyle.cpp
@@ -64,6 +64,10 @@ const FSlateBrush* FSubstrateMaterialEditorStyle::GetBrush(const FName& InName)
 #define BOX_BRUSH_SVG(RelativePath, ...) FSlateVectorBoxBrush(Style->RootToContentDir(RelativePath, TEXT(".svg")), __VA_ARGS__)
 #define BORDER_BRUSH_SVG(RelativePath, ...) FSlateVectorBorderBrush(Style->RootToContentDir(RelativePath, TEXT(".svg")), __VA_ARGS__)
 
+#define CORE_IMAGE_BRUSH_SVG( RelativePath, ... ) FSlateVectorImageBrush(Style->RootToCoreContentDir(RelativePath, TEXT(".svg")), __VA_ARGS__)
+#define CORE_BOX_BRUSH_SVG( RelativePath, ... ) FSlateVectorBoxBrush(Style->RootToCoreContentDir(RelativePath, TEXT(".svg")), __VA_ARGS__)
+#define CORE_BORDER_BRUSH_SVG( RelativePath, ... ) FSlateVectorBorderBrush(Style->RootToCoreContentDir(RelativePath, TEXT(".svg")), __VA_ARGS__)
+
 #define DEFAULT_FONT(...) FCoreStyle::GetDefaultFontStyle(__VA_ARGS__)
 
 #define IMAGE_PLUGIN_BRUSH(RelativePath, ...) FSlateImageBrush(FSubstrateMaterialEditorStyle::InContent(RelativePath, ".png"), __VA_ARGS__)
@@ -211,11 +215,12 @@ void FSubstrateMaterialEditorStyle::SetupLayerViewStyles(const TSharedRef<FSlate
 		.SetBackgroundBrush(*Style->GetBrush("LayerView.Background"))
 	);
 
-	const float LayerViewItemCornerRadius = 5.0f;
-	const float LayerViewItemBorderWidth = 1.0f;
+	const float LayerViewItemCornerRadius = 10.0f;
+	const float LayerViewItemBorderWidth = 5.0f;
+	const float LayerViewItemBorderWidthSelected = 1.0f;
 
 	const FLinearColor LayerViewItemFillColor = FLinearColor::Transparent;
-	constexpr FLinearColor LayerViewItemBorderColor = FLinearColor(1, 1, 1, 0.15f);
+	constexpr FLinearColor LayerViewItemBorderColor = FLinearColor(0, 0, 0, 1.f);
 
 	const FLinearColor LayerItemHoverFillColor = FStyleColors::Recessed.GetSpecifiedColor();
 	constexpr FLinearColor LayerItemHoverBorderColor = FLinearColor(1, 1, 1, 0.2f);
@@ -232,25 +237,26 @@ void FSubstrateMaterialEditorStyle::SetupLayerViewStyles(const TSharedRef<FSlate
 
 	Style->Set("LayerView.Row.Hovered", new FSlateRoundedBoxBrush(
 		LayerItemHoverFillColor, LayerViewItemCornerRadius,
-		LayerItemHoverBorderColor, LayerViewItemBorderWidth));
+		LayerItemHoverBorderColor, LayerViewItemBorderWidthSelected));
 
 	Style->Set("LayerView.Row.Selected", new FSlateRoundedBoxBrush(
 		LayerItemSelectFillColor, LayerViewItemCornerRadius,
-		LayerItemSelectBorderColor, LayerViewItemBorderWidth));
+		LayerItemSelectBorderColor, LayerViewItemBorderWidthSelected));
 
 	Style->Set("LayerView.Row.ActiveBrush", new FSlateRoundedBoxBrush(
 		LayerItemSelectFillColor, LayerViewItemCornerRadius,
-		LayerItemSelectBorderColor, LayerViewItemBorderWidth));
+		LayerItemSelectBorderColor, LayerViewItemBorderWidthSelected));
 	Style->Set("LayerView.Row.ActiveHoveredBrush", new FSlateRoundedBoxBrush(
 		LayerItemSelectFillColor, LayerViewItemCornerRadius,
-		LayerItemSelectBorderColor, LayerViewItemBorderWidth));
+		LayerItemSelectBorderColor, LayerViewItemBorderWidthSelected));
 	Style->Set("LayerView.Row.InactiveBrush", new FSlateRoundedBoxBrush(
 		LayerItemSelectFillColor, LayerViewItemCornerRadius,
-		LayerItemSelectBorderColor, LayerViewItemBorderWidth));
+		LayerItemSelectBorderColor, LayerViewItemBorderWidthSelected));
 	Style->Set("LayerView.Row.InactiveHoveredBrush", new FSlateRoundedBoxBrush(
 		LayerItemSelectFillColor, LayerViewItemCornerRadius,
-		LayerItemSelectBorderColor, LayerViewItemBorderWidth));
+		LayerItemSelectBorderColor, LayerViewItemBorderWidthSelected));
 
+	const float DropZoneMargin = 0.25f; 
 	Style->Set("LayerView.Row", FTableRowStyle()
 		.SetTextColor(FStyleColors::Foreground)
 		.SetSelectedTextColor(FStyleColors::ForegroundHover)
@@ -258,17 +264,19 @@ void FSubstrateMaterialEditorStyle::SetupLayerViewStyles(const TSharedRef<FSlate
 		.SetEvenRowBackgroundHoveredBrush(*Style->GetBrush(TEXT("LayerView.Row.Hovered")))
 		.SetOddRowBackgroundBrush(*Style->GetBrush(TEXT("LayerView.Row.Item")))
 		.SetOddRowBackgroundHoveredBrush(*Style->GetBrush(TEXT("LayerView.Row.Hovered")))
-		.SetSelectorFocusedBrush(*Style->GetBrush(TEXT("LayerView.Row.Selected")))
+		.SetSelectorFocusedBrush(*Style->GetBrush(TEXT("LayerView.Row.Item")))
 		.SetActiveBrush(*Style->GetBrush(TEXT("LayerView.Row.ActiveBrush")))
 		.SetActiveHoveredBrush(*Style->GetBrush(TEXT("LayerView.Row.ActiveHoveredBrush")))
 		.SetInactiveBrush(*Style->GetBrush(TEXT("LayerView.Row.InactiveBrush")))
 		.SetInactiveHoveredBrush(*Style->GetBrush(TEXT("LayerView.Row.InactiveHoveredBrush")))
-		.SetSelectorFocusedBrush(BORDER_BRUSH("Common/Selector", FMargin(4.f / 16.f), Style->GetColor(TEXT("Color.Select.Hover"))))
+		//.SetSelectorFocusedBrush(BORDER_BRUSH("Common/Selector", FMargin(4.f / 16.f), Style->GetColor(TEXT("Color.Select.Hover"))))
+		//.SetSelectorFocusedBrush(BOX_BRUSH("Common/DropZoneIndicator_Onto", FMargin(4.f / 16.f), Style->GetColor(TEXT("Color.Select.Hover"))))
 		.SetDropIndicator_Onto(BOX_BRUSH("Common/DropZoneIndicator_Onto", FMargin(4.0f / 16.0f), Style->GetColor(TEXT("Color.Select.Hover"))))
-		.SetDropIndicator_Above(BOX_BRUSH("Common/DropZoneIndicator_Above", FMargin(4.0f / 16.0f, 4.0f / 16.0f, 0.f, 0.f), Style->GetColor(TEXT("Color.Select.Hover"))))
-		.SetDropIndicator_Below(BOX_BRUSH("Common/DropZoneIndicator_Below", FMargin(4.0f / 16.0f, 0.f, 0.f, 4.0f / 16.0f), Style->GetColor(TEXT("Color.Select.Hover"))))
+		//.SetDropIndicator_Above(BORDER_BRUSH("Common/DropZoneIndicatorDashed_Above", FMargin(DropZoneMargin, DropZoneMargin, 0.f, 0.f), Style->GetColor(TEXT("Color.Select.Hover"))))
+		.SetDropIndicator_Above(BOX_BRUSH("Common/VerticalBoxDropZoneIndicator_Above", FMargin(10.0f / 16.0f, 10.0f / 16.0f, 0, 0), Style->GetColor(TEXT("Color.Select.Hover"))))
+		// .SetDropIndicator_Below(BORDER_BRUSH("Common/DropZoneIndicatorDashed_Below", FMargin(DropZoneMargin, 0, 0, DropZoneMargin), Style->GetColor(TEXT("Color.Select.Hover"))))
+		 .SetDropIndicator_Below(BOX_BRUSH("Common/VerticalBoxDropZoneIndicator_Below",  FMargin(10.0f / 16.0f, 0, 0, 10.0f / 16.0f), Style->GetColor(TEXT("Color.Select.Hover"))))
 	);
-
 	Style->Set("LayerView.AddIcon", new IMAGE_BRUSH("Icons/EditorIcons/LayerAdd", Icon16x16));
 	Style->Set("LayerView.DuplicateIcon", new IMAGE_BRUSH("Icons/EditorIcons/Duplicate_40x", Icon40x40));
 	Style->Set("LayerView.RemoveIcon", new IMAGE_BRUSH("Icons/EditorIcons/LayerRemove", Icon16x16));
@@ -362,3 +370,7 @@ void FSubstrateMaterialEditorStyle::SetupTextStyles(const TSharedRef<FSlateStyle
 
 #undef IMAGE_PLUGIN_BRUSH
 #undef IMAGE_PLUGIN_BRUSH_SVG
+
+#undef CORE_IMAGE_BRUSH_SVG
+#undef CORE_BOX_BRUSH_SVG
+#undef CORE_BORDER_BRUSH_SVG
diff --git a/Engine/Source/Editor/MaterialEditor/Private/MaterialPropertyHelpers.cpp b/Engine/Source/Editor/MaterialEditor/Private/MaterialPropertyHelpers.cpp
--- a/Engine/Source/Editor/MaterialEditor/Private/MaterialPropertyHelpers.cpp
+++ b/Engine/Source/Editor/MaterialEditor/Private/MaterialPropertyHelpers.cpp
@@ -89,7 +89,7 @@ FReply SLayerHandle::OnDragDetected(const FGeometry& MyGeometry, const FPointerE
 
 }
 
-TSharedPtr<FLayerDragDropOp> SLayerHandle::CreateDragDropOperation(TSharedPtr<SMaterialLayersFunctionsInstanceTreeItem> InOwningStack)
+TSharedPtr<FLayerDragDropOp> SLayerHandle::CreateDragDropOperation(TSharedPtr<IDraggableItem> InOwningStack)
 {
 	TSharedPtr<FLayerDragDropOp> Operation = MakeShareable(new FLayerDragDropOp(InOwningStack));
 
@@ -875,7 +875,7 @@ TArray<UFactory*> FMaterialPropertyHelpers::GetAssetFactories(EMaterialParameter
 }
 
 
-TSharedRef<SWidget> FMaterialPropertyHelpers::MakeStackReorderHandle(TSharedPtr<SMaterialLayersFunctionsInstanceTreeItem> InOwningStack)
+TSharedRef<SWidget> FMaterialPropertyHelpers::MakeStackReorderHandle(TSharedPtr<IDraggableItem> InOwningStack)
 {
 	TSharedRef<SLayerHandle> Handle = SNew(SLayerHandle)
 		.Content()
diff --git a/Engine/Source/Editor/MaterialEditor/Private/SMaterialLayersFunctionsTree.cpp b/Engine/Source/Editor/MaterialEditor/Private/SMaterialLayersFunctionsTree.cpp
--- a/Engine/Source/Editor/MaterialEditor/Private/SMaterialLayersFunctionsTree.cpp
+++ b/Engine/Source/Editor/MaterialEditor/Private/SMaterialLayersFunctionsTree.cpp
@@ -158,13 +158,14 @@ FReply SMaterialLayersFunctionsInstanceTreeItem::OnLayerDrop(const FDragDropEven
 	TSharedPtr<SMaterialLayersFunctionsInstanceTreeItem> LayerPtr = nullptr;
 	if (ArrayDropOp.IsValid() && ArrayDropOp->OwningStack.IsValid())
 	{
-		LayerPtr = ArrayDropOp->OwningStack.Pin();
-		LayerPtr->bIsBeingDragged = false;
+		LayerPtr = StaticCastWeakPtr<SMaterialLayersFunctionsInstanceTreeItem>(ArrayDropOp->OwningStack).Pin();
 	}
 	if (!LayerPtr.IsValid())
 	{
 		return FReply::Unhandled();
 	}
+	
+	LayerPtr->bIsBeingDragged = false;
 	TSharedPtr<FSortedParamData> SwappingPropertyData = LayerPtr->StackParameterData;
 	TSharedPtr<FSortedParamData> SwappablePropertyData = StackParameterData;
 	if (SwappingPropertyData.IsValid() && SwappablePropertyData.IsValid())
diff --git a/Engine/Source/Editor/MaterialEditor/Private/SMaterialLayersFunctionsTree.h b/Engine/Source/Editor/MaterialEditor/Private/SMaterialLayersFunctionsTree.h
--- a/Engine/Source/Editor/MaterialEditor/Private/SMaterialLayersFunctionsTree.h
+++ b/Engine/Source/Editor/MaterialEditor/Private/SMaterialLayersFunctionsTree.h
@@ -24,7 +24,7 @@ class SMaterialLayersFunctionsInstanceTree;
 class UDEditorParameterValue;
 class UMaterialEditorInstanceConstant;
 
-class SMaterialLayersFunctionsInstanceTreeItem : public STableRow< TSharedPtr<FSortedParamData> >
+class SMaterialLayersFunctionsInstanceTreeItem : public STableRow< TSharedPtr<FSortedParamData> >, public IDraggableItem
 {
 public:
 
@@ -59,7 +59,7 @@ public:
 	void OnNameChanged(const FText& InText, ETextCommit::Type CommitInfo, SMaterialLayersFunctionsInstanceTree* InTree, int32 Counter);
 
 
-	void OnLayerDragEnter(const FDragDropEvent& DragDropEvent)
+	void OnLayerDragEnter(const FDragDropEvent& DragDropEvent) override
 	{
 		if (StackParameterData->ParameterInfo.Index != 0)
 		{
@@ -67,12 +67,12 @@ public:
 		}
 	}
 
-	void OnLayerDragLeave(const FDragDropEvent& DragDropEvent)
+	void OnLayerDragLeave(const FDragDropEvent& DragDropEvent) override
 	{
 		bIsHoveredDragTarget = false;
 	}
 
-	void OnLayerDragDetected()
+	void OnLayerDragDetected() override
 	{
 		bIsBeingDragged = true;
 	}
diff --git a/Engine/Source/Editor/MaterialEditor/Private/SMaterialSubstrateTree.cpp b/Engine/Source/Editor/MaterialEditor/Private/SMaterialSubstrateTree.cpp
--- a/Engine/Source/Editor/MaterialEditor/Private/SMaterialSubstrateTree.cpp
+++ b/Engine/Source/Editor/MaterialEditor/Private/SMaterialSubstrateTree.cpp
@@ -153,7 +153,18 @@ void SMaterialSubstrateTreeItem::OnNameChanged(const FText& InText, ETextCommit:
 	InTree->FunctionInstanceHandle->NotifyPostChange(EPropertyChangeType::ValueSet);
 }
 
-FReply SMaterialSubstrateTreeItem::OnLayerDrop(const FDragDropEvent& DragDropEvent)
+TOptional<EItemDropZone> SMaterialSubstrateTreeItem::CanAcceptDrop(const FDragDropEvent& DragDropEvent, EItemDropZone DropZone, FSortedParamDataPtr Item)
+{
+	TSharedPtr<FLayerDragDropOp> LayerDragDropOperation = DragDropEvent.GetOperationAs<FLayerDragDropOp>();
+	TSharedPtr<FAssetDragDropOp> AssetDragDropOperation = DragDropEvent.GetOperationAs<FAssetDragDropOp>();
+	if (LayerDragDropOperation.IsValid() || AssetDragDropOperation.IsValid())
+	{
+		return DropZone;
+	}
+	return TOptional<EItemDropZone>();
+}
+
+FReply SMaterialSubstrateTreeItem::OnLayerDrop(const FDragDropEvent& DragDropEvent, EItemDropZone DropZone,	FSortedParamDataPtr TargetItem)
 {
 	if (!bIsHoveredDragTarget)
 	{
@@ -162,12 +173,11 @@ FReply SMaterialSubstrateTreeItem::OnLayerDrop(const FDragDropEvent& DragDropEve
 	FScopedTransaction Transaction(NSLOCTEXT("UnrealEd", "MoveLayer", "Move Layer"));
 	Tree->FunctionInstanceHandle->NotifyPreChange();
 	bIsHoveredDragTarget = false;
-	TSharedPtr<FSubstrateLayerDragDropOp> ArrayDropOp = DragDropEvent.GetOperationAs< FSubstrateLayerDragDropOp >();
+	TSharedPtr<FLayerDragDropOp> ArrayDropOp = DragDropEvent.GetOperationAs< FLayerDragDropOp >();
 	TSharedPtr<SMaterialSubstrateTreeItem> LayerPtr = nullptr;
 	if (ArrayDropOp.IsValid() && ArrayDropOp->OwningStack.IsValid())
 	{
-		LayerPtr = ArrayDropOp->OwningStack.Pin();
-		LayerPtr->bIsBeingDragged = false;
+		LayerPtr = StaticCastWeakPtr<SMaterialSubstrateTreeItem>(ArrayDropOp->OwningStack).Pin();
 	}
 	else
 	{
@@ -208,6 +218,8 @@ FReply SMaterialSubstrateTreeItem::OnLayerDrop(const FDragDropEvent& DragDropEve
 	{
 		return FReply::Unhandled();
 	}
+	
+	LayerPtr->bIsBeingDragged = false;
 	TSharedPtr<FSortedParamData> SwappingPropertyData = LayerPtr->StackParameterData;
 	TSharedPtr<FSortedParamData> SwappablePropertyData = StackParameterData;
 	if (SwappingPropertyData.IsValid() && SwappablePropertyData.IsValid())
@@ -278,6 +290,12 @@ void  SMaterialSubstrateTreeItem::OnOverrideParameter(bool NewValue, TObjectPtr<
 	OnOverrideParameter(NewValue, Parameter.Get());
 }
 
+void SMaterialSubstrateTreeItem::AddSubLayer()
+{
+	// create an empty sub layer at this point at the top of the sub-stack
+	
+}
+
 void SMaterialSubstrateTreeItem::Construct(const FArguments& InArgs, const TSharedRef<STableViewBase>& InOwnerTableView)
 {
 	StackParameterData = InArgs._StackParameterData;
@@ -320,17 +338,6 @@ void SMaterialSubstrateTreeItem::Construct(const FArguments& InArgs, const TShar
 		}
 		const float ThumbnailSize = 64.0f;
 		TArray<TSharedPtr<FSortedParamData>> AssetChildren = StackParameterData->Children;
-		if (AssetChildren.Num() > 0)
-		{
-			HeaderRowWidget->AddSlot()
-				.HAlign(HAlign_Center)
-				.VAlign(VAlign_Center)
-				.Padding(0.0f)
-				.AutoWidth()
-				[
-					SNullWidget::NullWidget
-				];
-		}
 		for (TSharedPtr<FSortedParamData> AssetChild : AssetChildren)
 		{
 			TSharedPtr<SBox> ThumbnailBox;
@@ -420,6 +427,14 @@ void SMaterialSubstrateTreeItem::Construct(const FArguments& InArgs, const TShar
 				.Visibility(Tree, &SMaterialSubstrateTree::GetUnlinkLayerVisibility, StackParameterData->ParameterInfo.Index)
 			];
 
+		HeaderRowWidget->AddSlot()
+				.HAlign(HAlign_Left)
+				.AutoWidth()
+				.VAlign(VAlign_Center)
+				[
+					PropertyCustomizationHelpers::MakeAddButton(FSimpleDelegate::CreateSP(this, &SMaterialSubstrateTreeItem::AddSubLayer))
+				];
+
 		// Can only remove layers that aren't the base layer.
 		if (StackParameterData->ParameterInfo.Index != 0)
 		{
@@ -445,24 +460,24 @@ void SMaterialSubstrateTreeItem::Construct(const FArguments& InArgs, const TShar
 			[
 						SAssignNew(FinalStack, SHorizontalBox)
 			];
-		// if (StackParameterData->ParameterInfo.Index != 0)
-		// {
-		// 	FinalStack->AddSlot()
-		// 		.HAlign(HAlign_Center)
-		// 		.VAlign(VAlign_Center)
-		// 		.Padding(2.5f, 0)
-		// 		.AutoWidth()
-		// 		[
-		// 			FMaterialPropertyHelpers::MakeStackReorderHandle(SharedThis(this))
-		// 		];
-		// }
-		// FinalStack->AddSlot()
-		// 	.AutoWidth()
-		// 	.VAlign(VAlign_Center)
-		// 	.Padding(FMargin(2.0f))
-		// 	[
-		// 		SNew(SExpanderArrow, SharedThis(this))
-		// 	];
+		if (StackParameterData->ParameterInfo.Index != 0)
+		{
+			FinalStack->AddSlot()
+				.HAlign(HAlign_Center)
+				.VAlign(VAlign_Center)
+				.Padding(2.5f, 0)
+				.AutoWidth()
+				[
+					FMaterialPropertyHelpers::MakeStackReorderHandle(SharedThis(this))
+				];
+		}
+		FinalStack->AddSlot()
+			.AutoWidth()
+			.VAlign(VAlign_Center)
+			.Padding(FMargin(2.0f))
+			[
+				SNew(SExpanderArrow, SharedThis(this))
+			];
 		FinalStack->AddSlot()
 			.Padding(FMargin(2.0f))
 			.VAlign(VAlign_Center)
@@ -478,7 +493,6 @@ void SMaterialSubstrateTreeItem::Construct(const FArguments& InArgs, const TShar
 
 	FOnTableRowDragEnter LayerDragDelegate = FOnTableRowDragEnter::CreateSP(this, &SMaterialSubstrateTreeItem::OnLayerDragEnter);
 	FOnTableRowDragLeave LayerDragLeaveDelegate = FOnTableRowDragLeave::CreateSP(this, &SMaterialSubstrateTreeItem::OnLayerDragLeave);
-	FOnTableRowDrop LayerDropDelegate = FOnTableRowDrop::CreateSP(this, &SMaterialSubstrateTreeItem::OnLayerDrop);
 
 	STableRow< TSharedPtr<FSortedParamData> >::ConstructInternal(
 		STableRow< TSharedPtr<FSortedParamData> >::FArguments()
@@ -486,11 +500,13 @@ void SMaterialSubstrateTreeItem::Construct(const FArguments& InArgs, const TShar
 		.OnPaintDropIndicator(this, &SMaterialSubstrateTreeItem::OnLayerItemPaintDropIndicator)
 		.Style(FSubstrateMaterialEditorStyle::Get(), "LayerView.Row")
 		.ShowSelection(true)
+		.OnCanAcceptDrop(this, &SMaterialSubstrateTreeItem::CanAcceptDrop)
+		.OnAcceptDrop(this, &SMaterialSubstrateTreeItem::OnLayerDrop)
 		.OnDragEnter(LayerDragDelegate)
-		.OnDragLeave(LayerDragLeaveDelegate)
-		.OnDrop(LayerDropDelegate),
+		.OnDragLeave(LayerDragLeaveDelegate),
 		InOwnerTableView
 	);
+
 }
 int32 SMaterialSubstrateTreeItem::OnLayerItemPaintDropIndicator(EItemDropZone InItemDropZone, const FPaintArgs& Args, const FGeometry& AllottedGeometry, const FSlateRect& MyCullingRect, FSlateWindowElementList& OutDrawElements, int32 LayerId, const FWidgetStyle& InWidgetStyle, bool bParentEnabled) const
 {
diff --git a/Engine/Source/Editor/MaterialEditor/Private/SMaterialSubstrateTree.h b/Engine/Source/Editor/MaterialEditor/Private/SMaterialSubstrateTree.h
--- a/Engine/Source/Editor/MaterialEditor/Private/SMaterialSubstrateTree.h
+++ b/Engine/Source/Editor/MaterialEditor/Private/SMaterialSubstrateTree.h
@@ -102,7 +102,7 @@ private:
 	FGetShowHiddenParameters ShowHiddenDelegate;
 };
 
-class SMaterialSubstrateTreeItem : public STableRow< FSortedParamDataPtr >
+class SMaterialSubstrateTreeItem : public STableRow< FSortedParamDataPtr >, public IDraggableItem
 {
 public:
 
@@ -139,8 +139,9 @@ public:
 	FText GetLayerName(SMaterialSubstrateTree* InTree, int32 Counter) const;
 	void OnNameChanged(const FText& InText, ETextCommit::Type CommitInfo, SMaterialSubstrateTree* InTree, int32 Counter);
 
-
-	void OnLayerDragEnter(const FDragDropEvent& DragDropEvent)
+	TOptional<EItemDropZone> CanAcceptDrop(const FDragDropEvent& DragDropEvent, EItemDropZone DropZone, FSortedParamDataPtr Item);
+	FReply OnAcceptDrop(const FDragDropEvent& DragDropEvent, EItemDropZone DropZone, FSortedParamDataPtr TargetItem);
+	void OnLayerDragEnter(const FDragDropEvent& DragDropEvent) override
 	{
 		if (StackParameterData->ParameterInfo.Index != 0)
 		{
@@ -148,19 +149,20 @@ public:
 		}
 	}
 
-	void OnLayerDragLeave(const FDragDropEvent& DragDropEvent)
+	void OnLayerDragLeave(const FDragDropEvent& DragDropEvent) override
 	{
 		bIsHoveredDragTarget = false;
 	}
 
-	void OnLayerDragDetected()
+	void OnLayerDragDetected() override
 	{
 		bIsBeingDragged = true;
 	}
 
-	FReply OnLayerDrop(const FDragDropEvent& DragDropEvent);
+	FReply OnLayerDrop(const FDragDropEvent& DragDropEvent, EItemDropZone DropZone, FSortedParamDataPtr TargetItem);
 	void OnOverrideParameter(bool NewValue, class UDEditorParameterValue* Parameter);
 	void OnOverrideParameter(bool NewValue, TObjectPtr<UDEditorParameterValue> Parameter);
+	void AddSubLayer();
 	/**
 	* Construct the widget
 	*
@@ -178,38 +180,3 @@ public:
 
 	FString GetInstancePath(SMaterialSubstrateTree* InTree) const;
 };
-
-class FSubstrateLayerDragDropOp final : public FDecoratedDragDropOp
-{
-public:
-	DRAG_DROP_OPERATOR_TYPE(FSubstrateLayerDragDropOp, FDecoratedDragDropOp)
-
-	FSubstrateLayerDragDropOp(TSharedPtr<SMaterialSubstrateTreeItem> InOwningStack)
-	{
-		OwningStack = InOwningStack;
-		DecoratorWidget = SNew(SBorder)
-			.BorderImage(FAppStyle::GetBrush("Graph.ConnectorFeedback.Border"))
-			.Content()
-			[
-				SNew(SHorizontalBox)
-				+ SHorizontalBox::Slot()
-				.AutoWidth()
-				.VAlign(VAlign_Center)
-				[
-					SNew(STextBlock)
-					.Text(NSLOCTEXT("LayerDragDrop", "PlaceLayerHere", "Place Layer and Blend Here"))
-				]
-			];
-
-		Construct();
-	};
-
-	TSharedPtr<SWidget> DecoratorWidget;
-
-	virtual TSharedPtr<SWidget> GetDefaultDecorator() const override
-	{
-		return DecoratorWidget;
-	}
-
-	TWeakPtr<class SMaterialSubstrateTreeItem> OwningStack;
-};
diff --git a/Engine/Source/Editor/MaterialEditor/Public/MaterialPropertyHelpers.h b/Engine/Source/Editor/MaterialEditor/Public/MaterialPropertyHelpers.h
--- a/Engine/Source/Editor/MaterialEditor/Public/MaterialPropertyHelpers.h
+++ b/Engine/Source/Editor/MaterialEditor/Public/MaterialPropertyHelpers.h
@@ -21,7 +21,6 @@ class IPropertyHandle;
 class UDEditorParameterValue;
 enum class ECheckBoxState : uint8;
 class UMaterialInterface;
-class SMaterialLayersFunctionsInstanceTreeItem;
 
 DECLARE_DELEGATE_OneParam(FGetShowHiddenParameters, bool&);
 
@@ -72,13 +71,24 @@ struct FUnsortedParamData
 	TSharedPtr<IPropertyHandle> ParameterHandle;
 };
 
+/*
+ * Interface for items that can be dragged and can show a layer handle
+ */
+class IDraggableItem
+{
+public:
+	virtual void OnLayerDragEnter(const FDragDropEvent& DragDropEvent) = 0; 
+	virtual void OnLayerDragLeave(const FDragDropEvent& DragDropEvent) = 0;
+	virtual void OnLayerDragDetected() = 0;
+};
+
 class SLayerHandle : public SCompoundWidget
 {
 public:
 	SLATE_BEGIN_ARGS(SLayerHandle)
 	{}
 	SLATE_DEFAULT_SLOT(FArguments, Content)
-		SLATE_ARGUMENT(TSharedPtr<SMaterialLayersFunctionsInstanceTreeItem>, OwningStack)
+		SLATE_ARGUMENT(TSharedPtr<IDraggableItem>, OwningStack)
 		SLATE_END_ARGS()
 
 		void Construct(const FArguments& InArgs);
@@ -90,10 +100,10 @@ public:
 
 
 	FReply OnDragDetected(const FGeometry& MyGeometry, const FPointerEvent& MouseEvent) override;
-	TSharedPtr<class FLayerDragDropOp> CreateDragDropOperation(TSharedPtr<SMaterialLayersFunctionsInstanceTreeItem> InOwningStack);
+	TSharedPtr<class FLayerDragDropOp> CreateDragDropOperation(TSharedPtr<IDraggableItem> InOwningStack);
 
 private:
-	TWeakPtr<SMaterialLayersFunctionsInstanceTreeItem> OwningStack;
+	TWeakPtr<IDraggableItem> OwningStack;
 };
 
 
@@ -102,7 +112,7 @@ class FLayerDragDropOp final : public FDecoratedDragDropOp
 public:
 	DRAG_DROP_OPERATOR_TYPE(FLayerDragDropOp, FDecoratedDragDropOp)
 
-	FLayerDragDropOp(TSharedPtr<SMaterialLayersFunctionsInstanceTreeItem> InOwningStack)
+	FLayerDragDropOp(TSharedPtr<IDraggableItem> InOwningStack)
 	{
 		OwningStack = InOwningStack;
 		DecoratorWidget = SNew(SBorder)
@@ -129,7 +139,7 @@ public:
 		return DecoratorWidget;
 	}
 
-	TWeakPtr<class SMaterialLayersFunctionsInstanceTreeItem> OwningStack;
+	TWeakPtr<class IDraggableItem> OwningStack;
 };
 
 /*-----------------------------------------------------------------------------
@@ -195,7 +205,7 @@ public:
 	*/
 	static FEditorParameterGroup&  GetParameterGroup(class UMaterial* InMaterial, FName& ParameterGroup, TArray<FEditorParameterGroup>& ParameterGroups);
 
-	static TSharedRef<SWidget> MakeStackReorderHandle(TSharedPtr<SMaterialLayersFunctionsInstanceTreeItem> InOwningStack);
+	static TSharedRef<SWidget> MakeStackReorderHandle(TSharedPtr<IDraggableItem> InOwningStack);
 
 	static bool OnShouldSetCurveAsset(const FAssetData& AssetData, TSoftObjectPtr<UCurveLinearColorAtlas> InAtlas);
 	static bool OnShouldFilterCurveAsset(const FAssetData& AssetData, TSoftObjectPtr<UCurveLinearColorAtlas> InAtlas);
