Post Process:  Upgraded material instance support, with overrides for User Scene Texture inputs / outputs and Blendable Location / Priority.  Allows a base post process material to be reused in different places or multiple times in the post process pipeline without requiring a completely different Material graph.  For example, a generic upsample, downsample, or blur Material can be implemented, rather than having to create a duplicate asset for each variation of location and input.  Besides saving time creating content, material instances, unlike materials, can be created dynamically in code, allowing game logic or future editing tools to interactively configure the pipeline.

Preview windows and the r.PostProcessing.UserSceneTextureDebug mode were updated to handle material instances.  If both a base material and an instance of it are being edited, the instance will be preferred in other preview windows over the base material, rather than including both and having their effects potentially overlap.  Sorting of previewed materials was improved by ordering them in the reverse order they are encountered, so dependencies earlier in a chain are added first.  Fixed cases where preview windows weren't refreshed on certain types of changes or debug mode toggling.

To handle generic materials that can work at multiple resolutions, a "ScaleRelativeToInput" field was added, which generates the output resolution relative to a given User Scene Texture input.  Positive divisors downsample, while negative divisors upsample (so 2 would divide resolution by 2, while -2 would multiply by 2).  To handle the use of a generic material where you want the input or output to be SceneColor (say a first pass that reads SceneColor, or a last pass that writes SceneColor), you now have the option to specify "SceneColor" as a User Scene Texture input or output name.  This special name doesn't create a transient User Scene Texture, but just hooks up SceneColor itself.

The final feature added is a flag which suppresses PreExposure adjustments (multiply by [REDACTED_PERSON] on input, multiply by [REDACTED_PERSON] on output), useful if creating a material that doesn't care about the absolute value of linear colors (such as a blur where the inputs and outputs have the same scale), or where the output is not a linear color (such as a mask, modulation, or some sort of non-color data like positions, normals, offsets, IDs, etc).  Especially simplifies writing custom HLSL by avoiding the need for confusing View.OneOverPreExposure.xxx boilerplate, in addition to providing a minor perf gain.

diff --git a/Engine/Plugins/Compositing/Composure/Source/Composure/Private/ComposurePostProcessPass.cpp b/Engine/Plugins/Compositing/Composure/Source/Composure/Private/ComposurePostProcessPass.cpp
--- a/Engine/Plugins/Compositing/Composure/Source/Composure/Private/ComposurePostProcessPass.cpp
+++ b/Engine/Plugins/Compositing/Composure/Source/Composure/Private/ComposurePostProcessPass.cpp
@@ -52,7 +52,7 @@ void UComposurePostProcessPass::SetSetupMaterial(UMaterialInterface* Material)
 			*BaseMaterial->GetName());
 		return;
 	}
-	else if (BaseMaterial->BlendableLocation != BL_SceneColorBeforeDOF)
+	else if (Material->GetBlendableLocation(BaseMaterial) != BL_SceneColorBeforeDOF)
 	{
 		UE_LOG(Composure, Error,
 			TEXT("Can't set setup material %s: is not at the before translucency post process location."),
diff --git a/Engine/Shaders/Private/PostProcessMaterialShaders.usf b/Engine/Shaders/Private/PostProcessMaterialShaders.usf
--- a/Engine/Shaders/Private/PostProcessMaterialShaders.usf
+++ b/Engine/Shaders/Private/PostProcessMaterialShaders.usf
@@ -291,7 +291,7 @@ void MainPS(
 	
 	// This is for debug only
 	OutColor.xyz = NeuralInput0;
-#if !MATERIALBLENDING_MODULATE && POST_PROCESS_MATERIAL_BEFORE_TONEMAP
+#if !MATERIALBLENDING_MODULATE && POST_PROCESS_MATERIAL_BEFORE_TONEMAP && !POST_PROCESS_DISABLE_PRE_EXPOSURE_SCALE
 	OutColor.xyz *= View.PreExposure;
 #endif
 
@@ -352,7 +352,7 @@ void MainPS(
 
 	OutColor = float4(EmissiveColor, Alpha );
 
-#if !MATERIALBLENDING_MODULATE && POST_PROCESS_MATERIAL_BEFORE_TONEMAP
+#if !MATERIALBLENDING_MODULATE && POST_PROCESS_MATERIAL_BEFORE_TONEMAP && !POST_PROCESS_DISABLE_PRE_EXPOSURE_SCALE
 	OutColor.xyz *= View.PreExposure;
 #endif
 } 
diff --git a/Engine/Source/Editor/MaterialEditor/Private/MaterialEditor.cpp b/Engine/Source/Editor/MaterialEditor/Private/MaterialEditor.cpp
--- a/Engine/Source/Editor/MaterialEditor/Private/MaterialEditor.cpp
+++ b/Engine/Source/Editor/MaterialEditor/Private/MaterialEditor.cpp
@@ -1138,30 +1138,7 @@ void FMaterialEditor::NotifyUserSceneTextureLoadOrUnload()
 {
 	if (Material->IsPostProcessMaterial() && !Material->UserSceneTexture.IsNone())
 	{
-		UAssetEditorSubsystem* AssetEditorSubsystem = GEditor->GetEditorSubsystem<UAssetEditorSubsystem>();
-		TArray<UObject*> EditedAssets = AssetEditorSubsystem->GetAllEditedAssets();
-		for (UObject* EditedAsset : EditedAssets)
-		{
-			UPreviewMaterial* EditedMaterialInterface = Cast<UPreviewMaterial>(EditedAsset);
-
-			if (EditedMaterialInterface && EditedMaterialInterface != Material)
-			{
-				UMaterial* EditedMaterial = EditedMaterialInterface->GetMaterial();
-				if (EditedMaterial->IsPostProcessMaterial())
-				{
-					TArray<IAssetEditorInstance*> Editors = AssetEditorSubsystem->FindEditorsForAsset(EditedAsset);
-					for (IAssetEditorInstance* Editor : Editors)
-					{
-						if (Editor->GetEditorName() == FMaterialEditor::GetToolkitFName())
-						{
-							// Calling "SetPreviewMaterial" will refresh the other editor
-							FMaterialEditor* MaterialEditor = (FMaterialEditor*)Editor;
-							MaterialEditor->SetPreviewMaterial(EditedMaterialInterface);
-						}
-					}
-				}
-			}
-		}
+		FMaterialEditorUtilities::RefreshPostProcessPreviewMaterials(Material);
 	}
 }
 
diff --git a/Engine/Source/Editor/MaterialEditor/Private/MaterialEditor.h b/Engine/Source/Editor/MaterialEditor/Private/MaterialEditor.h
--- a/Engine/Source/Editor/MaterialEditor/Private/MaterialEditor.h
+++ b/Engine/Source/Editor/MaterialEditor/Private/MaterialEditor.h
@@ -453,6 +453,12 @@ public:
 	/** Overrides function in FEditorUndoClient. Called to see if the context of the current undo/redo operation is a match for the client. */
 	virtual bool MatchesContext(const FTransactionContext& InContext, const TArray<TPair<UObject*, FTransactionObjectEvent>>& TransactionObjectContexts) const override;
 
+	/**
+	 * Notifies all other editors that the given material with a UserSceneTexture output was loaded or unloaded, so they can refresh their preview window
+	 * if needed.  Utility function shared between UMaterial and UMaterialInstance editors.
+	 */
+	static void NotifyAllUserSceneTextureLoadOrUnload(UMaterialInterface* MaterialInterface);
+
 public:
 	/** Set to true when modifications have been made to the material */
 	bool bMaterialDirty;
diff --git a/Engine/Source/Editor/MaterialEditor/Private/MaterialEditorInstanceDetailCustomization.cpp b/Engine/Source/Editor/MaterialEditor/Private/MaterialEditorInstanceDetailCustomization.cpp
--- a/Engine/Source/Editor/MaterialEditor/Private/MaterialEditorInstanceDetailCustomization.cpp
+++ b/Engine/Source/Editor/MaterialEditor/Private/MaterialEditorInstanceDetailCustomization.cpp
@@ -145,6 +145,7 @@ void FMaterialInstanceParameterDetails::CustomizeDetails(IDetailLayoutBuilder& D
 		}
 		DetailLayout.HideCategory("MaterialEditorInstanceConstant");
 		DetailLayout.HideProperty("Parent");
+		DetailLayout.HideProperty("PostProcessOverrides");
 		DetailLayout.HideProperty("PhysMaterial");
 		DetailLayout.HideProperty("LightmassSettings");
 		DetailLayout.HideProperty("bUseOldStyleMICEditorGroups");
@@ -216,6 +217,7 @@ void FMaterialInstanceParameterDetails::CustomizeDetails(IDetailLayoutBuilder& D
 				DetailLayout.HideProperty("Parent");
 			}
 
+			DetailLayout.HideProperty("PostProcessOverrides");
 			DetailLayout.HideProperty("PhysMaterial");
 			DetailLayout.HideProperty("LightmassSettings");
 			DetailLayout.HideProperty("bUseOldStyleMICEditorGroups");
@@ -227,6 +229,9 @@ void FMaterialInstanceParameterDetails::CustomizeDetails(IDetailLayoutBuilder& D
 		}
 		else
 		{
+			DetailLayout.HideProperty("PostProcessOverrides");
+			CreatePostProcessOverrideWidgets(DetailLayout);
+
 			// Add PhysMaterial property
 			DefaultCategory.AddProperty("PhysMaterial");
 
@@ -1304,6 +1309,102 @@ void FMaterialInstanceParameterDetails::CreateLightmassOverrideWidgets(IDetailLa
 		.OverrideResetToDefault(ResetExportResolutionScalePropertyOverride);
 }
 
+void FMaterialInstanceParameterDetails::CreatePostProcessOverrideWidgets(IDetailLayoutBuilder& DetailLayout)
+{
+	if (MaterialEditorInstance->PostProcessOverrides.bIsOverrideable)
+	{
+		FName PostProcessCategoryName = TEXT("PostProcessOverrides");
+		IDetailCategoryBuilder& PostProcessCategory = DetailLayout.EditCategory(PostProcessCategoryName, LOCTEXT("MICPostProcessOverridesTitle", "Post Process Overrides"));
+		PostProcessCategory.InitiallyCollapsed(true);
+
+		TAttribute<bool> IsOverrideLocationEnabled = TAttribute<bool>::Create(TAttribute<bool>::FGetter::CreateLambda([this] { return (bool)MaterialEditorInstance->PostProcessOverrides.bOverrideBlendableLocation; }));
+		TAttribute<bool> IsOverridePriorityEnabled = TAttribute<bool>::Create(TAttribute<bool>::FGetter::CreateLambda([this] { return (bool)MaterialEditorInstance->PostProcessOverrides.bOverrideBlendablePriority; }));
+
+		TSharedRef<IPropertyHandle> PostProcessOverridesProperty = DetailLayout.GetProperty("PostProcessOverrides");
+		TSharedPtr<IPropertyHandle> BlendableLocationProperty = PostProcessOverridesProperty->GetChildHandle("BlendableLocationOverride");
+		TSharedPtr<IPropertyHandle> BlendablePriorityProperty = PostProcessOverridesProperty->GetChildHandle("BlendablePriorityOverride");
+		TSharedPtr<IPropertyHandle> UserSceneTextureOutputProperty = PostProcessOverridesProperty->GetChildHandle("UserSceneTextureOutput");
+
+		FIsResetToDefaultVisible IsBlendableLocationPropertyResetVisible = FIsResetToDefaultVisible::CreateLambda([this](TSharedPtr<IPropertyHandle> InHandle) {
+			return MaterialEditorInstance->Parent && MaterialEditorInstance->Parent->GetMaterial() ?
+				MaterialEditorInstance->PostProcessOverrides.BlendableLocationOverride != MaterialEditorInstance->Parent->GetMaterial()->BlendableLocation : false;
+			});
+		FResetToDefaultHandler ResetBlendableLocationPropertyHandler = FResetToDefaultHandler::CreateLambda([this](TSharedPtr<IPropertyHandle> InHandle) {
+			if (MaterialEditorInstance->Parent && MaterialEditorInstance->Parent->GetMaterial())
+			{
+				MaterialEditorInstance->PostProcessOverrides.BlendableLocationOverride = MaterialEditorInstance->Parent->GetMaterial()->BlendableLocation;
+			}
+			});
+		FResetToDefaultOverride ResetBlendableLocationPropertyOverride = FResetToDefaultOverride::Create(IsBlendableLocationPropertyResetVisible, ResetBlendableLocationPropertyHandler);
+
+		IDetailPropertyRow& BlendableLocationPropertyRow = PostProcessCategory.AddProperty(BlendableLocationProperty);
+		BlendableLocationPropertyRow
+			.DisplayName(BlendableLocationProperty->GetPropertyDisplayName())
+			.ToolTip(BlendableLocationProperty->GetToolTipText())
+			.EditCondition(IsOverrideLocationEnabled, FOnBooleanValueChanged::CreateLambda([this](bool NewValue) {
+				MaterialEditorInstance->PostProcessOverrides.bOverrideBlendableLocation = NewValue;
+				MaterialEditorInstance->PostEditChange();
+				FEditorSupportDelegates::RedrawAllViewports.Broadcast();
+			}))
+			.Visibility(TAttribute<EVisibility>::Create(TAttribute<EVisibility>::FGetter::CreateSP(this, &FMaterialInstanceParameterDetails::IsOverriddenAndVisible, IsOverrideLocationEnabled)))
+			.OverrideResetToDefault(ResetBlendableLocationPropertyOverride);
+
+		FIsResetToDefaultVisible IsBlendablePriorityPropertyResetVisible = FIsResetToDefaultVisible::CreateLambda([this](TSharedPtr<IPropertyHandle> InHandle) {
+			return MaterialEditorInstance->Parent && MaterialEditorInstance->Parent->GetMaterial() ?
+				MaterialEditorInstance->PostProcessOverrides.BlendablePriorityOverride != MaterialEditorInstance->Parent->GetMaterial()->BlendablePriority : false;
+			});
+		FResetToDefaultHandler ResetBlendablePriorityPropertyHandler = FResetToDefaultHandler::CreateLambda([this](TSharedPtr<IPropertyHandle> InHandle) {
+			if (MaterialEditorInstance->Parent && MaterialEditorInstance->Parent->GetMaterial())
+			{
+				MaterialEditorInstance->PostProcessOverrides.BlendablePriorityOverride = MaterialEditorInstance->Parent->GetMaterial()->BlendablePriority;
+			}
+			});
+		FResetToDefaultOverride ResetBlendablePriorityPropertyOverride = FResetToDefaultOverride::Create(IsBlendablePriorityPropertyResetVisible, ResetBlendablePriorityPropertyHandler);
+
+		IDetailPropertyRow& BlendablePriorityPropertyRow = PostProcessCategory.AddProperty(BlendablePriorityProperty);
+		BlendablePriorityPropertyRow
+			.DisplayName(BlendablePriorityProperty->GetPropertyDisplayName())
+			.ToolTip(BlendablePriorityProperty->GetToolTipText())
+			.EditCondition(IsOverridePriorityEnabled, FOnBooleanValueChanged::CreateLambda([this](bool NewValue) {
+				MaterialEditorInstance->PostProcessOverrides.bOverrideBlendablePriority = NewValue;
+				MaterialEditorInstance->PostEditChange();
+				FEditorSupportDelegates::RedrawAllViewports.Broadcast();
+			}))
+			.Visibility(TAttribute<EVisibility>::Create(TAttribute<EVisibility>::FGetter::CreateSP(this, &FMaterialInstanceParameterDetails::IsOverriddenAndVisible, IsOverridePriorityEnabled)))
+			.OverrideResetToDefault(ResetBlendablePriorityPropertyOverride);
+
+		if (MaterialEditorInstance->PostProcessOverrides.UserSceneTextureInputs.Num())
+		{
+			static FName GroupName(TEXT("UserSceneTextures"));
+			IDetailGroup& UserSceneTexturesGroup = PostProcessCategory.AddGroup(GroupName, LOCTEXT("UserSceneTextureInputsGroup", "User Scene Texture Inputs"), false, true);
+
+			TSharedPtr<IPropertyHandle> UserSceneTexturesArrayProperty = PostProcessOverridesProperty->GetChildHandle("UserSceneTextureInputs");
+
+			for (int32 UserSceneTextureIndex = 0; UserSceneTextureIndex < MaterialEditorInstance->PostProcessOverrides.UserSceneTextureInputs.Num(); ++UserSceneTextureIndex)
+			{
+				TSharedPtr<IPropertyHandle> UserSceneTextureItemProperty = UserSceneTexturesArrayProperty->GetChildHandle(UserSceneTextureIndex);
+				TSharedPtr<IPropertyHandle> UserSceneTextureValueProperty = UserSceneTextureItemProperty->GetChildHandle("Value");
+
+				IDetailPropertyRow& PropertyRow = UserSceneTexturesGroup.AddPropertyRow(UserSceneTextureValueProperty.ToSharedRef());
+
+				PropertyRow.CustomWidget()
+					.NameContent()
+					[
+						SNew(STextBlock)
+						.Text(FText::FromName(MaterialEditorInstance->PostProcessOverrides.UserSceneTextureInputs[UserSceneTextureIndex].Key))
+						.Font(IDetailLayoutBuilder::GetDetailFont())
+					]
+					.ValueContent()
+					[
+						UserSceneTextureValueProperty->CreatePropertyValueWidget()
+					];
+			}
+		}
+
+		PostProcessCategory.AddProperty(UserSceneTextureOutputProperty);
+	}
+}
+
 UEnum* GetBlendModeEnum();
 
 void FMaterialInstanceParameterDetails::CreateBasePropertyOverrideWidgets(IDetailLayoutBuilder& DetailLayout, IDetailGroup& MaterialPropertyOverrideGroup)
diff --git a/Engine/Source/Editor/MaterialEditor/Private/MaterialEditorInstanceDetailCustomization.h b/Engine/Source/Editor/MaterialEditor/Private/MaterialEditorInstanceDetailCustomization.h
--- a/Engine/Source/Editor/MaterialEditor/Private/MaterialEditorInstanceDetailCustomization.h
+++ b/Engine/Source/Editor/MaterialEditor/Private/MaterialEditorInstanceDetailCustomization.h
@@ -104,6 +104,9 @@ private:
 	/** Creates all the lightmass property override widgets. */
 	void CreateLightmassOverrideWidgets(IDetailLayoutBuilder& DetailLayout);
 
+	/** Creates Blendable Location / Priority and UserSceneTexture input / output override widgets. */
+	void CreatePostProcessOverrideWidgets(IDetailLayoutBuilder& DetailLayout);
+
 	//Functions supporting BasePropertyOverrides
 
 	/** Creates all the base property override widgets. */
diff --git a/Engine/Source/Editor/MaterialEditor/Private/MaterialEditorModule.cpp b/Engine/Source/Editor/MaterialEditor/Private/MaterialEditorModule.cpp
--- a/Engine/Source/Editor/MaterialEditor/Private/MaterialEditorModule.cpp
+++ b/Engine/Source/Editor/MaterialEditor/Private/MaterialEditorModule.cpp
@@ -44,6 +44,15 @@ public:
 	/** Constructor, set up console commands and variables **/
 	FMaterialEditorModule()
 	{
+		// Trigger a redraw of post process preview materials when this debug setting changes
+		IConsoleVariable* CVarPostProcessUserSceneTextureDebug = IConsoleManager::Get().FindConsoleVariable(TEXT("r.PostProcessing.UserSceneTextureDebug"));
+		if (CVarPostProcessUserSceneTextureDebug)
+		{
+			CVarPostProcessUserSceneTextureDebug->SetOnChangedCallback(FConsoleVariableDelegate::CreateLambda([this](IConsoleVariable* Variable)
+			{
+				FMaterialEditorUtilities::RefreshPostProcessPreviewMaterials(nullptr, /* bRedrawOnly= */ true);
+			}));
+		}
 	}
 
 	/**
diff --git a/Engine/Source/Editor/MaterialEditor/Private/MaterialEditorUtilities.cpp b/Engine/Source/Editor/MaterialEditor/Private/MaterialEditorUtilities.cpp
--- a/Engine/Source/Editor/MaterialEditor/Private/MaterialEditorUtilities.cpp
+++ b/Engine/Source/Editor/MaterialEditor/Private/MaterialEditorUtilities.cpp
@@ -31,17 +31,20 @@
 #include "Materials/MaterialExpressionRerouteBase.h"
 #include "Materials/MaterialExpressionExecBegin.h"
 #include "Materials/MaterialExpressionExecEnd.h"
+#include "Materials/MaterialInstanceConstant.h"
 
 #include "DebugViewModeHelpers.h"
 #include "Toolkits/ToolkitManager.h"
 #include "MaterialEditor.h"
 #include "MaterialExpressionClasses.h"
+#include "MaterialInstanceEditor.h"
 #include "Materials/MaterialInstance.h"
 #include "MaterialUtilities.h"
 #include "Misc/ScopedSlowTask.h"
 #include "Templates/UniquePtr.h"
 #include "Materials/MaterialFunctionInstance.h"
 #include "Subsystems/AssetEditorSubsystem.h"
+#include "MaterialEditor/PreviewMaterial.h"
 
 #define LOCTEXT_NAMESPACE "MaterialEditorUtilities"
 
@@ -892,4 +895,64 @@ void FMaterialEditorUtilities::OpenSelectedParentEditor(UMaterialFunctionInterfa
 	}
 }
 
+void FMaterialEditorUtilities::RefreshPostProcessPreviewMaterials(UMaterialInterface* ExcludeMaterialInterface, bool bRedrawOnly)
+{
+	UAssetEditorSubsystem* AssetEditorSubsystem = GEditor->GetEditorSubsystem<UAssetEditorSubsystem>();
+	TArray<UObject*> EditedAssets = AssetEditorSubsystem->GetAllEditedAssets();
+	for (UObject* EditedAsset : EditedAssets)
+	{
+		UPreviewMaterial* EditedPreviewMaterial = Cast<UPreviewMaterial>(EditedAsset);
+		if (EditedPreviewMaterial && EditedPreviewMaterial != ExcludeMaterialInterface)
+		{
+			UMaterial* EditedMaterial = EditedPreviewMaterial->GetMaterial();
+			if (EditedMaterial->IsPostProcessMaterial())
+			{
+				TArray<IAssetEditorInstance*> Editors = AssetEditorSubsystem->FindEditorsForAsset(EditedAsset);
+				for (IAssetEditorInstance* Editor : Editors)
+				{
+					if (Editor->GetEditorName() == FName("MaterialEditor"))
+					{
+						FMaterialEditor* MaterialEditor = (FMaterialEditor*)Editor;
+						if (bRedrawOnly)
+						{
+							MaterialEditor->RefreshPreviewViewport();
+						}
+						else
+						{
+							// Calling "SetPreviewMaterial" will refresh the other editor
+							MaterialEditor->SetPreviewMaterial(EditedPreviewMaterial);
+						}
+					}
+				}
+			}
+		}
+
+		UMaterialInstanceConstant* EditedMaterialInstance = Cast<UMaterialInstanceConstant>(EditedAsset);
+		if (EditedMaterialInstance && EditedMaterialInstance != ExcludeMaterialInterface)
+		{
+			UMaterial* BaseMaterial = EditedMaterialInstance->GetBaseMaterial();
+			if (BaseMaterial && BaseMaterial->IsPostProcessMaterial())
+			{
+				TArray<IAssetEditorInstance*> Editors = AssetEditorSubsystem->FindEditorsForAsset(EditedAsset);
+				for (IAssetEditorInstance* Editor : Editors)
+				{
+					if (Editor->GetEditorName() == FName("MaterialInstanceEditor"))
+					{
+						FMaterialInstanceEditor* MaterialEditor = (FMaterialInstanceEditor*)Editor;
+						if (bRedrawOnly)
+						{
+							MaterialEditor->RefreshPreviewViewport();
+						}
+						else
+						{
+							// Calling "SetPreviewMaterial" will refresh the other editor
+							MaterialEditor->SetPreviewMaterial(EditedMaterialInstance);
+						}
+					}
+				}
+			}
+		}
+	}
+}
+
 #undef LOCTEXT_NAMESPACE
diff --git a/Engine/Source/Editor/MaterialEditor/Private/MaterialInstanceEditor.cpp b/Engine/Source/Editor/MaterialEditor/Private/MaterialInstanceEditor.cpp
--- a/Engine/Source/Editor/MaterialEditor/Private/MaterialInstanceEditor.cpp
+++ b/Engine/Source/Editor/MaterialEditor/Private/MaterialInstanceEditor.cpp
@@ -574,6 +574,9 @@ void FMaterialInstanceEditor::InitMaterialInstanceEditor( const EToolkitMode::Ty
 	}
 
 	Refresh();
+
+	// Notify other editors if this material editor has a post process named output, which may affect their preview
+	NotifyUserSceneTextureLoadOrUnload();
 }
 
 void FMaterialInstanceEditor::ReInitMaterialFunctionProxies()
@@ -750,6 +753,11 @@ FMaterialInstanceEditor::FMaterialInstanceEditor()
 
 FMaterialInstanceEditor::~FMaterialInstanceEditor()
 {
+	bDestructing = true;
+
+	// Notify other editors if this material editor has a post process named output, which may affect their preview
+	NotifyUserSceneTextureLoadOrUnload();
+
 	// Broadcast that this editor is going down to all listeners
 	OnMaterialEditorClosed().Broadcast();
 
@@ -1033,6 +1041,13 @@ void FMaterialInstanceEditor::GenerateInheritanceMenu(UToolMenu* Menu)
 	}
 }
 
+void FMaterialInstanceEditor::RefreshPreviewViewport()
+{
+	if (PreviewVC.IsValid())
+	{
+		PreviewVC->RefreshViewport();
+	}
+}
 
 TSharedRef<SDockTab> FMaterialInstanceEditor::SpawnTab_Preview( const FSpawnTabArgs& Args )
 {	
@@ -1740,4 +1755,26 @@ void FMaterialInstanceEditor::NotifyExternalMaterialChange()
 	MaterialStatsManager->SignalMaterialChanged();
 }
 
+void FMaterialInstanceEditor::NotifyUserSceneTextureLoadOrUnload()
+{
+	if (PreviewVC.IsValid() && PreviewVC->PreviewMaterial)
+	{
+		UMaterialInstance* MaterialInstance = Cast<UMaterialInstance>(PreviewVC->PreviewMaterial);
+		if (MaterialInstance)
+		{
+			UMaterial* BaseMaterial = MaterialInstance->GetMaterial();
+			if (BaseMaterial && BaseMaterial->IsPostProcessMaterial())
+			{
+				FName Output = MaterialInstance->GetUserSceneTextureOutput(BaseMaterial);
+
+				// Ignore special SceneColor output name -- this just writes to SceneColor, not a transient UserSceneTexture
+				if (!Output.IsNone() && Output != FName("SceneColor"))
+				{
+					FMaterialEditorUtilities::RefreshPostProcessPreviewMaterials(MaterialInstance);
+				}
+			}
+		}
+	}
+}
+
 #undef LOCTEXT_NAMESPACE
diff --git a/Engine/Source/Editor/MaterialEditor/Private/MaterialInstanceEditor.h b/Engine/Source/Editor/MaterialEditor/Private/MaterialInstanceEditor.h
--- a/Engine/Source/Editor/MaterialEditor/Private/MaterialInstanceEditor.h
+++ b/Engine/Source/Editor/MaterialEditor/Private/MaterialInstanceEditor.h
@@ -122,6 +122,10 @@ public:
 	// IMaterial Editor Interface
 	virtual void GenerateInheritanceMenu(class UToolMenu* Menu) override;
 
+	void RefreshPreviewViewport();
+
+	bool IsDestructing() const { return bDestructing; }
+
 protected:
 	//~ FAssetEditorToolkit interface
 	virtual void SaveAsset_Execute() override;
@@ -210,6 +214,9 @@ private:
 	virtual void PostRedo( bool bSuccess ) override;
 	// End of FEditorUndoClient
 
+	/** Called to notify other material editors when a post process material with a UserSceneTexture output is loaded or unloaded, as the material may be used in their preview */
+	void NotifyUserSceneTextureLoadOrUnload();
+
 private:
 	struct FOnScreenMessage
 	{
@@ -259,6 +266,9 @@ private:
 	/** If editing instance of a function instead of a material. */
 	bool bIsFunctionPreviewMaterial;
 
+	/** Set to true when editor is being destructed */
+	bool bDestructing = false;
+
 	TSharedPtr<FExtensibilityManager> MenuExtensibilityManager;
 	TSharedPtr<FExtensibilityManager> ToolBarExtensibilityManager;
 
diff --git a/Engine/Source/Editor/MaterialEditor/Private/SMaterialEditorViewport.cpp b/Engine/Source/Editor/MaterialEditor/Private/SMaterialEditorViewport.cpp
--- a/Engine/Source/Editor/MaterialEditor/Private/SMaterialEditorViewport.cpp
+++ b/Engine/Source/Editor/MaterialEditor/Private/SMaterialEditorViewport.cpp
@@ -18,6 +18,7 @@
 #include "MaterialEditor/MaterialEditorMeshComponent.h"
 #include "MaterialEditor/PreviewMaterial.h"
 #include "Materials/MaterialExpressionUserSceneTexture.h"
+#include "Materials/MaterialInstanceConstant.h"
 #include "Engine/StaticMesh.h"
 #include "Engine/Selection.h"
 #include "Editor.h"
@@ -26,6 +27,7 @@
 #include "MaterialEditorActions.h"
 #include "Slate/SceneViewport.h"
 #include "MaterialEditor.h"
+#include "MaterialInstanceEditor.h"
 #include "SMaterialEditorViewportToolBar.h"
 #include "Widgets/Docking/SDockTab.h"
 #include "Engine/TextureCube.h"
@@ -432,23 +434,95 @@ bool SMaterialEditor3DPreviewViewport::SetPreviewAssetByName(const TCHAR* InAsse
 }
 
 // Add user scene texture inputs from a material.  Doesn't clear TSet first, so can be used to accumulate inputs from multiple materials.
-void GetUserSceneTextureInputs(UMaterial* Material, TSet<FName>& OutUserSceneTextures)
+static void GetUserSceneTextureInputs(UMaterialInterface* Material, TSet<FName>& OutUserSceneTextures)
 {
-	TArray<const UMaterialExpressionUserSceneTexture*> UserSceneTextureExpressions;
-	Material->GetAllExpressionsInMaterialAndFunctionsOfType(UserSceneTextureExpressions);
+	UMaterial* BaseMaterial = Material->GetBaseMaterial();
+	if (BaseMaterial)
+	{
+		// Get inputs from base material.  TMap key stores input, value stores instance override if present.
+		TMap<FName, FName> NewInputs;
+		TArray<const UMaterialExpressionUserSceneTexture*> UserSceneTextureExpressions;
+		BaseMaterial->GetAllExpressionsInMaterialAndFunctionsOfType(UserSceneTextureExpressions);
+
+		for (const UMaterialExpressionUserSceneTexture* UserSceneTextureExpression : UserSceneTextureExpressions)
+		{
+			if (!UserSceneTextureExpression->UserSceneTexture.IsNone())
+			{
+				NewInputs.Add(UserSceneTextureExpression->UserSceneTexture, NAME_None);
+			}
+		}
+
+		// Then get any overrides from material instances
+		TSet<UMaterialInterface*> MaterialDependencies;
+		Material->GetDependencies(MaterialDependencies);
+
+		for (UMaterialInterface* MaterialDependency : MaterialDependencies)
+		{
+			UMaterialInstanceConstant* MaterialInstance = Cast<UMaterialInstanceConstant>(MaterialDependency);
+			if (MaterialInstance)
+			{
+				for (FUserSceneTextureOverride& Override : MaterialInstance->UserSceneTextureOverrides)
+				{
+					FName* FoundInput = NewInputs.Find(Override.Key);
+
+					// Only accept the first override of a given value
+					if (FoundInput && FoundInput->IsNone())
+					{
+						*FoundInput = Override.Value;
+					}
+				}
+			}
+		}
+
+		// Finally, add the inputs to the output
+		for (auto NewInputIterator : NewInputs)
+		{
+			if (!NewInputIterator.Value.IsNone())
+			{
+				OutUserSceneTextures.Add(NewInputIterator.Value);
+			}
+			else
+			{
+				OutUserSceneTextures.Add(NewInputIterator.Key);
+			}
+		}
+	}
+}
+
+static bool IsParentOfEditedMaterialInstance(const FMaterialEditor* MaterialEditor, const TArray<UMaterialInstanceConstant*>& EditedMaterialInstances)
+{
+	for (UMaterialInstanceConstant* OtherInstance : EditedMaterialInstances)
+	{
+		UMaterial* BaseMaterial = OtherInstance->GetBaseMaterial();
+		if (BaseMaterial == MaterialEditor->Material || BaseMaterial == MaterialEditor->OriginalMaterial)
+		{
+			return true;
+		}
+	}
+	return false;
+}
 
-	for (const UMaterialExpressionUserSceneTexture* UserSceneTextureExpression : UserSceneTextureExpressions)
+static bool IsParentOfEditedMaterialInstance(UMaterialInstanceConstant* EditedMaterialInstance, const TArray<UMaterialInstanceConstant*>& EditedMaterialInstances)
+{
+	for (UMaterialInstanceConstant* OtherInstance : EditedMaterialInstances)
 	{
-		if (!UserSceneTextureExpression->UserSceneTexture.IsNone())
+		if (OtherInstance != EditedMaterialInstance)
 		{
-			OutUserSceneTextures.Add(UserSceneTextureExpression->UserSceneTexture);
+			FMaterialInheritanceChain OtherInheritanceChain;
+			OtherInstance->GetMaterialInheritanceChain(OtherInheritanceChain);
+
+			if (OtherInheritanceChain.MaterialInstances.Contains(EditedMaterialInstance))
+			{
+				return true;
+			}
 		}
 	}
+	return false;
 }
 
 // Recursively get all edited materials that have UserSceneTexture outputs that feed into Material.  Also returns inputs that
 // are missing, which may be useful to report as warnings to the log in the future.
-void GetUserSceneTextureDependencies(UMaterial* Material, TSet<UPreviewMaterial*>& OutDependencies, TSet<FName>& OutMissingInputs)
+static void GetUserSceneTextureDependencies(UMaterialInterface* Material, TSet<UMaterialInterface*>& OutDependencies, TSet<FName>& OutMissingInputs)
 {
 	// Check if the current material has any UserSceneTexture inputs first
 	TSet<FName> InputsToProcess;
@@ -459,21 +533,66 @@ void GetUserSceneTextureDependencies(UMaterial* Material, TSet<UPreviewMaterial*
 	}
 
 	// Generate a global list of edited materials that generate a given UserSceneTexture output (minus Material itself)
-	TMap<FName, TSet<UPreviewMaterial*>> MaterialsByUserSceneTextureOutput;
+	TMap<FName, TSet<UMaterialInterface*>> MaterialsByUserSceneTextureOutput;
 
 	UAssetEditorSubsystem* AssetEditorSubsystem = GEditor->GetEditorSubsystem<UAssetEditorSubsystem>();
 	TArray<UObject*> EditedAssets = AssetEditorSubsystem->GetAllEditedAssets();
-	for (int32 AssetIdx = 0; AssetIdx < EditedAssets.Num(); AssetIdx++)
+
+	// Get a list of material instances first -- if a material or instance is a parent of other loaded instances, only consider the outermost child instance
+	TArray<UMaterialInstanceConstant*> EditedMaterialInstances;
+	for (UObject* EditedAsset : EditedAssets)
+	{
+		UMaterialInstanceConstant* EditedMaterialInstance = Cast<UMaterialInstanceConstant>(EditedAsset);
+		if (EditedMaterialInstance)
+		{
+			EditedMaterialInstances.Add(EditedMaterialInstance);
+		}
+	}
+
+	for (UObject* EditedAsset : EditedAssets)
 	{
-		UObject* EditedAsset = EditedAssets[AssetIdx];
 		UPreviewMaterial* EditedMaterial = Cast<UPreviewMaterial>(EditedAsset);
 
 		if (EditedMaterial && EditedMaterial != Material && EditedMaterial->IsPostProcessMaterial() && !EditedMaterial->UserSceneTexture.IsNone())
 		{
 			TArray<IAssetEditorInstance*> Editors = AssetEditorSubsystem->FindEditorsForAsset(EditedAsset);
-			if (!Editors.IsEmpty() && Editors[0]->GetEditorName() == FName("MaterialEditor") && !((const FMaterialEditor*)Editors[0])->bDestructing)
+			if (!Editors.IsEmpty() && Editors[0]->GetEditorName() == FName("MaterialEditor"))
 			{
-				MaterialsByUserSceneTextureOutput.FindOrAdd(EditedMaterial->UserSceneTexture).Add(EditedMaterial);
+				const FMaterialEditor* MaterialEditor = (const FMaterialEditor*)Editors[0];
+
+				// If we are editing a material instance and its parent, we only want the child instance to be previewed.  Previewing multiple
+				// copies of the same base material would cause confusing and indeterminate results.  We pass in the MaterialEditor rather than
+				// the material, so we can check against both original and previewed variations of the material.
+				if (!MaterialEditor->bDestructing && !IsParentOfEditedMaterialInstance(MaterialEditor, EditedMaterialInstances))
+				{
+					MaterialsByUserSceneTextureOutput.FindOrAdd(EditedMaterial->UserSceneTexture).Add(EditedMaterial);
+				}
+			}
+		}
+
+		UMaterialInstanceConstant* EditedMaterialInstance = Cast<UMaterialInstanceConstant>(EditedAsset);
+
+		// If we are editing a material instance and its parent, we only want the child instance to be previewed.
+		if (EditedMaterialInstance && EditedMaterialInstance != Material && !IsParentOfEditedMaterialInstance(EditedMaterialInstance, EditedMaterialInstances))
+		{
+			UMaterial* BaseMaterial = EditedMaterialInstance->GetMaterial();
+
+			if (BaseMaterial->IsPostProcessMaterial())
+			{
+				FName UserSceneTextureOutput = EditedMaterialInstance->GetUserSceneTextureOutput(BaseMaterial);
+				if (UserSceneTextureOutput != NAME_None)
+				{
+					TArray<IAssetEditorInstance*> Editors = AssetEditorSubsystem->FindEditorsForAsset(EditedAsset);
+					if (!Editors.IsEmpty() && Editors[0]->GetEditorName() == FName("MaterialInstanceEditor"))
+					{
+						const FMaterialInstanceEditor* MaterialInstanceEditor = (const FMaterialInstanceEditor*)Editors[0];
+
+						if (!MaterialInstanceEditor->IsDestructing())
+						{
+							MaterialsByUserSceneTextureOutput.FindOrAdd(UserSceneTextureOutput).Add(EditedMaterialInstance);
+						}
+					}
+				}
 			}
 		}
 	}
@@ -484,7 +603,7 @@ void GetUserSceneTextureDependencies(UMaterial* Material, TSet<UPreviewMaterial*
 	{
 		// Find materials that generate an input we care about
 		FName Input = InputsToProcess[FSetElementId::FromInteger(ElementIndex)];
-		TSet<UPreviewMaterial*>* MaterialsGeneratingInput = MaterialsByUserSceneTextureOutput.Find(Input);
+		TSet<UMaterialInterface*>* MaterialsGeneratingInput = MaterialsByUserSceneTextureOutput.Find(Input);
 
 		if (MaterialsGeneratingInput)
 		{
@@ -492,7 +611,7 @@ void GetUserSceneTextureDependencies(UMaterial* Material, TSet<UPreviewMaterial*
 			OutDependencies.Append(*MaterialsGeneratingInput);
 
 			// Add any inputs the new dependencies require
-			for (UPreviewMaterial* MaterialGeneratingOutput : *MaterialsGeneratingInput)
+			for (UMaterialInterface* MaterialGeneratingOutput : *MaterialsGeneratingInput)
 			{
 				GetUserSceneTextureInputs(MaterialGeneratingOutput, InputsToProcess);
 			}
@@ -522,6 +641,24 @@ void SMaterialEditor3DPreviewViewport::SetPreviewMaterial(UMaterialInterface* In
 		// Clear blendables, and re-add them (cleans up any post process materials with UserSceneTextures that are no longer used or loaded)
 		PostProcessVolumeActor->Settings.WeightedBlendables.Array.Empty(1);
 
+		{
+			// Add any edited post process materials that write UserSceneTextures used by this material, for better visualization.
+			// We want to add these before the main material, so the main material renders last (assuming equal priority).
+			TSet<UMaterialInterface*> UserSceneTextureDependencies;
+			TSet<FName> UserSceneTextureMissingInputs;
+
+			GetUserSceneTextureDependencies(PreviewMaterial, UserSceneTextureDependencies, UserSceneTextureMissingInputs);
+
+			// Add dependencies in reverse order, as dependency tree is traversed backwards from the main material's inputs, so earlier
+			// dependencies tend to end up later in the TSet.  This isn't a perfect dependency sort, but works correctly in typical cases.
+			// The user should set Blendable Priority as needed to sort in complex cases.
+			for (int32 DependencyIndex = UserSceneTextureDependencies.Num() - 1; DependencyIndex >= 0; --DependencyIndex)
+			{
+				UMaterialInterface* Dependency = UserSceneTextureDependencies[FSetElementId::FromInteger(DependencyIndex)];
+				PostProcessVolumeActor->AddOrUpdateBlendable(Dependency);
+			}
+		}
+
 		check (PreviewMaterial != nullptr);
 		PostProcessVolumeActor->AddOrUpdateBlendable(PreviewMaterial);
 		PostProcessVolumeActor->bEnabled = true;
@@ -538,18 +675,7 @@ void SMaterialEditor3DPreviewViewport::SetPreviewMaterial(UMaterialInterface* In
 			PreviewMeshComponent->MarkRenderStateDirty();
 		}
 
-		{
-			// Add any edited post process materials that write UserSceneTextures used by this material, for better visualization
-			TSet<UPreviewMaterial*> UserSceneTextureDependencies;
-			TSet<FName> UserSceneTextureMissingInputs;
-
-			GetUserSceneTextureDependencies(PreviewMaterial->GetMaterial(), UserSceneTextureDependencies, UserSceneTextureMissingInputs);
-
-			for (UPreviewMaterial* Dependency : UserSceneTextureDependencies)
-			{
-				PostProcessVolumeActor->AddOrUpdateBlendable(Dependency);
-			}
-		}
+		GetViewportClient()->RedrawRequested(GetSceneViewport().Get());
 	}
 	else
 	{
@@ -858,11 +984,21 @@ void SMaterialEditor3DPreviewViewport::OnPropertyChanged(UObject* ObjectBeingMod
 	FProperty* PropertyThatChanged = PropertyChangedEvent.Property;
 	static const FString MaterialDomain = TEXT("MaterialDomain");
 	static const FString UserSceneTexture = TEXT("UserSceneTexture");
-	if (ObjectBeingModified != nullptr && PropertyThatChanged != nullptr && (PropertyThatChanged->GetName() == MaterialDomain || PropertyThatChanged->GetName() == UserSceneTexture))
+	static const FString PostProcessOverrides = TEXT("PostProcessOverrides");
+
+	// We need to refresh other edited post process materials when a change is made that affects UserSceneTexture inputs or outputs, as previews
+	// include materials that generate UserSceneTexture dependencies.  Changing the material domain potentially converts a material to or from a
+	// Post Process domain material, adding or removing it as relevant to other previews.  Or changing any UserSceneTexture input or output, which
+	// includes the "UMaterial::UserSceneTexture" field, plus any FName field in the PostProcessOverrides member struct.
+	//
+	// We also need to refresh PreviewMaterial itself if its domain changes, regardless of whether it's a post process material.
+	if (ObjectBeingModified != nullptr && PropertyThatChanged != nullptr)
 	{
-		// Refresh preview material if this specific material has been modified, or if this is any other post process material that has an editor preview,
-		// which may use this preview material in its preview if it has the UserSceneTexture as one of its inputs.
-		if (ObjectBeingModified == PreviewMaterial || PreviewMaterial->GetMaterial()->IsPostProcessMaterial())
+		if ((ObjectBeingModified == PreviewMaterial && PropertyThatChanged->GetName() == MaterialDomain) ||
+			(PreviewMaterial->GetMaterial()->IsPostProcessMaterial() &&
+			 (PropertyThatChanged->GetName() == MaterialDomain ||
+			  PropertyThatChanged->GetName() == UserSceneTexture ||
+			  (PropertyThatChanged->IsA<FNameProperty>() && PropertyChangedEvent.MemberProperty && PropertyChangedEvent.MemberProperty->GetName() == PostProcessOverrides))))
 		{
 			SetPreviewMaterial(PreviewMaterial);
 		}
diff --git a/Engine/Source/Editor/MaterialEditor/Public/MaterialEditorUtilities.h b/Engine/Source/Editor/MaterialEditor/Public/MaterialEditorUtilities.h
--- a/Engine/Source/Editor/MaterialEditor/Public/MaterialEditorUtilities.h
+++ b/Engine/Source/Editor/MaterialEditor/Public/MaterialEditorUtilities.h
@@ -235,6 +235,13 @@ public:
 	static void OnShowMaterialInContentBrowser(const FAssetData InMaterial);
 	static void OnShowFunctionInContentBrowser(const FAssetData InFunction);
 
+	/**
+	 * Triggers a refresh or redraw for post process preview materials and material instances.  A refresh is useful when a material generating a UserSceneTexture
+	 * output is loaded or unloaded.  Other loaded materials that have UserSceneTexture inputs may include the material in question in their preview, which
+	 * necessitates a refresh.  A redraw is useful when debug settings change that may affect post process material viewports.
+	 */
+	static void RefreshPostProcessPreviewMaterials(UMaterialInterface* ExcludeMaterialInterface, bool bRedrawOnly = false);
+
 private:
 
 	static void OpenSelectedParentEditor(UMaterialFunctionInterface* InMaterialFunction);
diff --git a/Engine/Source/Editor/UnrealEd/Classes/MaterialEditor/MaterialEditorInstanceConstant.h b/Engine/Source/Editor/UnrealEd/Classes/MaterialEditor/MaterialEditorInstanceConstant.h
--- a/Engine/Source/Editor/UnrealEd/Classes/MaterialEditor/MaterialEditorInstanceConstant.h
+++ b/Engine/Source/Editor/UnrealEd/Classes/MaterialEditor/MaterialEditorInstanceConstant.h
@@ -12,6 +12,7 @@
 #include "Misc/Guid.h"
 #include "StaticParameterSet.h"
 #include "Editor/UnrealEdTypes.h"
+#include "Engine/BlendableInterface.h"
 #include "Materials/MaterialInstanceBasePropertyOverrides.h"
 #include "Materials/MaterialExpression.h"
 #include "MaterialEditorInstanceConstant.generated.h"
@@ -217,6 +218,48 @@ struct FEditorStaticComponentMaskParameterValue : public FEditorParameterValue
 	}
 };
 
+USTRUCT()
+struct FEditorUserSceneTextureOverride
+{
+	GENERATED_USTRUCT_BODY()
+
+	UPROPERTY()
+	FName Key;
+
+	UPROPERTY(EditAnywhere, Category = EditorParameterValue)
+	FName Value;
+};
+
+USTRUCT()
+struct FMaterialEditorPostProcessOverrides
+{
+	GENERATED_USTRUCT_BODY()
+
+	// Tracks if this is a material where post process overrides can be applied (MaterialDomain == MD_PostProcess, BlendableLocation != BL_ReplacingTonemapper)
+	UPROPERTY()
+	bool bIsOverrideable;
+
+	UPROPERTY(EditAnywhere, Category = PostProcessOverrideValue)
+	bool bOverrideBlendableLocation;
+
+	UPROPERTY(EditAnywhere, Category = PostProcessOverrideValue)
+	bool bOverrideBlendablePriority;
+
+	UPROPERTY(EditAnywhere, Category = PostProcessOverrideValue, meta = (DisplayName = "Blendable Location"), meta = (InvalidEnumValues = "BL_ReplacingTonemapper"))
+	TEnumAsByte<EBlendableLocation> BlendableLocationOverride;
+
+	UPROPERTY(EditAnywhere, Category = PostProcessOverrideValue, meta = (DisplayName = "Blendable Priority"))
+	int32 BlendablePriorityOverride;
+
+	/** Overrides for user scene texture inputs */
+	UPROPERTY(EditAnywhere, editfixedsize, Category = PostProcessOverrideValue)
+	TArray<FEditorUserSceneTextureOverride> UserSceneTextureInputs;
+
+	/** Override for user scene texture output */
+	UPROPERTY(EditAnywhere, Category = PostProcessOverrideValue)
+	FName UserSceneTextureOutput;
+};
+
 UCLASS(hidecategories=Object, collapsecategories, MinimalAPI)
 class UMaterialEditorInstanceConstant : public UObject
 {
@@ -282,6 +325,10 @@ class UMaterialEditorInstanceConstant : public UObject
 	UPROPERTY(EditAnywhere, Category = MaterialEditorInstanceConstant, meta = (editcondition = "bNaniteOverride"))
 	TObjectPtr<UMaterialInterface> NaniteOverrideMaterial;
 
+	/** Overrides specific to Post Process domain materials. */
+	UPROPERTY(EditAnywhere, Category = PostProcessOverrides)
+	FMaterialEditorPostProcessOverrides PostProcessOverrides;
+
 	//~ Begin UObject Interface.
 	UNREALED_API virtual void PostEditChangeProperty(FPropertyChangedEvent& PropertyChangedEvent) override;
 #if WITH_EDITOR
diff --git a/Engine/Source/Editor/UnrealEd/Private/Commandlets/DumpMaterialInfo.cpp b/Engine/Source/Editor/UnrealEd/Private/Commandlets/DumpMaterialInfo.cpp
--- a/Engine/Source/Editor/UnrealEd/Private/Commandlets/DumpMaterialInfo.cpp
+++ b/Engine/Source/Editor/UnrealEd/Private/Commandlets/DumpMaterialInfo.cpp
@@ -320,6 +320,7 @@ TArray<MaterialInfo::FPropertySet> GetMaterialInfoProperties()
 		FPropertySet::FromAccessor<int32, &FMaterialResource::GetNumCustomizedUVs>(TEXT("GetNumCustomizedUVs")),
 		FPropertySet::FromAccessor<int32, &FMaterialResource::GetBlendableLocation>(TEXT("GetBlendableLocation")),
 		FPropertySet::FromAccessor<bool, &FMaterialResource::GetBlendableOutputAlpha>(TEXT("GetBlendableOutputAlpha")),
+		FPropertySet::FromAccessor<bool, &FMaterialResource::GetDisablePreExposureScale>(TEXT("GetDisablePreExposureScale")),
 		FPropertySet::FromAccessor<bool, &FMaterialResource::IsStencilTestEnabled>(TEXT("IsStencilTestEnabled")),
 		FPropertySet::FromAccessor<uint32, &FMaterialResource::GetStencilRefValue>(TEXT("GetStencilRefValue")),
 		FPropertySet::FromAccessor<uint32, &FMaterialResource::GetStencilCompare>(TEXT("GetStencilCompare")),
diff --git a/Engine/Source/Editor/UnrealEd/Private/PreviewMaterial.cpp b/Engine/Source/Editor/UnrealEd/Private/PreviewMaterial.cpp
--- a/Engine/Source/Editor/UnrealEd/Private/PreviewMaterial.cpp
+++ b/Engine/Source/Editor/UnrealEd/Private/PreviewMaterial.cpp
@@ -36,6 +36,7 @@
 #include "MaterialEditingLibrary.h"
 #include "MaterialPropertyHelpers.h"
 #include "MaterialStatsCommon.h"
+#include "MaterialDomain.h"
 
 /**
  * Class for rendering the material on the preview mesh in the Material Editor
@@ -589,6 +590,8 @@ void UMaterialEditorInstanceConstant::RegenerateArrays()
 {
 	VisibleExpressions.Empty();
 	ParameterGroups.Empty();
+	PostProcessOverrides.bIsOverrideable = false;
+	PostProcessOverrides.UserSceneTextureInputs.Empty();
 
 	if (Parent)
 	{	
@@ -621,6 +624,31 @@ void UMaterialEditorInstanceConstant::RegenerateArrays()
 			}
 		}
 
+		if (ParentMaterial->MaterialDomain == MD_PostProcess && ParentMaterial->BlendableLocation != BL_ReplacingTonemapper)
+		{
+			PostProcessOverrides.bIsOverrideable = true;
+
+			UMaterialInterfaceEditorOnlyData* ParentMaterialEditorData = ParentMaterial->GetEditorOnlyData();
+			if (ParentMaterialEditorData && ParentMaterialEditorData->CachedExpressionData.IsValid())
+			{
+				const FMaterialCachedExpressionEditorOnlyData* ParentMaterialExpressionData = ParentMaterialEditorData->CachedExpressionData.Get();
+				for (FName UserSceneTextureInput : ParentMaterialExpressionData->UserSceneTextureInputs)
+				{
+					FName OverrideFound = NAME_None;
+					for (const FUserSceneTextureOverride& Override : SourceInstance->UserSceneTextureOverrides)
+					{
+						if (Override.Key == UserSceneTextureInput)
+						{
+							OverrideFound = Override.Value;
+							break;
+						}
+					}
+
+					PostProcessOverrides.UserSceneTextureInputs.Add({ UserSceneTextureInput, OverrideFound });
+				}
+			}
+		}
+
 		IMaterialEditorModule* MaterialEditorModule = &FModuleManager::LoadModuleChecked<IMaterialEditorModule>( "MaterialEditor" );
 		MaterialEditorModule->GetVisibleMaterialParameters(ParentMaterial, SourceInstance, VisibleExpressions);
 	}
@@ -854,6 +882,41 @@ void UMaterialEditorInstanceConstant::CopyToSourceInstance(const bool bForceStat
 		FMaterialParameterInfo RefractionInfo(TEXT("RefractionDepthBias"));
 		SourceInstance->SetScalarParameterValueEditorOnly(RefractionInfo, RefractionDepthBias);
 
+		// Copy UserSceneTextureOverrides
+		SourceInstance->UserSceneTextureOverrides.Empty();
+		for (FEditorUserSceneTextureOverride& Override : PostProcessOverrides.UserSceneTextureInputs)
+		{
+			if (Override.Value != NAME_None && Override.Value != Override.Key)
+			{
+				SourceInstance->UserSceneTextureOverrides.Add({ Override.Key, Override.Value });
+			}
+		}
+		if (PostProcessOverrides.UserSceneTextureOutput != NAME_None)
+		{
+			// UserSceneTextureOutput override uses key of NAME_None
+			SourceInstance->UserSceneTextureOverrides.Add({ NAME_None, PostProcessOverrides.UserSceneTextureOutput });
+		}
+
+		// Copy other post process overrides (BlendableLocation / BlendablePriority) -- BL_ReplacingTonemapper is disallowed on overrides
+		SourceInstance->bOverrideBlendableLocation = PostProcessOverrides.bOverrideBlendableLocation && PostProcessOverrides.BlendableLocationOverride != BL_ReplacingTonemapper;
+		SourceInstance->bOverrideBlendablePriority = PostProcessOverrides.bOverrideBlendablePriority;
+		if (SourceInstance->bOverrideBlendableLocation)
+		{
+			SourceInstance->BlendableLocationOverride = PostProcessOverrides.BlendableLocationOverride;
+		}
+		else
+		{
+			SourceInstance->BlendableLocationOverride = Parent ? Parent->GetMaterial()->BlendableLocation : TEnumAsByte<EBlendableLocation>(BL_SceneColorAfterTonemapping);
+		}
+		if (SourceInstance->bOverrideBlendablePriority)
+		{
+			SourceInstance->BlendablePriorityOverride = PostProcessOverrides.BlendablePriorityOverride;
+		}
+		else
+		{
+			SourceInstance->BlendablePriorityOverride = Parent ? Parent->GetMaterial()->BlendablePriority : 0;
+		}
+
 		SourceInstance->bOverrideSubsurfaceProfile = bOverrideSubsurfaceProfile;
 		SourceInstance->SubsurfaceProfile = SubsurfaceProfile;
 
@@ -1042,6 +1105,38 @@ void UMaterialEditorInstanceConstant::CopyBasePropertiesFromParent()
 	bOverrideSubsurfaceProfile = SourceInstance->bOverrideSubsurfaceProfile;
 	// Copy the subsurface profile. GetSubsurfaceProfile_Internal() will return either the overridden profile or one from a parent
 	SubsurfaceProfile = SourceInstance->GetSubsurfaceProfile_Internal();
+
+	// Post process blendable location and priority overrides
+	PostProcessOverrides.bOverrideBlendableLocation = SourceInstance->bOverrideBlendableLocation;
+	if (SourceInstance->bOverrideBlendableLocation)
+	{
+		PostProcessOverrides.BlendableLocationOverride = SourceInstance->BlendableLocationOverride;
+	}
+	else
+	{
+		PostProcessOverrides.BlendableLocationOverride = Parent ? Parent->GetMaterial()->BlendableLocation : TEnumAsByte<EBlendableLocation>(BL_SceneColorAfterTonemapping);
+	}
+
+	PostProcessOverrides.bOverrideBlendablePriority = SourceInstance->bOverrideBlendablePriority;
+	if (SourceInstance->bOverrideBlendablePriority)
+	{
+		PostProcessOverrides.BlendablePriorityOverride = SourceInstance->BlendablePriorityOverride;
+	}
+	else
+	{
+		PostProcessOverrides.BlendablePriorityOverride = Parent ? Parent->GetMaterial()->BlendablePriority : 0;
+	}
+
+	// UserSceneTextureOutput override uses Key == NAME_None.  UserSceneTextureInputs are initialized in RegenerateArrays, as those are affected
+	// by graph reachability.
+	for (const FUserSceneTextureOverride& Override : SourceInstance->UserSceneTextureOverrides)
+	{
+		if (Override.Key == NAME_None)
+		{
+			PostProcessOverrides.UserSceneTextureOutput = Override.Value;
+			break;
+		}
+	}
 }
 
 #if WITH_EDITOR
diff --git a/Engine/Source/Runtime/Core/Private/UObject/DevObjectVersion.cpp b/Engine/Source/Runtime/Core/Private/UObject/DevObjectVersion.cpp
--- a/Engine/Source/Runtime/Core/Private/UObject/DevObjectVersion.cpp
+++ b/Engine/Source/Runtime/Core/Private/UObject/DevObjectVersion.cpp
@@ -135,7 +135,7 @@ FDevSystemGuids::FDevSystemGuids()
 	, POSESEARCHDB_DERIVEDDATA_VER(0x389117E4, 0x807A4CC0, 0x9F37C2E6, 0xD808A78D)
 	, SkeletalMeshDerivedDataVersion(0x9B5F4544, 0x76D7481C, 0x9AD3F614, 0xA6C07904)
 	, STATICMESH_DERIVEDDATA_VER(0x2C6C400C, 0x2EDF47B3, 0x9BD36689, 0x7B77B208)
-	, MaterialTranslationDDCVersion(0xE6097D5A, 0x39794E1B, 0xA3FDDA91, 0x6E2221B6)
+	, MaterialTranslationDDCVersion(0x628D9AFC, 0xC10245C5, 0x86967652, 0xAD52E887)
 {
 }
 
diff --git a/Engine/Source/Runtime/Engine/Private/Materials/HLSLMaterialTranslator.cpp b/Engine/Source/Runtime/Engine/Private/Materials/HLSLMaterialTranslator.cpp
--- a/Engine/Source/Runtime/Engine/Private/Materials/HLSLMaterialTranslator.cpp
+++ b/Engine/Source/Runtime/Engine/Private/Materials/HLSLMaterialTranslator.cpp
@@ -7849,7 +7849,7 @@ int32 FHLSLMaterialTranslator::SceneTextureLookup(int32 ViewportUV, uint32 InSce
 	SubstrateCtx.SubstrateMaterialComplexity.bIsSingle = false;
 	
 	if (((SceneTextureId >= PPI_PostProcessInput0 && SceneTextureId <= PPI_PostProcessInput6) || (SceneTextureId >= PPI_UserSceneTexture0 && SceneTextureId <= PPI_UserSceneTexture6)) &&
-		Material->GetMaterialDomain() == MD_PostProcess && Material->GetBlendableLocation() != BL_SceneColorAfterTonemapping)
+		Material->GetMaterialDomain() == MD_PostProcess && Material->GetBlendableLocation() != BL_SceneColorAfterTonemapping && !Material->GetDisablePreExposureScale())
 	{
 		return AddInlinedCodeChunk(MCT_Float4, TEXT("(float4(View.OneOverPreExposure.xxx, 1) * %s)"), *CoerceParameter(LookUp, MCT_Float4));
 	}
diff --git a/Engine/Source/Runtime/Engine/Private/Materials/Material.cpp b/Engine/Source/Runtime/Engine/Private/Materials/Material.cpp
--- a/Engine/Source/Runtime/Engine/Private/Materials/Material.cpp
+++ b/Engine/Source/Runtime/Engine/Private/Materials/Material.cpp
@@ -1928,7 +1928,7 @@ void UMaterial::GetDependentFunctions(TArray<UMaterialFunctionInterface*>& Depen
 }
 #endif // WITH_EDITORONLY_DATA
 
-extern FPostProcessMaterialNode* IteratePostProcessMaterialNodes(const FFinalPostProcessSettings& Dest, const UMaterial* Material, FBlendableEntry*& Iterator);
+extern FPostProcessMaterialNode* IteratePostProcessMaterialNodes(const FFinalPostProcessSettings& Dest, const UMaterialInterface* Material, const UMaterial* Base, FBlendableEntry*& Iterator);
 
 void UMaterialInterface::OverrideBlendableSettings(class FSceneView& View, float Weight) const
 {
@@ -1947,7 +1947,7 @@ void UMaterialInterface::OverrideBlendableSettings(class FSceneView& View, float
 
 	FBlendableEntry* Iterator = 0;
 
-	FPostProcessMaterialNode* DestNode = IteratePostProcessMaterialNodes(Dest, Base, Iterator);
+	FPostProcessMaterialNode* DestNode = IteratePostProcessMaterialNodes(Dest, this, Base, Iterator);
 
 	// is this the first one of this material?
 	if(!DestNode)
@@ -1962,7 +1962,7 @@ void UMaterialInterface::OverrideBlendableSettings(class FSceneView& View, float
 
 			InitialMID->CopyScalarAndVectorParameters(*SourceData, View.FeatureLevel);
 
-			FPostProcessMaterialNode InitialNode(InitialMID, Base->BlendableLocation, Base->BlendablePriority, Base->bIsBlendable);
+			FPostProcessMaterialNode InitialNode(InitialMID, GetBlendableLocation(Base), GetBlendablePriority(Base), Base->bIsBlendable);
 
 			// no blending needed on this one
 			FPostProcessMaterialNode* InitialDestNode = Dest.BlendableManager.PushBlendableData(1.0f, InitialNode);
@@ -2163,6 +2163,18 @@ bool UMaterial::GetRefractionSettings(float& OutBiasValue) const
 	return true;
 }
 
+EBlendableLocation UMaterial::GetBlendableLocation(const UMaterial* Base) const
+{
+	check(Base == this);
+	return BlendableLocation;
+}
+
+int32 UMaterial::GetBlendablePriority(const UMaterial* Base) const
+{
+	check(Base == this);
+	return BlendablePriority;
+}
+
 void UMaterial::GetDependencies(TSet<UMaterialInterface*>& Dependencies) 
 {
 	Dependencies.Add(this);
@@ -4599,6 +4611,7 @@ bool UMaterial::CanEditChange(const FProperty* InProperty) const
 		if (PropertyName == GET_MEMBER_NAME_STRING_CHECKED(UMaterial, BlendableLocation) ||
 			PropertyName == GET_MEMBER_NAME_STRING_CHECKED(UMaterial, BlendablePriority) || 
 			PropertyName == GET_MEMBER_NAME_STRING_CHECKED(UMaterial, BlendableOutputAlpha) ||
+			PropertyName == GET_MEMBER_NAME_STRING_CHECKED(UMaterial, bDisablePreExposureScale) ||
 			PropertyName == GET_MEMBER_NAME_STRING_CHECKED(UMaterial, bIsBlendable) ||
 			PropertyName == GET_MEMBER_NAME_STRING_CHECKED(UMaterial, bEnableStencilTest) ||
 			PropertyName == GET_MEMBER_NAME_STRING_CHECKED(UMaterial, StencilCompare) ||
@@ -4611,7 +4624,8 @@ bool UMaterial::CanEditChange(const FProperty* InProperty) const
 		}
 
 		if (PropertyName == GET_MEMBER_NAME_STRING_CHECKED(UMaterial, UserSceneTexture) ||
-			PropertyName == GET_MEMBER_NAME_STRING_CHECKED(UMaterial, UserTextureDivisor))
+			PropertyName == GET_MEMBER_NAME_STRING_CHECKED(UMaterial, UserTextureDivisor) ||
+			PropertyName == GET_MEMBER_NAME_STRING_CHECKED(UMaterial, ResolutionRelativeToInput))
 		{
 			// "Replacing Tonemapper" blendable location doesn't support a UserSceneTexture output
 			return MaterialDomain == MD_PostProcess && BlendableLocation != BL_ReplacingTonemapper;
diff --git a/Engine/Source/Runtime/Engine/Private/Materials/MaterialCachedData.cpp b/Engine/Source/Runtime/Engine/Private/Materials/MaterialCachedData.cpp
--- a/Engine/Source/Runtime/Engine/Private/Materials/MaterialCachedData.cpp
+++ b/Engine/Source/Runtime/Engine/Private/Materials/MaterialCachedData.cpp
@@ -27,6 +27,7 @@
 #include "Materials/MaterialExpressionFunctionInput.h"
 #include "Materials/MaterialExpressionStaticBool.h"
 #include "Materials/MaterialExpressionLandscapeGrassOutput.h"
+#include "Materials/MaterialExpressionUserSceneTexture.h"
 #include "Materials/MaterialFunctionInterface.h"
 #include "Materials/MaterialParameterCollection.h"
 #include "MaterialHLSLTree.h"
@@ -394,6 +395,14 @@ void FMaterialCachedExpressionData::UpdateForExpressions(const FMaterialCachedEx
 
 		Expression->GetIncludeFilePaths(EditorOnlyData->ExpressionIncludeFilePaths);
 
+		if (UMaterialExpressionUserSceneTexture* ExpressionUserSceneTexture = Cast<UMaterialExpressionUserSceneTexture>(Expression))
+		{
+			if (!ExpressionUserSceneTexture->UserSceneTexture.IsNone())
+			{
+				EditorOnlyData->UserSceneTextureInputs.Add(ExpressionUserSceneTexture->UserSceneTexture);
+			}
+		}
+
 		if (UMaterialExpressionCollectionParameter* ExpressionCollectionParameter = Cast<UMaterialExpressionCollectionParameter>(Expression))
 		{
 			UMaterialParameterCollection* Collection = ExpressionCollectionParameter->Collection;
diff --git a/Engine/Source/Runtime/Engine/Private/Materials/MaterialHLSLTree.cpp b/Engine/Source/Runtime/Engine/Private/Materials/MaterialHLSLTree.cpp
--- a/Engine/Source/Runtime/Engine/Private/Materials/MaterialHLSLTree.cpp
+++ b/Engine/Source/Runtime/Engine/Private/Materials/MaterialHLSLTree.cpp
@@ -2431,7 +2431,7 @@ void FExpressionSceneTexture::EmitValueShader(FEmitContext& Context, FEmitScope&
 	}
 
 	if (((SceneTextureId >= PPI_PostProcessInput0 && SceneTextureId <= PPI_PostProcessInput6) || (SceneTextureId >= PPI_UserSceneTexture0 && SceneTextureId <= PPI_UserSceneTexture6)) &&
-		Context.Material->GetMaterialDomain() == MD_PostProcess && Context.Material->GetBlendableLocation() != BL_SceneColorAfterTonemapping)
+		Context.Material->GetMaterialDomain() == MD_PostProcess && Context.Material->GetBlendableLocation() != BL_SceneColorAfterTonemapping && !Context.Material->GetDisablePreExposureScale())
 	{
 		EmitLookup = Context.EmitExpression(Scope, Shader::EValueType::Float4, TEXT("(float4(View.OneOverPreExposure.xxx, 1) * %)"), EmitLookup);
 	}
diff --git a/Engine/Source/Runtime/Engine/Private/Materials/MaterialInstance.cpp b/Engine/Source/Runtime/Engine/Private/Materials/MaterialInstance.cpp
--- a/Engine/Source/Runtime/Engine/Private/Materials/MaterialInstance.cpp
+++ b/Engine/Source/Runtime/Engine/Private/Materials/MaterialInstance.cpp
@@ -341,6 +341,77 @@ bool FMaterialInstanceResource::GetParameterValue(EMaterialParameterType Type, c
 	return bResult;
 }
 
+bool FMaterialInstanceResource::GetUserSceneTextureOverride(FName& InOutName) const
+{
+	checkSlow(IsInParallelRenderingThread());
+
+	// Number of overrides possible is small (maximum 6, in most practical cases 1 or 2), and FName comparison cheap,
+	// so the assumption is that an array search will be cheaper than the overhead of going through a hash lookup.
+	// Plus an array takes half the space of THashedMaterialParameterMap, saving memory.
+	for (const FUserSceneTextureOverride& Override : UserSceneTextureOverrides)
+	{
+		if (Override.Key == InOutName && Override.Value != NAME_None)
+		{
+			InOutName = Override.Value;
+			return true;
+		}
+	}
+
+	if (Parent)
+	{
+		return Parent->GetRenderProxy()->GetUserSceneTextureOverride(InOutName);
+	}
+	else
+	{
+		return false;
+	}
+}
+
+EBlendableLocation FMaterialInstanceResource::GetBlendableLocation(const FMaterial* Base) const
+{
+	check(Base);
+	checkSlow(IsInParallelRenderingThread());
+
+	// Can't be overridden to BL_ReplacingTonemapper
+	if (PostProcessBlendableOverrides.bOverrideBlendableLocation && PostProcessBlendableOverrides.BlendableLocationOverride != BL_ReplacingTonemapper)
+	{
+		// Can't be overridden from BL_ReplacingTonemapper 
+		if ((EBlendableLocation)Base->GetBlendableLocation() == BL_ReplacingTonemapper)
+		{
+			return BL_ReplacingTonemapper;
+		}
+
+		return PostProcessBlendableOverrides.BlendableLocationOverride;
+	}
+	else if (Parent)
+	{
+		return Parent->GetRenderProxy()->GetBlendableLocation(Base);
+	}
+	else
+	{
+		return (EBlendableLocation)Base->GetBlendableLocation();
+	}
+}
+
+int32 FMaterialInstanceResource::GetBlendablePriority(const FMaterial* Base) const
+{
+	check(Base);
+	checkSlow(IsInParallelRenderingThread());
+
+	if (PostProcessBlendableOverrides.bOverrideBlendablePriority)
+	{
+		return PostProcessBlendableOverrides.BlendablePriorityOverride;
+	}
+	else if (Parent)
+	{
+		return Parent->GetRenderProxy()->GetBlendablePriority(Base);
+	}
+	else
+	{
+		return Base->GetBlendablePriority();
+	}
+}
+
 void UMaterialInstance::PropagateDataToMaterialProxy()
 {
 	if (Resource)
@@ -426,6 +497,8 @@ void FMaterialInstanceResource::InitMIParameters(FMaterialInstanceParameterSet&
 	TextureCollectionParameterArray.Array = MoveTemp(ParameterSet.TextureCollectionParameters);
 	RuntimeVirtualTextureParameterArray.Array = MoveTemp(ParameterSet.RuntimeVirtualTextureParameters);
 	SparseVolumeTextureParameterArray.Array = MoveTemp(ParameterSet.SparseVolumeTextureParameters);
+	UserSceneTextureOverrides = MoveTemp(ParameterSet.UserSceneTextureOverrides);
+	PostProcessBlendableOverrides = ParameterSet.PostProcessBlendableOverrides;
 
 
 	// Build hash tables.
@@ -839,6 +912,12 @@ void GameThread_InitMIParameters(const UMaterialInstance& Instance)
 		}
 	}
 
+	ParameterSet.UserSceneTextureOverrides = Instance.UserSceneTextureOverrides;
+	ParameterSet.PostProcessBlendableOverrides.bOverrideBlendableLocation = Instance.bOverrideBlendableLocation;
+	ParameterSet.PostProcessBlendableOverrides.bOverrideBlendablePriority = Instance.bOverrideBlendablePriority;
+	ParameterSet.PostProcessBlendableOverrides.BlendableLocationOverride = Instance.BlendableLocationOverride;
+	ParameterSet.PostProcessBlendableOverrides.BlendablePriorityOverride = Instance.BlendablePriorityOverride;
+
 	ENQUEUE_RENDER_COMMAND(InitMIParameters)(
 		[Resource, Parameters = MoveTemp(ParameterSet)](FRHICommandListImmediate& RHICmdList) mutable
 		{
@@ -1098,6 +1177,91 @@ bool UMaterialInstance::GetRefractionSettings(float& OutBiasValue) const
 	}
 }
 
+bool UMaterialInstance::GetUserSceneTextureOverride(FName& InOutName) const
+{
+	// Number of overrides possible is small (maximum 6, in most practical cases 1 or 2), and FName comparison cheap,
+	// so the assumption is that an array search will be cheaper than the overhead of going through a hash lookup.
+	// Plus an array takes half the space of THashedMaterialParameterMap, saving memory.
+	for (const FUserSceneTextureOverride& Override : UserSceneTextureOverrides)
+	{
+		if (Override.Key == InOutName)
+		{
+			InOutName = Override.Value;
+			return true;
+		}
+	}
+
+	if (Parent)
+	{
+		return Parent->GetUserSceneTextureOverride(InOutName);
+	}
+	else
+	{
+		return false;
+	}
+}
+
+FName UMaterialInstance::GetUserSceneTextureOutput(const UMaterial* Base) const
+{
+	FName Result = NAME_None;
+
+	// Replacing tonemapper can't override output.
+	if (Base->BlendableLocation != BL_ReplacingTonemapper)
+	{
+		// UserSceneTexture output overrides are stored under key "NAME_None".  We store them in the override lookup to save space
+		// in the base structure, by avoiding a separate field just for the output override.
+		if (!GetUserSceneTextureOverride(Result) && Base)
+		{
+			// If no override was found, get the result from the base material
+			Result = FName(Base->UserSceneTexture);
+		}
+	}
+	return Result;
+
+}
+
+EBlendableLocation UMaterialInstance::GetBlendableLocation(const UMaterial* Base) const
+{
+	check(Base);
+
+	// Replacing Tonemapper can't be overridden from
+	if (Base->BlendableLocation == BL_ReplacingTonemapper)
+	{
+		return BL_ReplacingTonemapper;
+	}
+
+	// Replacing Tonemapper can't be overridden to
+	if (bOverrideBlendableLocation && BlendableLocationOverride != BL_ReplacingTonemapper)
+	{
+		return BlendableLocationOverride;
+	}
+	else if (Parent)
+	{
+		return Parent->GetBlendableLocation(Base);
+	}
+	else
+	{
+		return Base->BlendableLocation;
+	}
+}
+
+int32 UMaterialInstance::GetBlendablePriority(const UMaterial* Base) const
+{
+	check(Base);
+	if (bOverrideBlendablePriority)
+	{
+		return BlendablePriorityOverride;
+	}
+	else if (Parent)
+	{
+		return Parent->GetBlendablePriority(Base);
+	}
+	else
+	{
+		return Base->BlendablePriority;
+	}
+}
+
 void UMaterialInstance::GetTextureExpressionValues(const FMaterialResource* MaterialResource, TArray<UTexture*>& OutTextures, TArray< TArray<int32> >* OutIndices) const
 {
 	check(MaterialResource);
@@ -4448,10 +4612,10 @@ void UMaterialInstance::GetResourceSizeEx(FResourceSizeEx& CumulativeResourceSiz
 	}
 }
 
-FPostProcessMaterialNode* IteratePostProcessMaterialNodes(const FFinalPostProcessSettings& Dest, const UMaterial* Material, FBlendableEntry*& Iterator)
+FPostProcessMaterialNode* IteratePostProcessMaterialNodes(const FFinalPostProcessSettings& Dest, const UMaterialInterface* Material, const UMaterial* Base, FBlendableEntry*& Iterator)
 {
-	EBlendableLocation Location = Material->BlendableLocation;
-	int32 Priority = Material->BlendablePriority;
+	EBlendableLocation Location = Material->GetBlendableLocation(Base);
+	int32 Priority = Material->GetBlendablePriority(Base);
 
 	for (;;)
 	{
diff --git a/Engine/Source/Runtime/Engine/Private/Materials/MaterialInstanceSupport.h b/Engine/Source/Runtime/Engine/Private/Materials/MaterialInstanceSupport.h
--- a/Engine/Source/Runtime/Engine/Private/Materials/MaterialInstanceSupport.h
+++ b/Engine/Source/Runtime/Engine/Private/Materials/MaterialInstanceSupport.h
@@ -222,6 +222,9 @@ public:
 	virtual UMaterialInterface* GetMaterialInterface() const override;
 	
 	virtual bool GetParameterValue(EMaterialParameterType Type, const FHashedMaterialParameterInfo& ParameterInfo, FMaterialParameterValue& OutValue, const FMaterialRenderContext& Context) const override;
+	virtual bool GetUserSceneTextureOverride(FName& InOutValue) const override;
+	virtual EBlendableLocation GetBlendableLocation(const FMaterial* Base) const override;
+	virtual int32 GetBlendablePriority(const FMaterial* Base) const override;
 
 	void GameThread_SetParent(UMaterialInterface* ParentMaterialInterface);
 
@@ -242,6 +245,7 @@ public:
 		TextureCollectionParameterArray.Empty();
 		RuntimeVirtualTextureParameterArray.Empty();
 		SparseVolumeTextureParameterArray.Empty();
+		UserSceneTextureOverrides.Empty();
 	}
 
 	/**
@@ -406,6 +410,10 @@ private:
 	THashedMaterialParameterMap<const USparseVolumeTexture*> SparseVolumeTextureParameterArray;
 	/** Remap layer indices for parent */
 	TArray<int32> ParentLayerIndexRemap;
+	/** User Scene Texture overrides for this material instance. */
+	TArray<FUserSceneTextureOverride> UserSceneTextureOverrides;
+	/** Post Process overrides for this material instance. */
+	FPostProcessBlendableOverrides PostProcessBlendableOverrides;
 };
 
 template <> FORCEINLINE THashedMaterialParameterMap<bool>& FMaterialInstanceResource::GetValueArray() { return StaticSwitchParameterArray; }
@@ -435,6 +443,8 @@ struct FMaterialInstanceParameterSet
 	TArray<THashedMaterialParameterMap<const UTextureCollection*>::TNamedParameter>		TextureCollectionParameters;
 	TArray<THashedMaterialParameterMap<const URuntimeVirtualTexture*>::TNamedParameter>	RuntimeVirtualTextureParameters;
 	TArray<THashedMaterialParameterMap<const USparseVolumeTexture*>::TNamedParameter>	SparseVolumeTextureParameters;
+	TArray<FUserSceneTextureOverride>													UserSceneTextureOverrides;
+	FPostProcessBlendableOverrides														PostProcessBlendableOverrides;
 };
 	
 /** Finds a parameter by name from the game thread. */
diff --git a/Engine/Source/Runtime/Engine/Private/Materials/MaterialInterface.cpp b/Engine/Source/Runtime/Engine/Private/Materials/MaterialInterface.cpp
--- a/Engine/Source/Runtime/Engine/Private/Materials/MaterialInterface.cpp
+++ b/Engine/Source/Runtime/Engine/Private/Materials/MaterialInterface.cpp
@@ -1217,6 +1217,24 @@ bool UMaterialInterface::GetRefractionSettings(float& OutBiasValue) const
 	return false;
 }
 
+bool UMaterialInterface::GetUserSceneTextureOverride(FName& InOutName) const
+{
+	return false;
+}
+
+// Certain implementations support overrides to the base material setting, default behavior is to just return the value from the base material
+EBlendableLocation UMaterialInterface::GetBlendableLocation(const UMaterial* Base) const
+{
+	check(Base);
+	return Base->BlendableLocation;
+}
+int32 UMaterialInterface::GetBlendablePriority(const UMaterial* Base) const
+{
+	check(Base);
+	return Base->BlendablePriority;
+}
+
+
 #if WITH_EDITOR
 bool UMaterialInterface::GetParameterDesc(const FHashedMaterialParameterInfo& ParameterInfo, FString& OutDesc) const
 {
diff --git a/Engine/Source/Runtime/Engine/Private/Materials/MaterialRenderProxy.cpp b/Engine/Source/Runtime/Engine/Private/Materials/MaterialRenderProxy.cpp
--- a/Engine/Source/Runtime/Engine/Private/Materials/MaterialRenderProxy.cpp
+++ b/Engine/Source/Runtime/Engine/Private/Materials/MaterialRenderProxy.cpp
@@ -250,6 +250,36 @@ bool FMaterialRenderProxy::GetTextureCollectionValue(const FHashedMaterialParame
 	return false;
 }
 
+FName FMaterialRenderProxy::GetUserSceneTextureOutput(const FMaterial* Base) const
+{
+	FName Result = NAME_None;
+
+	// Replacing tonemapper can't override output.
+	if (Base->GetBlendableLocation() != BL_ReplacingTonemapper)
+	{
+		// UserSceneTexture output overrides are stored under key "NAME_None".  We store them in the override lookup to save space
+		// in the base structure, by avoiding a separate field just for the output override.
+		if (!GetUserSceneTextureOverride(Result) && Base)
+		{
+			// If no override was found, get the result from the base material
+			Result = FName(Base->GetRenderingThreadShaderMap()->GetUserSceneTextureOutput());
+		}
+	}
+	return Result;
+}
+
+EBlendableLocation FMaterialRenderProxy::GetBlendableLocation(const FMaterial* Base) const
+{
+	check(Base);
+	return (EBlendableLocation)Base->GetBlendableLocation();
+}
+
+int32 FMaterialRenderProxy::GetBlendablePriority(const FMaterial* Base) const
+{
+	check(Base);
+	return Base->GetBlendablePriority();
+}
+
 static void OnVirtualTextureDestroyedCB(const FVirtualTextureProducerHandle& InHandle, void* Baton)
 {
 	FMaterialRenderProxy* MaterialProxy = static_cast<FMaterialRenderProxy*>(Baton);
diff --git a/Engine/Source/Runtime/Engine/Private/Materials/MaterialShader.cpp b/Engine/Source/Runtime/Engine/Private/Materials/MaterialShader.cpp
--- a/Engine/Source/Runtime/Engine/Private/Materials/MaterialShader.cpp
+++ b/Engine/Source/Runtime/Engine/Private/Materials/MaterialShader.cpp
@@ -2516,6 +2516,7 @@ void FMaterialShaderMap::Compile(
 		NewContent->UserSceneTextureOutput = FScriptName(MaterialObject->UserSceneTexture);
 		NewContent->UserTextureDivisorX = MaterialObject->UserTextureDivisor.X;
 		NewContent->UserTextureDivisorY = MaterialObject->UserTextureDivisor.Y;
+		NewContent->ResolutionRelativeToInput = FScriptName(MaterialObject->ResolutionRelativeToInput);
 	}
 
 #if ALLOW_SHADERMAP_DEBUG_DATA
diff --git a/Engine/Source/Runtime/Engine/Private/Materials/MaterialShared.cpp b/Engine/Source/Runtime/Engine/Private/Materials/MaterialShared.cpp
--- a/Engine/Source/Runtime/Engine/Private/Materials/MaterialShared.cpp
+++ b/Engine/Source/Runtime/Engine/Private/Materials/MaterialShared.cpp
@@ -2244,6 +2244,11 @@ bool FMaterialResource::GetBlendableOutputAlpha() const
 	return Material->IsPostProcessMaterialOutputingAlpha();
 }
 
+bool FMaterialResource::GetDisablePreExposureScale() const
+{
+	return GetMaterialDomain() == MD_PostProcess && Material->bDisablePreExposureScale;
+}
+
 bool FMaterialResource::IsStencilTestEnabled() const
 {
 	return GetMaterialDomain() == MD_PostProcess && Material->bEnableStencilTest;
@@ -5654,6 +5659,7 @@ FMaterialShaderParameters::FMaterialShaderParameters(const FMaterial* InMaterial
 	bIsUsedWithHeterogeneousVolumes = InMaterial->IsUsedWithHeterogeneousVolumes();
 	bIsMobileSeparateTranslucencyEnabled = InMaterial->IsMobileSeparateTranslucencyEnabled();
 	bAlwaysEvaluateWorldPositionOffset = InMaterial->ShouldAlwaysEvaluateWorldPositionOffset();
+	bDisablePreExposureScale = InMaterial->GetDisablePreExposureScale();
 }
 
 #if WITH_EDITOR
diff --git a/Engine/Source/Runtime/Engine/Public/MaterialCachedData.h b/Engine/Source/Runtime/Engine/Public/MaterialCachedData.h
--- a/Engine/Source/Runtime/Engine/Public/MaterialCachedData.h
+++ b/Engine/Source/Runtime/Engine/Public/MaterialCachedData.h
@@ -177,6 +177,9 @@ struct FMaterialCachedExpressionEditorOnlyData
 
 	UPROPERTY()
 	TSet<FString> ExpressionIncludeFilePaths;
+
+	UPROPERTY()
+	TSet<FName> UserSceneTextureInputs;
 };
 
 USTRUCT()
diff --git a/Engine/Source/Runtime/Engine/Public/MaterialShared.h b/Engine/Source/Runtime/Engine/Public/MaterialShared.h
--- a/Engine/Source/Runtime/Engine/Public/MaterialShared.h
+++ b/Engine/Source/Runtime/Engine/Public/MaterialShared.h
@@ -1049,6 +1049,7 @@ struct FMaterialShaderParameters
 			uint64 bIsUsedWithHeterogeneousVolumes : 1;
 			uint64 bIsMobileSeparateTranslucencyEnabled : 1;
 			uint64 bAlwaysEvaluateWorldPositionOffset : 1;
+			uint64 bDisablePreExposureScale : 1;
 		};
 	};
 
@@ -1378,6 +1379,7 @@ private:
 	LAYOUT_FIELD(FScriptName, UserSceneTextureOutput);
 	LAYOUT_FIELD(int32, UserTextureDivisorX);
 	LAYOUT_FIELD(int32, UserTextureDivisorY);
+	LAYOUT_FIELD(FScriptName, ResolutionRelativeToInput);
 
 	LAYOUT_FIELD_EDITORONLY(TMemoryImageArray<FMaterialProcessedSource>, ShaderProcessedSource);
 	LAYOUT_FIELD_EDITORONLY(FMemoryImageString, FriendlyName);
@@ -1626,7 +1628,8 @@ public:
 	bool UsesSceneTexture(uint32 TexId) const { return (GetContent()->MaterialCompilationOutput.UsedSceneTextures & (1ull << TexId)) != 0; }
 	TConstArrayView<FScriptName> GetUserSceneTextureInputs() const { return GetContent()->MaterialCompilationOutput.UserSceneTextureInputs; }
 	FScriptName GetUserSceneTextureOutput() const { return GetContent()->UserSceneTextureOutput; }
-	FIntPoint GetUserTextureDivisor() const { return FIntPoint(FMath::Max(GetContent()->UserTextureDivisorX, 1), FMath::Max(GetContent()->UserTextureDivisorY, 1)); }
+	FIntPoint GetUserTextureDivisor() const { return FIntPoint(GetContent()->UserTextureDivisorX, GetContent()->UserTextureDivisorY); }
+	FScriptName GetResolutionRelativeToInput() const { return GetContent()->ResolutionRelativeToInput; }
 
 	bool UsesPathTracingBufferTexture(uint32 TexId) const { return (GetContent()->MaterialCompilationOutput.UsedPathTracingBufferTextures & (1ull << TexId)) != 0;}
 
@@ -2211,6 +2214,7 @@ public:
 	virtual int32 GetBlendableLocation() const { return 0; }
 	virtual int32 GetBlendablePriority() const { return 0; }
 	virtual bool GetBlendableOutputAlpha() const { return false; }
+	virtual bool GetDisablePreExposureScale() const { return false; }
 	virtual bool IsStencilTestEnabled() const { return false; }
 	virtual uint32 GetStencilRefValue() const { return 0; }
 	virtual uint32 GetStencilCompare() const { return 0; }
@@ -2968,6 +2972,7 @@ public:
 	ENGINE_API virtual int32 GetBlendableLocation() const override;
 	ENGINE_API virtual int32 GetBlendablePriority() const override;
 	ENGINE_API virtual bool GetBlendableOutputAlpha() const override;
+	ENGINE_API virtual bool GetDisablePreExposureScale() const override;
 	ENGINE_API virtual bool IsStencilTestEnabled() const override;
 	ENGINE_API virtual uint32 GetStencilRefValue() const override;
 	ENGINE_API virtual uint32 GetStencilCompare() const override;
diff --git a/Engine/Source/Runtime/Engine/Public/Materials/Material.h b/Engine/Source/Runtime/Engine/Public/Materials/Material.h
--- a/Engine/Source/Runtime/Engine/Public/Materials/Material.h
+++ b/Engine/Source/Runtime/Engine/Public/Materials/Material.h
@@ -996,6 +996,19 @@ public:
 	UPROPERTY(EditAnywhere, BlueprintReadWrite, Category = PostProcessMaterial, meta = (DisplayName = "User Texture Divisor"))
 	FIntPoint UserTextureDivisor = FIntPoint(0, 0);
 
+	/** Set output resolution relative to the given User Scene Texture input.  Negative User Texture Divisors upsample (clamped at full resolution). */
+	UPROPERTY(EditAnywhere, BlueprintReadWrite, Category = PostProcessMaterial, meta = (DisplayName = "Resolution Relative To Input"))
+	FName ResolutionRelativeToInput;
+
+	/**
+	 * Disable pre-exposure scale in post process materials (multiply by View.OneOverPreExposure on inputs, View.PreExposure on output).  Useful for
+	 * materials that don't care about the absolute intensity of SceneColor (for example, a blur), simplifying custom HLSL and saving some performance.
+	 * Or useful for non-color UserSceneTextures (for example, mask, matte, modulation, offset, or ID textures), where pre-exposure scale is
+	 * undesirable.  Pre-exposure scale can be manually reapplied via custom HLSL if needed on specific inputs or the output.
+	 */
+	UPROPERTY(EditAnywhere, Category = PostProcessMaterial)
+	uint8 bDisablePreExposureScale : 1;
+
 	/**
 	* Indicates that the material and its instances can be used with neural network engine.
 	* This will result in the shaders required to support neural network engine being compiled which will increase shader compile time and memory usage.
@@ -1139,6 +1152,8 @@ public:
 	ENGINE_API virtual bool GetParameterValue(EMaterialParameterType Type, const FMemoryImageMaterialParameterInfo& ParameterInfo, FMaterialParameterMetadata& OutValue, EMaterialGetParameterValueFlags Flags = EMaterialGetParameterValueFlags::Default) const override;
 
 	ENGINE_API virtual bool GetRefractionSettings(float& OutBiasValue) const override;
+	ENGINE_API virtual EBlendableLocation GetBlendableLocation(const UMaterial* Base) const override;
+	ENGINE_API virtual int32 GetBlendablePriority(const UMaterial* Base) const override;
 	ENGINE_API virtual void GetDependencies(TSet<UMaterialInterface*>& Dependencies) override;
 	ENGINE_API virtual FMaterialRenderProxy* GetRenderProxy() const override;
 	ENGINE_API virtual UPhysicalMaterial* GetPhysicalMaterial() const override;
diff --git a/Engine/Source/Runtime/Engine/Public/Materials/MaterialInstance.h b/Engine/Source/Runtime/Engine/Public/Materials/MaterialInstance.h
--- a/Engine/Source/Runtime/Engine/Public/Materials/MaterialInstance.h
+++ b/Engine/Source/Runtime/Engine/Public/Materials/MaterialInstance.h
@@ -476,6 +476,27 @@ struct FFontParameterValue
 	}
 };
 
+USTRUCT(BlueprintType)
+struct FUserSceneTextureOverride
+{
+	GENERATED_USTRUCT_BODY()
+		
+	/** Key value of NONE represents override of UserSceneTexture output */
+	UPROPERTY(EditAnywhere, BlueprintReadWrite, Category=UserSceneTextureOverrideValue)
+	FName Key;
+
+	UPROPERTY(EditAnywhere, BlueprintReadWrite, Category=UserSceneTextureOverrideValue)
+	FName Value;
+};
+
+struct FPostProcessBlendableOverrides
+{
+	bool bOverrideBlendableLocation;
+	bool bOverrideBlendablePriority;
+	TEnumAsByte<EBlendableLocation> BlendableLocationOverride;
+	int32 BlendablePriorityOverride;
+};
+
 template<class T>
 bool CompareValueArraysByExpressionGUID(const TArray<T>& InA, const TArray<T>& InB)
 {
@@ -616,6 +637,14 @@ class UMaterialInstance : public UMaterialInterface
 	UPROPERTY(EditAnywhere, BlueprintReadOnly, Category = MaterialInstance)
 	uint8 bOverrideSubsurfaceProfile:1;
 
+	/** For post process materials, use BlendableLocationOverride. */
+	UPROPERTY(EditAnywhere, BlueprintReadWrite, Category = MaterialInstance)
+	uint8 bOverrideBlendableLocation:1;
+
+	/** For post process materials, use BlendablePriorityOverride. */
+	UPROPERTY(EditAnywhere, BlueprintReadWrite, Category = MaterialInstance)
+	uint8 bOverrideBlendablePriority : 1;
+
 	uint8 TwoSided : 1;
 	uint8 bIsThinSurface : 1;
 	uint8 DitheredLODTransition : 1;
@@ -631,6 +660,12 @@ public:
 
 	TEnumAsByte<EBlendMode> BlendMode;
 
+	UPROPERTY(EditAnywhere, BlueprintReadWrite, Category = MaterialInstance)
+	TEnumAsByte<EBlendableLocation> BlendableLocationOverride;
+
+	UPROPERTY(EditAnywhere, BlueprintReadWrite, Category = MaterialInstance)
+	int32 BlendablePriorityOverride;
+
 	FMaterialShadingModelField ShadingModels;
 
 	//Cached copies of the base property overrides or the value from the parent to avoid traversing the parent chain for each access.
@@ -689,6 +724,10 @@ public:
 	UPROPERTY(EditAnywhere, BlueprintReadOnly, Category=MaterialInstance, meta = (EditFixedOrder))
 	TArray<struct FFontParameterValue> FontParameterValues;
 
+	/** User scene texture overrides.  Applies to post process domain materials only. */
+	UPROPERTY(EditAnywhere, Category = MaterialInstance)
+	TArray<struct FUserSceneTextureOverride> UserSceneTextureOverrides;
+
 #if WITH_EDITORONLY_DATA
 	UPROPERTY()
 	bool bOverrideBaseProperties_DEPRECATED;
@@ -831,7 +870,11 @@ public:
 #endif
 	virtual ENGINE_API void RecacheUniformExpressions(bool bRecreateUniformBuffer) const override;
 	virtual ENGINE_API bool GetRefractionSettings(float& OutBiasValue) const override;
-	
+	virtual ENGINE_API bool GetUserSceneTextureOverride(FName& InOutName) const override;
+	ENGINE_API FName GetUserSceneTextureOutput(const UMaterial* Base) const;
+	virtual ENGINE_API EBlendableLocation GetBlendableLocation(const UMaterial* Base) const override;
+	virtual ENGINE_API int32 GetBlendablePriority(const UMaterial* Base) const override;
+
 	virtual ENGINE_API FGraphEventArray PrecachePSOs(const FPSOPrecacheVertexFactoryDataList& VertexFactoryDataList, const FPSOPrecacheParams& PreCacheParams, EPSOPrecachePriority Priority, TArray<FMaterialPSOPrecacheRequestID>& OutMaterialPSORequestIDs) override;
 
 #if WITH_EDITOR
diff --git a/Engine/Source/Runtime/Engine/Public/Materials/MaterialInterface.h b/Engine/Source/Runtime/Engine/Public/Materials/MaterialInterface.h
--- a/Engine/Source/Runtime/Engine/Public/Materials/MaterialInterface.h
+++ b/Engine/Source/Runtime/Engine/Public/Materials/MaterialInterface.h
@@ -913,6 +913,9 @@ public:
 #endif
 	ENGINE_API bool GetFontParameterValue(const FHashedMaterialParameterInfo& ParameterInfo,class UFont*& OutFontValue, int32& OutFontPage, bool bOveriddenOnly = false) const;
 	ENGINE_API virtual bool GetRefractionSettings(float& OutBiasValue) const;
+	ENGINE_API virtual bool GetUserSceneTextureOverride(FName& InOutValue) const;
+	ENGINE_API virtual EBlendableLocation GetBlendableLocation(const UMaterial* Base) const;
+	ENGINE_API virtual int32 GetBlendablePriority(const UMaterial* Base) const;
 
 	/**
 		Access to overridable properties of the base material.
diff --git a/Engine/Source/Runtime/Engine/Public/Materials/MaterialRenderProxy.h b/Engine/Source/Runtime/Engine/Public/Materials/MaterialRenderProxy.h
--- a/Engine/Source/Runtime/Engine/Public/Materials/MaterialRenderProxy.h
+++ b/Engine/Source/Runtime/Engine/Public/Materials/MaterialRenderProxy.h
@@ -9,6 +9,7 @@
 #include "RHIImmutableSamplerState.h"
 #include "Async/Mutex.h"
 #include "Tasks/Task.h"
+#include "Engine/BlendableInterface.h"
 
 enum class EMaterialParameterType : uint8;
 
@@ -173,6 +174,11 @@ public:
 	ENGINE_API bool GetTextureValue(const FHashedMaterialParameterInfo& ParameterInfo, const USparseVolumeTexture** OutValue, const FMaterialRenderContext& Context) const;
 	ENGINE_API bool GetTextureCollectionValue(const FHashedMaterialParameterInfo& ParameterInfo, const UTextureCollection** OutValue, const FMaterialRenderContext& Context) const;
 	virtual bool GetParameterValue(EMaterialParameterType Type, const FHashedMaterialParameterInfo& ParameterInfo, FMaterialParameterValue& OutValue, const FMaterialRenderContext& Context) const = 0;
+	virtual bool GetUserSceneTextureOverride(FName& InOutValue) const { return false; }
+	ENGINE_API FName GetUserSceneTextureOutput(const FMaterial* Base) const;
+	ENGINE_API virtual EBlendableLocation GetBlendableLocation(const FMaterial* Base) const;
+	ENGINE_API virtual int32 GetBlendablePriority(const FMaterial* Base) const;
+
 
 	bool IsDeleted() const
 	{
diff --git a/Engine/Source/Runtime/Renderer/Internal/SceneTextures.h b/Engine/Source/Runtime/Renderer/Internal/SceneTextures.h
--- a/Engine/Source/Runtime/Renderer/Internal/SceneTextures.h
+++ b/Engine/Source/Runtime/Renderer/Internal/SceneTextures.h
@@ -40,9 +40,9 @@ struct FUserSceneTextureEventData
 	EUserSceneTextureEvent Event;
 	FName Name;
 	uint16 AllocationOrder;
-	uint16 ViewIndex;					// Necessary to differentiate events from multiple views in split screen
-	const FMaterial* Material;
-	FIntPoint RectSize;					// Only filled in for EUserSceneTextureEvent::Output
+	uint16 ViewIndex;								// Necessary to differentiate events from multiple views in split screen
+	const UMaterialInterface* MaterialInterface;
+	FIntPoint RectSize;								// Only filled in for EUserSceneTextureEvent::Output
 };
 #endif
 
@@ -95,8 +95,9 @@ struct FMinimalSceneTextures
 
 	RENDERER_API FSceneTextureShaderParameters GetSceneTextureShaderParameters(ERHIFeatureLevel::Type FeatureLevel) const;
 
-	FRDGTextureRef FindOrAddUserSceneTexture(FRDGBuilder& GraphBuilder, int32 ViewIndex, FName Name, FIntPoint ResolutionDivisor, bool& bOutFirstRender, const FMaterial* Material, const FIntRect& OutputRect) const;
-	FScreenPassTextureSlice GetUserSceneTexture(FRDGBuilder& GraphBuilder, const FViewInfo& View, int32 ViewIndex, FName Name, const FMaterial* Material) const;
+	FRDGTextureRef FindOrAddUserSceneTexture(FRDGBuilder& GraphBuilder, int32 ViewIndex, FName Name, FIntPoint ResolutionDivisor, bool& bOutFirstRender, const UMaterialInterface* MaterialInterface, const FIntRect& OutputRect) const;
+	FScreenPassTextureSlice GetUserSceneTexture(FRDGBuilder& GraphBuilder, const FViewInfo& View, int32 ViewIndex, FName Name, const UMaterialInterface* MaterialInterface) const;
+	FIntPoint GetUserSceneTextureDivisor(FName Name) const;
 
 #if !(UE_BUILD_SHIPPING)
 	const FTransientUserSceneTexture* FindUserSceneTextureByEvent(const FUserSceneTextureEventData& Event) const;
diff --git a/Engine/Source/Runtime/Renderer/Private/PostProcess/PostProcessMaterial.cpp b/Engine/Source/Runtime/Renderer/Private/PostProcess/PostProcessMaterial.cpp
--- a/Engine/Source/Runtime/Renderer/Private/PostProcess/PostProcessMaterial.cpp
+++ b/Engine/Source/Runtime/Renderer/Private/PostProcess/PostProcessMaterial.cpp
@@ -8,6 +8,7 @@
 #include "DataDrivenShaderPlatformInfo.h"
 #include "RendererModule.h"
 #include "Materials/Material.h"
+#include "Materials/MaterialInstance.h"
 #include "Materials/MaterialRenderProxy.h"
 #include "MaterialDomain.h"
 #include "MaterialShaderType.h"
@@ -55,6 +56,8 @@ TAutoConsoleVariable<int32> CVarPostProcessingDisableMaterials(
 	TEXT(" Allows to disable post process materials. \n"),
 	ECVF_Scalability | ECVF_RenderThreadSafe);
 
+static FName NAME_SceneColor("SceneColor");
+
 static bool IsPostProcessStencilTestAllowed()
 {
 	return CVarPostProcessAllowStencilTest.GetValueOnRenderThread() != 0;
@@ -69,7 +72,7 @@ enum class EMaterialCustomDepthPolicy : uint32
 	Enabled
 };
 
-static EMaterialCustomDepthPolicy GetMaterialCustomDepthPolicy(const FMaterial* Material)
+static EMaterialCustomDepthPolicy GetMaterialCustomDepthPolicy(const FMaterialRenderProxy* MaterialRenderProxy, const FMaterial* Material)
 {
 	check(Material);
 
@@ -81,7 +84,7 @@ static EMaterialCustomDepthPolicy GetMaterialCustomDepthPolicy(const FMaterial*
 		{
 			UE_LOG(LogRenderer, Warning, TEXT("PostProcessMaterial uses stencil test, but stencil not allocated. Set r.CustomDepth to 3 to allocate custom stencil."));
 		}
-		else if (Material->GetBlendableLocation() == BL_SceneColorAfterTonemapping)
+		else if (MaterialRenderProxy->GetBlendableLocation(Material) == BL_SceneColorAfterTonemapping)
 		{
 			// We can't support custom stencil after tonemapping due to target size differences
 			UE_LOG(LogRenderer, Warning, TEXT("PostProcessMaterial uses stencil test, but is set to blend After Tonemapping. This is not supported."));
@@ -258,6 +261,7 @@ public:
 		// SSR input should not be affected by exposure so it should be specified separately from POST_PROCESS_MATERIAL_BEFORE_TONEMAP 
 		// in order to be able to make DRS independent CameraVector and WorldPosition nodes.
 		OutEnvironment.SetDefine(TEXT("POST_PROCESS_MATERIAL_SSRINPUT"), (Location == BL_SSRInput) ? 1 : 0);
+		OutEnvironment.SetDefine(TEXT("POST_PROCESS_DISABLE_PRE_EXPOSURE_SCALE"), Parameters.MaterialParameters.bDisablePreExposureScale ? 1 : 0);
 
 		if (IsMobilePlatform(Parameters.Platform))
 		{
@@ -566,7 +570,7 @@ void AddNeuralPostProcessPass(
 	FRDGTextureRef DepthStencilTexture = nullptr;
 
 	// Allocate custom depth stencil texture(s) and depth stencil state.
-	const EMaterialCustomDepthPolicy CustomStencilPolicy = GetMaterialCustomDepthPolicy(Material);
+	const EMaterialCustomDepthPolicy CustomStencilPolicy = GetMaterialCustomDepthPolicy(MaterialRenderProxy, Material);
 
 	if (CustomStencilPolicy == EMaterialCustomDepthPolicy::Enabled &&
 		!Inputs.bManualStencilTest &&
@@ -704,7 +708,7 @@ FScreenPassTexture AddPostProcessMaterialPass(
 	TShaderRef<FPostProcessMaterialPS> PixelShader;
 	GetMaterialInfo(MaterialInterface, FeatureLevel, Inputs, Material, MaterialRenderProxy, MaterialShaderMap, VertexShader, PixelShader);
 
-	EBlendableLocation BlendableLocation = EBlendableLocation(Material->GetBlendableLocation());
+	EBlendableLocation BlendableLocation = MaterialRenderProxy->GetBlendableLocation(Material);
 	const FScreenPassTextureSlice SceneColorOutput = Inputs.GetSceneColorOutput(BlendableLocation);
 
 	check(VertexShader.IsValid());
@@ -716,7 +720,7 @@ FScreenPassTexture AddPostProcessMaterialPass(
 	FRDGTextureRef DepthStencilTexture = nullptr;
 
 	// Allocate custom depth stencil texture(s) and depth stencil state.
-	const EMaterialCustomDepthPolicy CustomStencilPolicy = GetMaterialCustomDepthPolicy(Material);
+	const EMaterialCustomDepthPolicy CustomStencilPolicy = GetMaterialCustomDepthPolicy(MaterialRenderProxy, Material);
 
 	if (CustomStencilPolicy == EMaterialCustomDepthPolicy::Enabled &&
 		!Inputs.bManualStencilTest &&
@@ -981,7 +985,7 @@ FPostProcessMaterialChain GetPostProcessMaterialChain(const FViewInfo& View, EBl
 
 		if (Material && (Material->BlendableLocation == Location || Location == EBlendableLocation::BL_MAX))
 		{
-			Nodes.Add(FPostProcessMaterialNode(Material, Material->BlendableLocation, Material->BlendablePriority, Material->bIsBlendable));
+			Nodes.Add(FPostProcessMaterialNode(Material, VisMaterial->GetBlendableLocation(Material), VisMaterial->GetBlendablePriority(Material), Material->bIsBlendable));
 		}
 	}
 
@@ -1077,10 +1081,18 @@ FScreenPassTexture AddPostProcessMaterialChain(
 	const UMaterialInterface* LastOutputWrite = nullptr;
 	for (int32 MaterialIndex = Materials.Num() - 1; MaterialIndex >= 0; MaterialIndex--)
 	{
-		const FMaterial* Material = Materials[MaterialIndex]->GetRenderProxy()->GetMaterialNoFallback(View.FeatureLevel);
-		if (!Material || Material->GetRenderingThreadShaderMap()->GetUserSceneTextureOutput().IsNone())
+		const FMaterialRenderProxy* MaterialRenderProxy = Materials[MaterialIndex]->GetRenderProxy();
+		const FMaterial* Material = MaterialRenderProxy->GetMaterialNoFallback(View.FeatureLevel);
+
+		FName UserSceneTextureOutput = NAME_None;
+		if (Material)
+		{
+			UserSceneTextureOutput = MaterialRenderProxy->GetUserSceneTextureOutput(Material);
+		}
+
+		if (UserSceneTextureOutput.IsNone())
 		{
-			// Doesn't write to a UserSceneTexture, so it writes to the output
+			// Doesn't write to a UserSceneTexture, so it writes to the default SceneColor output
 			LastOutputWrite = Materials[MaterialIndex];
 			break;
 		}
@@ -1101,13 +1113,19 @@ FScreenPassTexture AddPostProcessMaterialChain(
 		if (Material)
 		{
 			const FMaterialShaderMap* MaterialShaderMap = Material->GetRenderingThreadShaderMap();
+			const FSceneTextures& SceneTextures = View.GetSceneTextures();
+
+			bool bFoundResolutionRelativeToInput = false;
+			FName ResolutionRelativeToInput = NAME_None;
 
 			TConstArrayView<FScriptName> UserSceneTextureInputs = MaterialShaderMap->GetUserSceneTextureInputs();
 			if (!UserSceneTextureInputs.IsEmpty())
 			{
 				UserSceneTextureInputNum = UserSceneTextureInputs.Num();
 
-				const FSceneTextures& SceneTextures = View.GetSceneTextures();
+				// We need to apply material instance input overrides to ResolutionRelativeToInput as well, so get the name here to
+				// handle that in the input loop.
+				ResolutionRelativeToInput = FName(MaterialShaderMap->GetResolutionRelativeToInput());
 
 				int32 PostProcessIndex = 0;
 				for (int32 UserIndex = 0; UserIndex < UserSceneTextureInputs.Num();)
@@ -1117,8 +1135,33 @@ FScreenPassTexture AddPostProcessMaterialChain(
 					// Skip over this slot if it's used by a SceneTexture node
 					if (!MaterialShaderMap->UsesSceneTexture(PPI_PostProcessInput0 + PostProcessIndex))
 					{
-						// Not used as a SceneTexture, so it's used by the next UserSceneTexture
-						Inputs.SetUserSceneTextureInput((EPostProcessMaterialInput)PostProcessIndex, SceneTextures.GetUserSceneTexture(GraphBuilder, View, ViewIndex, FName(UserSceneTextureInputs[UserIndex]), Material));
+						FName UserSceneTextureInput(UserSceneTextureInputs[UserIndex]);
+						bool bIsResolutionSource = ResolutionRelativeToInput == UserSceneTextureInputs[UserIndex];
+
+						MaterialInterface->GetRenderProxy()->GetUserSceneTextureOverride(UserSceneTextureInput);
+
+						if (bIsResolutionSource)
+						{
+							// Copy the overridden input to ResolutionRelativeToInput, and track that we found it
+							ResolutionRelativeToInput = UserSceneTextureInput;
+							bFoundResolutionRelativeToInput = true;
+						}
+
+						// Not used as a SceneTexture, so it's used by the next UserSceneTexture.  The special name "SceneColor" indicates use of
+						// "SceneColor" as input.
+						if (UserSceneTextureInput == NAME_SceneColor)
+						{
+							Inputs.SetUserSceneTextureInput((EPostProcessMaterialInput)PostProcessIndex, CurrentInput);
+
+							// Need to disable optimization that attempts to reuse SceneColor as the output, when SceneColor isn't used as an input.  Normally
+							// the use of SceneColor as an input is detected by the flags on the original FMaterialShaderMap (accessed via the UsesSceneTexture
+							// function), but those flags won't be set if a UserSceneTexture input is overridden to point at SceneColor.
+							Inputs.bAllowSceneColorInputAsOutput = false;
+						}
+						else
+						{
+							Inputs.SetUserSceneTextureInput((EPostProcessMaterialInput)PostProcessIndex, SceneTextures.GetUserSceneTexture(GraphBuilder, View, ViewIndex, UserSceneTextureInput, MaterialInterface));
+						}
 						UserIndex++;
 					}
 					PostProcessIndex++;
@@ -1127,11 +1170,72 @@ FScreenPassTexture AddPostProcessMaterialChain(
 
 #if WITH_EDITOR
 			// If this blendable is being previewed, don't write to the UserSceneTexture -- instead it will write to SceneColor
-			if (Material->GetMaterialInterface() != View.FinalPostProcessSettings.PreviewBlendable)
+			bool bIsPreviewBlendable = false;
+			if (View.FinalPostProcessSettings.PreviewBlendable)
+			{
+				if (Material->GetMaterialInterface() == View.FinalPostProcessSettings.PreviewBlendable)
+				{
+					// Material matches
+					bIsPreviewBlendable = true;
+				}
+				else
+				{
+					FMaterialInheritanceChain MaterialInheritance;
+					MaterialInterface->GetMaterialInheritanceChain(MaterialInheritance);
+					if (MaterialInheritance.MaterialInstances.Contains(View.FinalPostProcessSettings.PreviewBlendable))
+					{
+						// Material instance matches
+						bIsPreviewBlendable = true;
+					}
+				}
+			}
+
+			if (!bIsPreviewBlendable)
 #endif
 			{
-				UserSceneTextureOutput = FName(MaterialShaderMap->GetUserSceneTextureOutput());
-				UserTextureDivisor = MaterialShaderMap->GetUserTextureDivisor();
+				UserSceneTextureOutput = MaterialInterface->GetRenderProxy()->GetUserSceneTextureOutput(Material);
+
+				// If output is set to the name "SceneColor", that means actually write to "SceneColor" as opposed to a transient UserSceneTexture.
+				// The purpose of this is to give a general purpose Material asset operating on UserSceneTexture inputs and outputs the option to
+				// read or write SceneColor as well, say if they are the first or last building block in a chain of materials.
+				if (UserSceneTextureOutput == NAME_SceneColor)
+				{
+					// Clear to none, so it writes to SceneColor downstream
+					UserSceneTextureOutput = NAME_None;
+				}
+				else
+				{
+					UserTextureDivisor = MaterialShaderMap->GetUserTextureDivisor();
+
+					if (bFoundResolutionRelativeToInput)
+					{
+						// UserTextureDivisor is a relative divisor to the input, with positive values representing downscale, and negative upscale
+						FIntPoint InputDivisor = SceneTextures.GetUserSceneTextureDivisor(ResolutionRelativeToInput);
+
+						if (UserTextureDivisor.X >= 0)
+						{
+							UserTextureDivisor.X = InputDivisor.X * FMath::Max(UserTextureDivisor.X, 1);
+						}
+						else
+						{
+							UserTextureDivisor.X = FMath::Max(InputDivisor.X / FMath::Abs(UserTextureDivisor.X), 1);
+						}
+
+						if (UserTextureDivisor.Y >= 0)
+						{
+							UserTextureDivisor.Y = InputDivisor.Y * FMath::Max(UserTextureDivisor.Y, 1);
+						}
+						else
+						{
+							UserTextureDivisor.Y = FMath::Max(InputDivisor.Y / FMath::Abs(UserTextureDivisor.Y), 1);
+						}
+					}
+					else
+					{
+						UserTextureDivisor.X = FMath::Max(UserTextureDivisor.X, 1);
+						UserTextureDivisor.Y = FMath::Max(UserTextureDivisor.Y, 1);
+					}
+				}
 			}
 		}
 		
@@ -1143,7 +1247,7 @@ FScreenPassTexture AddPostProcessMaterialChain(
 		{
 			// Writing to UserSceneTexture, don't set Outputs or CurrentInput, as this is writing to a disjoint texture that's not part of the chain
 			FIntRect OutputRect = GetDownscaledViewRect(View.UnconstrainedViewRect, View.GetFamilyViewRect().Max, UserTextureDivisor);
-			FRDGTextureRef UserOutput = ((FViewFamilyInfo*)View.Family)->GetSceneTextures().FindOrAddUserSceneTexture(GraphBuilder, ViewIndex, UserSceneTextureOutput, UserTextureDivisor, Inputs.bUserSceneTextureFirstRender, Material, OutputRect);
+			FRDGTextureRef UserOutput = ((FViewFamilyInfo*)View.Family)->GetSceneTextures().FindOrAddUserSceneTexture(GraphBuilder, ViewIndex, UserSceneTextureOutput, UserTextureDivisor, Inputs.bUserSceneTextureFirstRender, MaterialInterface, OutputRect);
 			Inputs.OverrideOutput = FScreenPassRenderTarget(UserOutput, OutputRect, ERenderTargetLoadAction::ELoad);
 			Inputs.bUserSceneTextureOutput = true;
 
@@ -1171,7 +1275,7 @@ FScreenPassTexture AddPostProcessMaterialChain(
 #if !UE_BUILD_SHIPPING
 		if (UserSceneTextureInputNum || !UserSceneTextureOutput.IsNone())
 		{
-			View.GetSceneTextures().UserSceneTextureEvents.Add({ EUserSceneTextureEvent::Pass, NAME_None, 0, (uint16)ViewIndex, Material });
+			View.GetSceneTextures().UserSceneTextureEvents.Add({ EUserSceneTextureEvent::Pass, NAME_None, 0, (uint16)ViewIndex, MaterialInterface });
 		}
 #endif
 	}
diff --git a/Engine/Source/Runtime/Renderer/Private/PostProcess/PostProcessing.cpp b/Engine/Source/Runtime/Renderer/Private/PostProcess/PostProcessing.cpp
--- a/Engine/Source/Runtime/Renderer/Private/PostProcess/PostProcessing.cpp
+++ b/Engine/Source/Runtime/Renderer/Private/PostProcess/PostProcessing.cpp
@@ -74,6 +74,7 @@
 #include "SparseVolumeTexture/SparseVolumeTextureStreamingVisualize.h"
 #include "CanvasItem.h"
 #include "MobileSSR.h"
+#include "Materials/MaterialRenderProxy.h"
 
 bool IsMobileEyeAdaptationEnabled(const FViewInfo& View);
 
@@ -3282,7 +3283,10 @@ static void AddUserSceneTextureDebugPass(FRDGBuilder& GraphBuilder, const FViewI
 			{
 				if (EventData.ViewIndex == ViewIndex && EventData.Event == EUserSceneTextureEvent::Pass)
 				{
-					Text = FString::Printf(TEXT("%s [%d]"), GBlendableLocationShortNames[FMath::Min((uint32)EventData.Material->GetBlendableLocation(), (uint32)BL_MAX)], EventData.Material->GetBlendablePriority());
+					const FMaterialRenderProxy* RenderProxy = EventData.MaterialInterface->GetRenderProxy();
+					const FMaterial* Material = RenderProxy->GetMaterialNoFallback(View.FeatureLevel);
+
+					Text = FString::Printf(TEXT("%s [%d]"), GBlendableLocationShortNames[FMath::Min((uint32)RenderProxy->GetBlendableLocation(Material), (uint32)BL_MAX)], RenderProxy->GetBlendablePriority(Material));
 
 					float BlendableInfoWidth = CanvasDrawShadowedStringReturnWidth(Canvas, PrintX, PrintY, Text, Font, TextColor);
 					MaxBlendableInfoWidth = FMath::Max(MaxBlendableInfoWidth, BlendableInfoWidth);
@@ -3300,7 +3304,15 @@ static void AddUserSceneTextureDebugPass(FRDGBuilder& GraphBuilder, const FViewI
 			{
 				if (EventData.ViewIndex == ViewIndex && EventData.Event == EUserSceneTextureEvent::Pass)
 				{
-					Text = FString::Printf(TEXT("%s:"), *EventData.Material->GetFriendlyName());
+					const UMaterialInterface* MaterialInterface = EventData.MaterialInterface;
+
+					// Skip over runtime generated dynamic instance when producing name
+					while (MaterialInterface->IsA<UMaterialInstanceDynamic>())
+					{
+						MaterialInterface = ((const UMaterialInstanceDynamic*)MaterialInterface)->Parent;
+					}
+
+					Text = FString::Printf(TEXT("%s:"), *MaterialInterface->GetName());
 
 					float NameWidth = CanvasDrawShadowedStringReturnWidth(Canvas, PrintX, PrintY, Text, Font, TextColor);
 					MaxNameWidth = FMath::Max(MaxNameWidth, NameWidth);
@@ -3346,7 +3358,7 @@ static void AddUserSceneTextureDebugPass(FRDGBuilder& GraphBuilder, const FViewI
 							PrintX += CanvasDrawShadowedStringReturnWidth(Canvas, PrintX, PrintY, *Text, Font, UserTexture->bUsed ? GrayTextColor : YellowTextColor);
 							bAnyUnused = bAnyUnused || !UserTexture->bUsed;
 
-							if (EventData.Material->GetBlendMode() != BLEND_Opaque)
+							if (EventData.MaterialInterface->GetBlendMode() != BLEND_Opaque)
 							{
 								PrintX += CanvasDrawShadowedStringReturnWidth(Canvas, PrintX, PrintY, TEXT("  Blend"), Font, GrayTextColor);
 							}
diff --git a/Engine/Source/Runtime/Renderer/Private/SceneTextures.cpp b/Engine/Source/Runtime/Renderer/Private/SceneTextures.cpp
--- a/Engine/Source/Runtime/Renderer/Private/SceneTextures.cpp
+++ b/Engine/Source/Runtime/Renderer/Private/SceneTextures.cpp
@@ -489,7 +489,7 @@ FSceneTextureShaderParameters FMinimalSceneTextures::GetSceneTextureShaderParame
 	return OutSceneTextureShaderParameters;
 }
 
-FRDGTextureRef FMinimalSceneTextures::FindOrAddUserSceneTexture(FRDGBuilder& GraphBuilder, int32 ViewIndex, FName Name, FIntPoint ResolutionDivisor, bool& bOutFirstRender, const FMaterial* Material, const FIntRect& OutputRect) const
+FRDGTextureRef FMinimalSceneTextures::FindOrAddUserSceneTexture(FRDGBuilder& GraphBuilder, int32 ViewIndex, FName Name, FIntPoint ResolutionDivisor, bool& bOutFirstRender, const UMaterialInterface* MaterialInterface, const FIntRect& OutputRect) const
 {
 	check(ResolutionDivisor.X >= 1 && ResolutionDivisor.Y >= 1);
 
@@ -559,7 +559,7 @@ FRDGTextureRef FMinimalSceneTextures::FindOrAddUserSceneTexture(FRDGBuilder& Gra
 	}
 
 #if !(UE_BUILD_SHIPPING)
-	UserSceneTextureEvents.Add({ EUserSceneTextureEvent::Output, Name, (*TransientTextures)[0].AllocationOrder, (uint16)ViewIndex, Material, OutputRect.Size() });
+	UserSceneTextureEvents.Add({ EUserSceneTextureEvent::Output, Name, (*TransientTextures)[0].AllocationOrder, (uint16)ViewIndex, MaterialInterface, OutputRect.Size() });
 #endif
 
 	// If out of mask range, treat it as the first render.  This may result in certain transparent post process materials writing to UserSceneTextures
@@ -578,7 +578,7 @@ FRDGTextureRef FMinimalSceneTextures::FindOrAddUserSceneTexture(FRDGBuilder& Gra
 	return (*TransientTextures)[0].Texture;
 }
 
-FScreenPassTextureSlice FMinimalSceneTextures::GetUserSceneTexture(FRDGBuilder& GraphBuilder, const FViewInfo& View, int32 ViewIndex, FName Name, const FMaterial* Material) const
+FScreenPassTextureSlice FMinimalSceneTextures::GetUserSceneTexture(FRDGBuilder& GraphBuilder, const FViewInfo& View, int32 ViewIndex, FName Name, const UMaterialInterface* MaterialInterface) const
 {
 	TArray<FTransientUserSceneTexture>* TransientTextures = UserSceneTextures.Find(Name);
 	if (TransientTextures)
@@ -589,7 +589,7 @@ FScreenPassTextureSlice FMinimalSceneTextures::GetUserSceneTexture(FRDGBuilder&
 
 #if !(UE_BUILD_SHIPPING)
 		(*TransientTextures)[0].bUsed = true;
-		UserSceneTextureEvents.Add({ EUserSceneTextureEvent::FoundInput, Name, (*TransientTextures)[0].AllocationOrder, (uint16)ViewIndex, Material });
+		UserSceneTextureEvents.Add({ EUserSceneTextureEvent::FoundInput, Name, (*TransientTextures)[0].AllocationOrder, (uint16)ViewIndex, MaterialInterface });
 #endif
 
 		return TransientTextureSlice;
@@ -597,13 +597,26 @@ FScreenPassTextureSlice FMinimalSceneTextures::GetUserSceneTexture(FRDGBuilder&
 	else
 	{
 #if !(UE_BUILD_SHIPPING)
-		UserSceneTextureEvents.Add({ EUserSceneTextureEvent::MissingInput, Name, 0, (uint16)ViewIndex, Material });
+		UserSceneTextureEvents.Add({ EUserSceneTextureEvent::MissingInput, Name, 0, (uint16)ViewIndex, MaterialInterface });
 #endif
 
 		return FScreenPassTextureSlice();
 	}
 }
 
+FIntPoint FMinimalSceneTextures::GetUserSceneTextureDivisor(FName Name) const
+{
+	TArray<FTransientUserSceneTexture>* TransientTextures = UserSceneTextures.Find(Name);
+	if (TransientTextures)
+	{
+		return (*TransientTextures)[0].ResolutionDivisor;
+	}
+	else
+	{
+		return FIntPoint(1, 1);
+	}
+}
+
 #if !(UE_BUILD_SHIPPING)
 const FTransientUserSceneTexture* FMinimalSceneTextures::FindUserSceneTextureByEvent(const FUserSceneTextureEventData& Event) const
 {
