Material editor crash on UI materials -- need null check to skip preview material refresh specific to post process materials.

#rnx

diff --git a/Engine/Source/Editor/MaterialEditor/Private/SMaterialEditorViewport.cpp b/Engine/Source/Editor/MaterialEditor/Private/SMaterialEditorViewport.cpp
--- a/Engine/Source/Editor/MaterialEditor/Private/SMaterialEditorViewport.cpp
+++ b/Engine/Source/Editor/MaterialEditor/Private/SMaterialEditorViewport.cpp
@@ -992,7 +992,7 @@ void SMaterialEditor3DPreviewViewport::OnPropertyChanged(UObject* ObjectBeingMod
 	// includes the "UMaterial::UserSceneTexture" field, plus any FName field in the PostProcessOverrides member struct.
 	//
 	// We also need to refresh PreviewMaterial itself if its domain changes, regardless of whether it's a post process material.
-	if (ObjectBeingModified != nullptr && PropertyThatChanged != nullptr)
+	if (ObjectBeingModified != nullptr && PropertyThatChanged != nullptr && PreviewMaterial)
 	{
 		if ((ObjectBeingModified == PreviewMaterial && PropertyThatChanged->GetName() == MaterialDomain) ||
 			(PreviewMaterial->GetMaterial()->IsPostProcessMaterial() &&
