INtroducing the concept of tree in the MaterialLayersFunctions

Introducing a new struct, the MaterialLayersFunctionsTree which hold a tree structure, aka a hierarchy of nodes
The tree is a new member/ property of the MaterialLayersFunctions

Each node of the tree has a Payload associated which reference a LayerFunction and a BLendFunction in the existing Layers and Blends array of the struct.
So we are adding the tree "outline" on top of the existing containers of the struct.

When the Macro for the new material layer prototype is enabled, then the tree is activated and a new set of methods to create/remove/move the layerNodes should be use which allocate the node of the tree along with the referenced Layer and Blend functions in their legacy containers.
This is how the new Authoring widget the SMaterialSubstrateTree manage the editing of the tree.

When the macro is disabled, nothing the tree is not used from the legacy widget.

#virtualized

diff --git a/Engine/Source/Editor/MaterialEditor/Private/SMaterialSubstrateTree.cpp b/Engine/Source/Editor/MaterialEditor/Private/SMaterialSubstrateTree.cpp
--- a/Engine/Source/Editor/MaterialEditor/Private/SMaterialSubstrateTree.cpp
+++ b/Engine/Source/Editor/MaterialEditor/Private/SMaterialSubstrateTree.cpp
@@ -54,6 +54,15 @@
 #define LOCTEXT_NAMESPACE "MaterialSubstrateTree"
 
 
+FString CreateNodeKeyForBackgroundParameter(const FString& InNodeKeyBase)
+{
+	return FString(InNodeKeyBase + TEXT("_BackgroundParameter"));
+}
+bool IsNodeKeyForBackgroundParameter(const FString InNodeKey)
+{
+	return (InNodeKey.Find(TEXT("_BackgroundParameter")) >= 0);
+}
+
 /// ===========================================================================================================
 /// SMaterialSubstrateItem Methods
 /// ===========================================================================================================
@@ -138,19 +147,45 @@ ECheckBoxState SMaterialSubstrateTreeItem::GetFilterChecked(SMaterialSubstrateTr
 	return GetFilterState(InTree, InStackData) ? ECheckBoxState::Checked : ECheckBoxState::Unchecked;
 }
 
-FText SMaterialSubstrateTreeItem::GetLayerName(SMaterialSubstrateTree* InTree, int32 Counter) const
+FReply SMaterialSubstrateTreeItem::ToggleLayerVisibility()
+{
+	int32 LayerFuncIndex = Tree->FunctionInstance->GetLayerFuncIndex(StackParameterData->ParameterInfo.Index);
+	return Tree->ToggleLayerVisibility(LayerFuncIndex);
+}
+
+bool SMaterialSubstrateTreeItem::IsLayerVisible() const
+{
+	int32 LayerFuncIndex = Tree->FunctionInstance->GetLayerFuncIndex(StackParameterData->ParameterInfo.Index);
+	return Tree->IsLayerVisible(LayerFuncIndex);
+}
+
+FReply SMaterialSubstrateTreeItem::UnlinkLayer()
+{
+	int32 LayerFuncIndex = Tree->FunctionInstance->GetLayerFuncIndex(StackParameterData->ParameterInfo.Index);
+	return Tree->UnlinkLayer(LayerFuncIndex);
+}
+
+EVisibility SMaterialSubstrateTreeItem::GetUnlinkLayerVisibility() const
+{
+	int32 LayerFuncIndex = Tree->FunctionInstance->GetLayerFuncIndex(StackParameterData->ParameterInfo.Index);
+	return Tree->GetUnlinkLayerVisibility(LayerFuncIndex);
+}
+
+FText SMaterialSubstrateTreeItem::GetLayerName() const
 {
-	return InTree->FunctionInstance->GetLayerName(Counter);
+	int32 LayerFuncIndex = Tree->FunctionInstance->GetLayerFuncIndex(StackParameterData->ParameterInfo.Index);
+	return Tree->FunctionInstance->GetLayerName(LayerFuncIndex);
 }
 
-void SMaterialSubstrateTreeItem::OnNameChanged(const FText& InText, ETextCommit::Type CommitInfo, SMaterialSubstrateTree* InTree, int32 Counter)
+void SMaterialSubstrateTreeItem::OnNameChanged(const FText& InText, ETextCommit::Type CommitInfo)
 {
 	const FScopedTransaction Transaction(LOCTEXT("RenamedSection", "Renamed layer and blend section"));
-	InTree->FunctionInstanceHandle->NotifyPreChange();
-	InTree->FunctionInstance->EditorOnly.LayerNames[Counter] = InText;
-	InTree->FunctionInstance->UnlinkLayerFromParent(Counter);
-	InTree->MaterialEditorInstance->CopyToSourceInstance(true);
-	InTree->FunctionInstanceHandle->NotifyPostChange(EPropertyChangeType::ValueSet);
+	int32 LayerFuncIndex = Tree->FunctionInstance->GetLayerFuncIndex(StackParameterData->ParameterInfo.Index);
+	Tree->FunctionInstanceHandle->NotifyPreChange();
+	Tree->FunctionInstance->EditorOnly.LayerNames[LayerFuncIndex] = InText;
+	Tree->FunctionInstance->UnlinkLayerFromParent(LayerFuncIndex);
+	Tree->MaterialEditorInstance->CopyToSourceInstance(true);
+	Tree->FunctionInstanceHandle->NotifyPostChange(EPropertyChangeType::ValueSet);
 }
 
 TOptional<EItemDropZone> SMaterialSubstrateTreeItem::CanAcceptDrop(const FDragDropEvent& DragDropEvent, EItemDropZone DropZone, FSortedParamDataPtr Item)
@@ -186,6 +221,30 @@ FReply SMaterialSubstrateTreeItem::OnLayerDrop(const FDragDropEvent& DragDropEve
 
 		if (AssetDropOp.IsValid())
 		{
+			bool bIsBackgroundItem = IsNodeKeyForBackgroundParameter(StackParameterData->NodeKey);
+
+			// Drop above or below could CREATE a new layer node:
+			int32 TargetNodeId = StackParameterData->ParameterInfo.Index;
+			int32 ParentNodeId = (bIsBackgroundItem ? TargetNodeId : Tree->FunctionInstance->GetNodeParent(TargetNodeId));
+			auto ChildrenNodeId = Tree->FunctionInstance->GetNodeChildren(ParentNodeId);
+			int32 SiblingIdx = -1;
+			bool FoundTarget = ChildrenNodeId.Find(TargetNodeId, SiblingIdx);
+
+			switch (DropZone)
+			{
+			// NOTE: the Sibling Idx should be swaped when we will disoplay the items from bottom up
+			case EItemDropZone::AboveItem:
+				TargetNodeId = Tree->FunctionInstance->AppendLayerNode(ParentNodeId, SiblingIdx);
+				break;
+
+			case EItemDropZone::BelowItem:
+				TargetNodeId = Tree->FunctionInstance->AppendLayerNode(ParentNodeId, SiblingIdx + 1);
+				break;
+			default:
+				break;
+			}
+
+			// Then drop
 			bool DidModifyTree = false;
 			for (const FAssetData& AssetData : AssetDropOp->GetAssets())
 			{
@@ -194,13 +253,14 @@ FReply SMaterialSubstrateTreeItem::OnLayerDrop(const FDragDropEvent& DragDropEve
 				if (AssetData.AssetClassPath.GetAssetName() == TEXT("MaterialFunctionMaterialLayer"))
 				{
 					InAssociation = EMaterialParameterAssociation::LayerParameter;
-					Tree->InsertLayerFromAsset(AssetData, StackParameterData->ParameterInfo.Index, InAssociation);
+
+					Tree->RefreshOnAssetChange(AssetData, TargetNodeId, InAssociation);
 					DidModifyTree = true;
 				}
 				else if (AssetData.AssetClassPath.GetAssetName() == TEXT("MaterialFunctionMaterialLayerBlend"))
 				{
 					InAssociation = EMaterialParameterAssociation::BlendParameter;
-					Tree->InsertLayerFromAsset(AssetData, StackParameterData->ParameterInfo.Index, InAssociation);
+					Tree->RefreshOnAssetChange(AssetData, TargetNodeId, InAssociation);
 					DidModifyTree = true;
 				}
 			}
@@ -290,12 +350,6 @@ void  SMaterialSubstrateTreeItem::OnOverrideParameter(bool NewValue, TObjectPtr<
 	OnOverrideParameter(NewValue, Parameter.Get());
 }
 
-void SMaterialSubstrateTreeItem::AddSubLayer()
-{
-	// create an empty sub layer at this point at the top of the sub-stack
-	
-}
-
 void SMaterialSubstrateTreeItem::Construct(const FArguments& InArgs, const TSharedRef<STableViewBase>& InOwnerTableView)
 {
 	StackParameterData = InArgs._StackParameterData;
@@ -308,6 +362,12 @@ void SMaterialSubstrateTreeItem::Construct(const FArguments& InArgs, const TShar
 	FText NameOverride;
 	TSharedRef<SVerticalBox> WrapperWidget = SNew(SVerticalBox);
 	EHorizontalAlignment ValueAlignment = HAlign_Left;
+
+	bool bIsBackgroundItem = IsNodeKeyForBackgroundParameter(StackParameterData->NodeKey);
+	bool bCanAppendSubLayer = Tree->FunctionInstance->CanAppendLayerNode(StackParameterData->ParameterInfo.Index) && !bIsBackgroundItem;
+	bool bCanRemoveLayer = Tree->FunctionInstance->CanRemoveLayerNode(StackParameterData->ParameterInfo.Index) && !bIsBackgroundItem;
+	bool bCanReorderLayer = !(bIsBackgroundItem || StackParameterData->ParameterInfo.Index == 0);
+
 // STACK --------------------------------------------------
 	if (StackParameterData->StackDataType == EStackDataType::Stack)
 	{
@@ -318,16 +378,18 @@ void SMaterialSubstrateTreeItem::Construct(const FArguments& InArgs, const TShar
 		// 		SNullWidget::NullWidget
 		// 	];
 #if WITH_EDITOR
-		NameOverride = Tree->FunctionInstance->GetLayerName(StackParameterData->ParameterInfo.Index);
+		int32 LayerFuncIndex = Tree->FunctionInstance->GetLayerFuncIndex(StackParameterData->ParameterInfo.Index);
+		int32 BlendFuncIndex = Tree->FunctionInstance->GetBlendFuncIndex(StackParameterData->ParameterInfo.Index);
+		NameOverride = Tree->FunctionInstance->GetLayerName(LayerFuncIndex);
 #endif
 		TSharedRef<SHorizontalBox> HeaderRowWidget = SNew(SHorizontalBox);
 
 		if (StackParameterData->ParameterInfo.Index != 0)
 		{
-			TAttribute<bool>::FGetter IsEnabledGetter = TAttribute<bool>::FGetter::CreateSP(Tree, &SMaterialSubstrateTree::IsLayerVisible, StackParameterData->ParameterInfo.Index);
+			TAttribute<bool>::FGetter IsEnabledGetter = TAttribute<bool>::FGetter::CreateSP(this, &SMaterialSubstrateTreeItem::IsLayerVisible);
 			TAttribute<bool> IsEnabledAttribute = TAttribute<bool>::Create(IsEnabledGetter);
 
-			FOnClicked VisibilityClickedDelegate = FOnClicked::CreateSP(Tree, &SMaterialSubstrateTree::ToggleLayerVisibility, StackParameterData->ParameterInfo.Index);
+			FOnClicked VisibilityClickedDelegate = FOnClicked::CreateSP(this, &SMaterialSubstrateTreeItem::ToggleLayerVisibility);
 
 			HeaderRowWidget->AddSlot()
 				.AutoWidth()
@@ -338,47 +400,51 @@ void SMaterialSubstrateTreeItem::Construct(const FArguments& InArgs, const TShar
 		}
 		const float ThumbnailSize = 64.0f;
 		TArray<TSharedPtr<FSortedParamData>> AssetChildren = StackParameterData->Children;
+		// Extract the asset elements to represent them as thumbnail boxes
 		for (TSharedPtr<FSortedParamData> AssetChild : AssetChildren)
 		{
-			TSharedPtr<SBox> ThumbnailBox;
-			UObject* AssetObject = nullptr;
-			AssetChild->ParameterHandle->GetValue(AssetObject);
-			int32 PreviewIndex = INDEX_NONE;
-			int32 ThumbnailIndex = INDEX_NONE;
-			EMaterialParameterAssociation PreviewAssociation = EMaterialParameterAssociation::GlobalParameter;
-			if (AssetObject)
+			if (AssetChild->StackDataType == EStackDataType::Asset)
 			{
-				if (Cast<UMaterialFunctionInterface>(AssetObject)->GetMaterialFunctionUsage() == EMaterialFunctionUsage::MaterialLayer)
+				TSharedPtr<SBox> ThumbnailBox;
+				UObject* AssetObject = nullptr;
+				AssetChild->ParameterHandle->GetValue(AssetObject);
+				int32 PreviewIndex = INDEX_NONE;
+				int32 ThumbnailIndex = INDEX_NONE;
+				EMaterialParameterAssociation PreviewAssociation = EMaterialParameterAssociation::GlobalParameter;
+				if (AssetObject)
 				{
-					PreviewIndex = StackParameterData->ParameterInfo.Index;
-					PreviewAssociation = EMaterialParameterAssociation::LayerParameter;
-					Tree->UpdateThumbnailMaterial(PreviewAssociation, PreviewIndex);
-					ThumbnailIndex = PreviewIndex;
-				}
-				if (Cast<UMaterialFunctionInterface>(AssetObject)->GetMaterialFunctionUsage() == EMaterialFunctionUsage::MaterialLayerBlend)
-				{
-					PreviewIndex = StackParameterData->ParameterInfo.Index;
-					PreviewAssociation = EMaterialParameterAssociation::BlendParameter;
-					Tree->UpdateThumbnailMaterial(PreviewAssociation, PreviewIndex, true);
-					ThumbnailIndex = PreviewIndex - 1;
+					if (Cast<UMaterialFunctionInterface>(AssetObject)->GetMaterialFunctionUsage() == EMaterialFunctionUsage::MaterialLayer)
+					{
+						PreviewIndex = LayerFuncIndex;
+						PreviewAssociation = EMaterialParameterAssociation::LayerParameter;
+						Tree->UpdateThumbnailMaterial(PreviewAssociation, PreviewIndex);
+						ThumbnailIndex = PreviewIndex;
+					}
+					if (Cast<UMaterialFunctionInterface>(AssetObject)->GetMaterialFunctionUsage() == EMaterialFunctionUsage::MaterialLayerBlend)
+					{
+						PreviewIndex = BlendFuncIndex;
+						PreviewAssociation = EMaterialParameterAssociation::BlendParameter;
+						Tree->UpdateThumbnailMaterial(PreviewAssociation, PreviewIndex);
+						ThumbnailIndex = PreviewIndex;
+					}
 				}
-			}
-			HeaderRowWidget->AddSlot()
-				.AutoWidth()
-				.HAlign(HAlign_Center)
-				.VAlign(VAlign_Center)
-				.Padding(4.0f)
-				.MaxWidth(ThumbnailSize)
-				[
-					SAssignNew(ThumbnailBox, SBox)
+				HeaderRowWidget->AddSlot()
+					.AutoWidth()
+					.HAlign(HAlign_Center)
+					.VAlign(VAlign_Center)
+					.Padding(4.0f)
+					.MaxWidth(ThumbnailSize)
 					[
-						Tree->CreateThumbnailWidget(PreviewAssociation, ThumbnailIndex, ThumbnailSize)
-					]
-				];
-			ThumbnailBox->SetMaxDesiredHeight(ThumbnailSize);
-			ThumbnailBox->SetMinDesiredHeight(ThumbnailSize);
-			ThumbnailBox->SetMinDesiredWidth(ThumbnailSize);
-			ThumbnailBox->SetMaxDesiredWidth(ThumbnailSize);
+						SAssignNew(ThumbnailBox, SBox)
+							[
+								Tree->CreateThumbnailWidget(PreviewAssociation, ThumbnailIndex, ThumbnailSize)
+							]
+					];
+				ThumbnailBox->SetMaxDesiredHeight(ThumbnailSize);
+				ThumbnailBox->SetMinDesiredHeight(ThumbnailSize);
+				ThumbnailBox->SetMinDesiredWidth(ThumbnailSize);
+				ThumbnailBox->SetMaxDesiredWidth(ThumbnailSize);
+			}
 		}
 
 		
@@ -391,8 +457,8 @@ void SMaterialSubstrateTreeItem::Construct(const FArguments& InArgs, const TShar
 				.Padding(5.0f)
 				[
 					SNew(SEditableTextBox)
-					.Text(TAttribute<FText>::Create(TAttribute<FText>::FGetter::CreateSP(this, &SMaterialSubstrateTreeItem::GetLayerName, InArgs._InTree, StackParameterData->ParameterInfo.Index)))
-					.OnTextCommitted(FOnTextCommitted::CreateSP(this, &SMaterialSubstrateTreeItem::OnNameChanged, InArgs._InTree, StackParameterData->ParameterInfo.Index))
+					.Text(TAttribute<FText>::Create(TAttribute<FText>::FGetter::CreateSP(this, &SMaterialSubstrateTreeItem::GetLayerName)))
+					.OnTextCommitted(FOnTextCommitted::CreateSP(this, &SMaterialSubstrateTreeItem::OnNameChanged))
 				];
 		}
 		else
@@ -414,6 +480,7 @@ void SMaterialSubstrateTreeItem::Construct(const FArguments& InArgs, const TShar
 			[
 				SNullWidget::NullWidget
 			];
+
 		HeaderRowWidget->AddSlot()
 			.AutoWidth()
 			.VAlign(VAlign_Center)
@@ -422,28 +489,32 @@ void SMaterialSubstrateTreeItem::Construct(const FArguments& InArgs, const TShar
 				SNew(SButton)
 				.Text(LOCTEXT("Unlink", "Unlink"))
 				.HAlign(HAlign_Center)
-				.OnClicked(Tree, &SMaterialSubstrateTree::UnlinkLayer, StackParameterData->ParameterInfo.Index)
+				.OnClicked(this, &SMaterialSubstrateTreeItem::UnlinkLayer)
 				.ToolTipText(LOCTEXT("UnlinkLayer", "Whether or not to unlink this layer/blend combination from the parent."))
-				.Visibility(Tree, &SMaterialSubstrateTree::GetUnlinkLayerVisibility, StackParameterData->ParameterInfo.Index)
+				.Visibility(this, &SMaterialSubstrateTreeItem::GetUnlinkLayerVisibility)
 			];
 
-		HeaderRowWidget->AddSlot()
+		// Can only add sub layer up to a certain level
+		if (bCanAppendSubLayer)
+		{
+			HeaderRowWidget->AddSlot()
 				.HAlign(HAlign_Left)
 				.AutoWidth()
 				.VAlign(VAlign_Center)
 				[
-					PropertyCustomizationHelpers::MakeAddButton(FSimpleDelegate::CreateSP(this, &SMaterialSubstrateTreeItem::AddSubLayer))
+					PropertyCustomizationHelpers::MakeAddButton(FSimpleDelegate::CreateSP(InArgs._InTree, &SMaterialSubstrateTree::AddNodeLayer, StackParameterData->ParameterInfo.Index))
 				];
+		}
 
 		// Can only remove layers that aren't the base layer.
-		if (StackParameterData->ParameterInfo.Index != 0)
+		if (bCanRemoveLayer)
 		{
 			HeaderRowWidget->AddSlot()
 				.AutoWidth()
 				.VAlign(VAlign_Center)
 				.Padding(0.0f, 0.0f, 5.0f, 0.0f)
 				[
-					PropertyCustomizationHelpers::MakeClearButton(FSimpleDelegate::CreateSP(InArgs._InTree, &SMaterialSubstrateTree::RemoveLayer, StackParameterData->ParameterInfo.Index))
+					PropertyCustomizationHelpers::MakeClearButton(FSimpleDelegate::CreateSP(InArgs._InTree, &SMaterialSubstrateTree::RemoveNodeLayer, StackParameterData->ParameterInfo.Index))
 				];
 		}
 		LeftSideWidget = HeaderRowWidget;
@@ -458,9 +529,10 @@ void SMaterialSubstrateTreeItem::Construct(const FArguments& InArgs, const TShar
 		WrapperWidget->AddSlot()
 			.Padding(15.0f)/*.AutoHeight()*/
 			[
-						SAssignNew(FinalStack, SHorizontalBox)
+				SAssignNew(FinalStack, SHorizontalBox)
 			];
-		if (StackParameterData->ParameterInfo.Index != 0)
+
+		if (bCanReorderLayer)
 		{
 			FinalStack->AddSlot()
 				.HAlign(HAlign_Center)
@@ -471,13 +543,18 @@ void SMaterialSubstrateTreeItem::Construct(const FArguments& InArgs, const TShar
 					FMaterialPropertyHelpers::MakeStackReorderHandle(SharedThis(this))
 				];
 		}
-		FinalStack->AddSlot()
-			.AutoWidth()
-			.VAlign(VAlign_Center)
-			.Padding(FMargin(2.0f))
-			[
-				SNew(SExpanderArrow, SharedThis(this))
-			];
+
+		if (bCanAppendSubLayer)
+		{
+			FinalStack->AddSlot()
+				.AutoWidth()
+				.VAlign(VAlign_Center)
+				.Padding(FMargin(2.0f))
+				[
+					SNew(SExpanderArrow, SharedThis(this))
+				];
+		}
+
 		FinalStack->AddSlot()
 			.Padding(FMargin(2.0f))
 			.VAlign(VAlign_Center)
@@ -527,14 +604,17 @@ int32 SMaterialSubstrateTreeItem::OnLayerItemPaintDropIndicator(EItemDropZone In
 }
 FString SMaterialSubstrateTreeItem::GetInstancePath(SMaterialSubstrateTree* InTree) const
 {
+	int32 LayerFuncIndex = Tree->FunctionInstance->GetLayerFuncIndex(StackParameterData->ParameterInfo.Index);
+	int32 BlendFuncIndex = Tree->FunctionInstance->GetBlendFuncIndex(StackParameterData->ParameterInfo.Index);
+
 	FString InstancePath;
-	if (StackParameterData->ParameterInfo.Association == EMaterialParameterAssociation::BlendParameter && InTree->FunctionInstance->Blends.IsValidIndex(StackParameterData->ParameterInfo.Index))
+	if (StackParameterData->ParameterInfo.Association == EMaterialParameterAssociation::BlendParameter && InTree->FunctionInstance->Blends.IsValidIndex(BlendFuncIndex))
 	{
-		InstancePath = InTree->FunctionInstance->Blends[StackParameterData->ParameterInfo.Index]->GetPathName();
+		InstancePath = InTree->FunctionInstance->Blends[BlendFuncIndex]->GetPathName();
 	}
-	else if (StackParameterData->ParameterInfo.Association == EMaterialParameterAssociation::LayerParameter && InTree->FunctionInstance->Layers.IsValidIndex(StackParameterData->ParameterInfo.Index))
+	else if (StackParameterData->ParameterInfo.Association == EMaterialParameterAssociation::LayerParameter && InTree->FunctionInstance->Layers.IsValidIndex(LayerFuncIndex))
 	{
-		InstancePath = InTree->FunctionInstance->Layers[StackParameterData->ParameterInfo.Index]->GetPathName();
+		InstancePath = InTree->FunctionInstance->Layers[LayerFuncIndex]->GetPathName();
 	}
 	return InstancePath;
 }
@@ -639,8 +719,15 @@ void SMaterialSubstrateTree::SetParentsExpansionState()
 	}
 }
 
-void SMaterialSubstrateTree::RefreshOnAssetChange(const struct FAssetData& InAssetData, int32 Index, EMaterialParameterAssociation MaterialType)
+void SMaterialSubstrateTree::RefreshOnAssetChange(const struct FAssetData& InAssetData, int32 InNodeId, EMaterialParameterAssociation MaterialType)
 {
+	auto NodePayload = FunctionInstance->GetNodePayload(InNodeId);
+
+	int32 Index = (MaterialType == EMaterialParameterAssociation::BlendParameter ? NodePayload.Blend : NodePayload.Layer);
+	// Early exit no op if the index for the asset modified is not valid
+	if (Index < 0)
+		return;
+
 	FMaterialPropertyHelpers::OnMaterialLayerAssetChanged(InAssetData, Index, MaterialType, FunctionInstanceHandle, FunctionInstance);
 	//set their overrides back to 0
 	MaterialEditorInstance->CleanParameterStack(Index, MaterialType);
@@ -656,42 +743,39 @@ void SMaterialSubstrateTree::ResetAssetToDefault(TSharedPtr<FSortedParamData> In
 	CreateGroupsWidget();
 	RequestTreeRefresh();
 }
-void SMaterialSubstrateTree::InsertLayerFromAsset(const struct FAssetData& InAssetData, int32 Index, EMaterialParameterAssociation MaterialType)
-{
-	const FScopedTransaction Transaction(LOCTEXT("InsertLayerAndBlend", "Insert a new Layer and a Blend into it"));
-	FunctionInstanceHandle->NotifyPreChange();
-	//TODO: Need to implement insert at index position instead of this
-	int32 LayerIndex = FunctionInstance->AppendBlendedLayer();
 
-	if (MaterialType == LayerParameter)
-	{
-		// Blend indices are offset by 1, no blend for base layer, so we only update the index for Layer
-		Index = LayerIndex;
-	}
-	FMaterialPropertyHelpers::OnMaterialLayerAssetChanged(InAssetData, Index, MaterialType, FunctionInstanceHandle, FunctionInstance);
-
-	FunctionInstanceHandle->NotifyFinishedChangingProperties();
-	MaterialEditorInstance->CleanParameterStack(Index, MaterialType);
-	CreateGroupsWidget();
-	MaterialEditorInstance->ResetOverrides(Index, MaterialType);
-	RequestTreeRefresh();
-}
-void SMaterialSubstrateTree::AddLayer()
+void SMaterialSubstrateTree::AddNodeLayer(int32 InParent)
 {
-	const FScopedTransaction Transaction(LOCTEXT("AddLayerAndBlend", "Add a new Layer and a Blend into it"));
+	/// Only if can really add a sub layer!
+	if (!FunctionInstance->CanAppendLayerNode(InParent))
+		return;
+
+	const FScopedTransaction Transaction(LOCTEXT("AddLayer", "Add a new Layer in the tree"));
 	FunctionInstanceHandle->NotifyPreChange();
-	FunctionInstance->AppendBlendedLayer();
+
+	// Create a new node
+	FunctionInstance->AppendLayerNode(InParent);
+	
 	FunctionInstanceHandle->NotifyPostChange(EPropertyChangeType::ArrayAdd);
 	CreateGroupsWidget();
 	RequestTreeRefresh();
 }
 
-void SMaterialSubstrateTree::RemoveLayer(int32 Index)
+void SMaterialSubstrateTree::RemoveNodeLayer(int32 InNodeId)
 {
+	/// Only if can really remove a sub layer!
+	if (!FunctionInstance->CanRemoveLayerNode(InNodeId))
+		return;
+
 	const FScopedTransaction Transaction(LOCTEXT("RemoveLayerAndBlend", "Remove a Layer and the attached Blend"));
 	FunctionInstanceHandle->NotifyPreChange();
-	FunctionInstance->RemoveBlendedLayerAt(Index);
-	MaterialEditorInstance->SourceInstance->RemoveLayerParameterIndex(Index);
+	
+	// Remove the node
+	auto NodePayload = FunctionInstance->GetNodePayload(InNodeId);
+
+	FunctionInstance->RemoveLayerNodeAt(InNodeId);
+	
+	MaterialEditorInstance->SourceInstance->RemoveLayerParameterIndex(NodePayload.Layer);
 	FunctionInstanceHandle->NotifyPostChange(EPropertyChangeType::ArrayRemove);
 	CreateGroupsWidget();
 	RequestTreeRefresh();
@@ -800,6 +884,107 @@ TSharedPtr<IDetailTreeNode> SMaterialSubstrateTree::FindParameterGroupsNode(TSha
 	return nullptr;
 }
 
+struct FRecursiveCreateWidgetsContext
+{
+	UDEditorParameterValue* Parameter;
+	TSharedPtr<IPropertyHandle>	LayerHandle;
+	TSharedPtr<IPropertyHandle>	BlendHandle;
+};
+
+void SMaterialSubstrateTree::RecursiveCreateWidgets(FRecursiveCreateWidgetsContext* InContext, FNodeId InNodeId, TArray<TSharedPtr<FSortedParamData>>& InParentContainer, bool GenerateChildren, bool bIsBackgroundItem)
+{
+	auto Payload = FunctionInstance->Tree.Payloads[InNodeId];
+
+
+	TSharedRef<FSortedParamData> StackProperty = MakeShared<FSortedParamData>();
+	StackProperty->StackDataType = EStackDataType::Stack;
+	StackProperty->Parameter = InContext->Parameter;
+	StackProperty->ParameterInfo.Index = InNodeId;
+	StackProperty->NodeKey = FString::FromInt(StackProperty->ParameterInfo.Index);
+
+	if (bIsBackgroundItem)
+	{
+		StackProperty->NodeKey = CreateNodeKeyForBackgroundParameter(StackProperty->NodeKey);
+	}
+
+
+	if (GenerateChildren)
+	{
+		// Create the default background sub item representing the layer and blend func stored in the layer item
+		RecursiveCreateWidgets(InContext, InNodeId, StackProperty->Children, false, true);
+
+		// Sub layers
+		auto RootChildren = FunctionInstance->GetNodeChildren(InNodeId);
+		for (int i = 0; i < RootChildren.Num(); ++i)
+		{
+			RecursiveCreateWidgets(InContext, RootChildren[i], StackProperty->Children, false, false);
+		}
+	}
+
+	if (Payload.Layer != -1)
+	{
+		TSharedRef<FSortedParamData> ChildProperty = MakeShared<FSortedParamData>();
+		ChildProperty->StackDataType = EStackDataType::Asset;
+		ChildProperty->Parameter = InContext->Parameter;
+		ChildProperty->ParameterHandle = InContext->LayerHandle->AsArray()->GetElement(Payload.Layer);
+		ChildProperty->ParameterNode = Generator->FindTreeNode(ChildProperty->ParameterHandle);
+		ChildProperty->ParameterInfo.Index = Payload.Layer;
+		ChildProperty->ParameterInfo.Association = EMaterialParameterAssociation::LayerParameter;
+		ChildProperty->NodeKey = FString::FromInt(ChildProperty->ParameterInfo.Index) + FString::FromInt(ChildProperty->ParameterInfo.Association);
+
+		{
+			UObject* AssetObject = nullptr;
+			ChildProperty->ParameterHandle->GetValue(AssetObject);
+			if (AssetObject)
+			{
+				if (MaterialEditorInstance->StoredLayerPreviews[Payload.Layer] == nullptr)
+				{
+					MaterialEditorInstance->StoredLayerPreviews[Payload.Layer] = (NewObject<UMaterialInstanceConstant>(MaterialEditorInstance, NAME_None));
+				}
+				UMaterialInterface* EditedMaterial = Cast<UMaterialFunctionInterface>(AssetObject)->GetPreviewMaterial();
+				if (MaterialEditorInstance->StoredLayerPreviews[Payload.Layer] && MaterialEditorInstance->StoredLayerPreviews[Payload.Layer]->Parent != EditedMaterial)
+				{
+					MaterialEditorInstance->StoredLayerPreviews[Payload.Layer]->SetParentEditorOnly(EditedMaterial);
+				}
+			}
+		}
+
+		StackProperty->Children.Add(ChildProperty);
+	}
+
+	if (Payload.Blend != -1)
+	{
+		TSharedRef<FSortedParamData> ChildProperty = MakeShared<FSortedParamData>();
+		ChildProperty->StackDataType = EStackDataType::Asset;
+		ChildProperty->Parameter = InContext->Parameter;
+		ChildProperty->ParameterHandle = InContext->BlendHandle->AsArray()->GetElement(Payload.Blend);
+		ChildProperty->ParameterNode = Generator->FindTreeNode(ChildProperty->ParameterHandle);
+		ChildProperty->ParameterInfo.Index = Payload.Blend;
+		ChildProperty->ParameterInfo.Association = EMaterialParameterAssociation::BlendParameter;
+		ChildProperty->NodeKey = FString::FromInt(ChildProperty->ParameterInfo.Index) + FString::FromInt(ChildProperty->ParameterInfo.Association);
+		{
+			UObject* AssetObject = nullptr;
+			ChildProperty->ParameterHandle->GetValue(AssetObject);
+			if (AssetObject)
+			{
+				if (MaterialEditorInstance->StoredBlendPreviews[Payload.Blend] == nullptr)
+				{
+					MaterialEditorInstance->StoredBlendPreviews[Payload.Blend] = (NewObject<UMaterialInstanceConstant>(MaterialEditorInstance, NAME_None));
+				}
+				UMaterialInterface* EditedMaterial = Cast<UMaterialFunctionInterface>(AssetObject)->GetPreviewMaterial();
+				if (MaterialEditorInstance->StoredBlendPreviews[Payload.Blend] && MaterialEditorInstance->StoredBlendPreviews[Payload.Blend]->Parent != EditedMaterial)
+				{
+					MaterialEditorInstance->StoredBlendPreviews[Payload.Blend]->SetParentEditorOnly(EditedMaterial);
+				}
+			}
+		}
+
+		StackProperty->Children.Add(ChildProperty);
+	}
+
+	InParentContainer.Add(StackProperty);
+}
+
 void SMaterialSubstrateTree::CreateGroupsWidget()
 {
 	check(MaterialEditorInstance);
@@ -866,123 +1051,43 @@ void SMaterialSubstrateTree::CreateGroupsWidget()
 					FunctionInstance = reinterpret_cast<FMaterialLayersFunctions*>(*It);
 					FunctionInstanceHandle = ParameterValueProperty;
 
+
 					TSharedPtr<IPropertyHandle>	LayerHandle = ChildHandle->GetChildHandle("Layers").ToSharedRef();
 					TSharedPtr<IPropertyHandle> BlendHandle = ChildHandle->GetChildHandle("Blends").ToSharedRef();
-					uint32 LayerChildren;
-					LayerHandle->GetNumChildren(LayerChildren);
-					uint32 BlendChildren;
-					BlendHandle->GetNumChildren(BlendChildren);
-					if (MaterialEditorInstance->StoredLayerPreviews.Num() != LayerChildren)
+					uint32 NumLayerChildren;
+					LayerHandle->GetNumChildren(NumLayerChildren);
+					uint32 NumBlendChildren;
+					BlendHandle->GetNumChildren(NumBlendChildren);
+					if (MaterialEditorInstance->StoredLayerPreviews.Num() != NumLayerChildren)
 					{
 						MaterialEditorInstance->StoredLayerPreviews.Empty();
-						MaterialEditorInstance->StoredLayerPreviews.AddDefaulted(LayerChildren);
+						MaterialEditorInstance->StoredLayerPreviews.AddDefaulted(NumLayerChildren);
 					}
-					if (MaterialEditorInstance->StoredBlendPreviews.Num() != BlendChildren)
+					if (MaterialEditorInstance->StoredBlendPreviews.Num() != NumBlendChildren)
 					{
 						MaterialEditorInstance->StoredBlendPreviews.Empty();
-						MaterialEditorInstance->StoredBlendPreviews.AddDefaulted(BlendChildren);
+						MaterialEditorInstance->StoredBlendPreviews.AddDefaulted(NumBlendChildren);
 					}
 
-					TSharedRef<FSortedParamData> StackProperty = MakeShared<FSortedParamData>();
-					StackProperty->StackDataType = EStackDataType::Stack;
-					StackProperty->Parameter = Parameter;
-					StackProperty->ParameterInfo.Index = LayerChildren - 1;
-					StackProperty->NodeKey = FString::FromInt(StackProperty->ParameterInfo.Index);
-
-
-					TSharedRef<FSortedParamData> ChildProperty = MakeShared<FSortedParamData>();
-					ChildProperty->StackDataType = EStackDataType::Asset;
-					ChildProperty->Parameter = Parameter;
-					ChildProperty->ParameterHandle = LayerHandle->AsArray()->GetElement(LayerChildren - 1);
-					ChildProperty->ParameterNode = Generator->FindTreeNode(ChildProperty->ParameterHandle);
-					ChildProperty->ParameterInfo.Index = LayerChildren - 1;
-					ChildProperty->ParameterInfo.Association = EMaterialParameterAssociation::LayerParameter;
-					ChildProperty->NodeKey = FString::FromInt(ChildProperty->ParameterInfo.Index) + FString::FromInt(ChildProperty->ParameterInfo.Association);
 
+#ifdef ENABLE_MATERIAL_LAYER_PROTOTYPE					
 					{
-						UObject* AssetObject = nullptr;
-						ChildProperty->ParameterHandle->GetValue(AssetObject);
-						if (AssetObject)
+						TSharedPtr<IPropertyHandle>	TreeHandle = ChildHandle->GetChildHandle("Tree").ToSharedRef();
+						
+						FRecursiveCreateWidgetsContext Context {
+							.Parameter = Parameter,
+							.LayerHandle = LayerHandle,
+							.BlendHandle = BlendHandle
+						};
+
+						auto RootChildren = FunctionInstance->GetNodeChildren(-1);
+						for (int i = 0; i < RootChildren.Num(); ++i)
 						{
-							if (MaterialEditorInstance->StoredLayerPreviews[LayerChildren - 1] == nullptr)
-							{
-								MaterialEditorInstance->StoredLayerPreviews[LayerChildren - 1] = (NewObject<UMaterialInstanceConstant>(MaterialEditorInstance, NAME_None));
-							}
-							UMaterialInterface* EditedMaterial = Cast<UMaterialFunctionInterface>(AssetObject)->GetPreviewMaterial();
-							if (MaterialEditorInstance->StoredLayerPreviews[LayerChildren - 1] && MaterialEditorInstance->StoredLayerPreviews[LayerChildren - 1]->Parent != EditedMaterial)
-							{
-								MaterialEditorInstance->StoredLayerPreviews[LayerChildren - 1]->SetParentEditorOnly(EditedMaterial);
-							}
-						}
+							RecursiveCreateWidgets(&Context, RootChildren[i], LayerProperties, true, false);
+						}	
 					}
+#endif // ENABLE_MATERIAL_LAYER_PROTOTYPE
 
-					StackProperty->Children.Add(ChildProperty);
-					LayerProperties.Add(StackProperty);
-
-					if (BlendChildren > 0 && LayerChildren > BlendChildren)
-					{
-						for (int32 Counter = BlendChildren - 1; Counter >= 0; Counter--)
-						{
-							ChildProperty = MakeShared<FSortedParamData>();
-							ChildProperty->StackDataType = EStackDataType::Asset;
-							ChildProperty->Parameter = Parameter;
-							ChildProperty->ParameterHandle = BlendHandle->AsArray()->GetElement(Counter);
-							ChildProperty->ParameterNode = Generator->FindTreeNode(ChildProperty->ParameterHandle);
-							ChildProperty->ParameterInfo.Index = Counter;
-							ChildProperty->ParameterInfo.Association = EMaterialParameterAssociation::BlendParameter;
-							ChildProperty->NodeKey = FString::FromInt(ChildProperty->ParameterInfo.Index) + FString::FromInt(ChildProperty->ParameterInfo.Association);
-							{
-								UObject* AssetObject = nullptr;
-								ChildProperty->ParameterHandle->GetValue(AssetObject);
-								if (AssetObject)
-								{
-									if (MaterialEditorInstance->StoredBlendPreviews[Counter] == nullptr)
-									{
-										MaterialEditorInstance->StoredBlendPreviews[Counter] = (NewObject<UMaterialInstanceConstant>(MaterialEditorInstance, NAME_None));
-									}
-									UMaterialInterface* EditedMaterial = Cast<UMaterialFunctionInterface>(AssetObject)->GetPreviewMaterial();
-									if (MaterialEditorInstance->StoredBlendPreviews[Counter] && MaterialEditorInstance->StoredBlendPreviews[Counter]->Parent != EditedMaterial)
-									{
-										MaterialEditorInstance->StoredBlendPreviews[Counter]->SetParentEditorOnly(EditedMaterial);
-									}
-								}
-							}
-							LayerProperties.Last()->Children.Add(ChildProperty);
-
-							StackProperty = MakeShared<FSortedParamData>();
-							StackProperty->StackDataType = EStackDataType::Stack;
-							StackProperty->Parameter = Parameter;
-							StackProperty->ParameterInfo.Index = Counter;
-							StackProperty->NodeKey = FString::FromInt(StackProperty->ParameterInfo.Index);
-							LayerProperties.Add(StackProperty);
-
-							ChildProperty = MakeShared<FSortedParamData>();
-							ChildProperty->StackDataType = EStackDataType::Asset;
-							ChildProperty->Parameter = Parameter;
-							ChildProperty->ParameterHandle = LayerHandle->AsArray()->GetElement(Counter);
-							ChildProperty->ParameterNode = Generator->FindTreeNode(ChildProperty->ParameterHandle);
-							ChildProperty->ParameterInfo.Index = Counter;
-							ChildProperty->ParameterInfo.Association = EMaterialParameterAssociation::LayerParameter;
-							ChildProperty->NodeKey = FString::FromInt(ChildProperty->ParameterInfo.Index) + FString::FromInt(ChildProperty->ParameterInfo.Association);
-							{
-								UObject* AssetObject = nullptr;
-								ChildProperty->ParameterHandle->GetValue(AssetObject);
-								if (AssetObject)
-								{
-									if (MaterialEditorInstance->StoredLayerPreviews[Counter] == nullptr)
-									{
-										MaterialEditorInstance->StoredLayerPreviews[Counter] = (NewObject<UMaterialInstanceConstant>(MaterialEditorInstance, NAME_None));
-									}
-									UMaterialInterface* EditedMaterial = Cast<UMaterialFunctionInterface>(AssetObject)->GetPreviewMaterial();
-									if (MaterialEditorInstance->StoredLayerPreviews[Counter] && MaterialEditorInstance->StoredLayerPreviews[Counter]->Parent != EditedMaterial)
-									{
-										MaterialEditorInstance->StoredLayerPreviews[Counter]->SetParentEditorOnly(EditedMaterial);
-									}
-								}
-							}
-							LayerProperties.Last()->Children.Add(ChildProperty);
-						}
-					}
 				}
 				else
 				{
@@ -1061,8 +1166,6 @@ TSharedRef<SWidget> SMaterialSubstrateTree::CreateThumbnailWidget(EMaterialParam
 
 void SMaterialSubstrateTree::UpdateThumbnailMaterial(TEnumAsByte<EMaterialParameterAssociation> InAssociation, int32 InIndex, bool bAlterBlendIndex)
 {
-	// Need to invert index b/c layer properties is generated in reverse order
-	TArray<TSharedPtr<FSortedParamData>> AssetChildren = LayerProperties[LayerProperties.Num() - 1 - InIndex]->Children;
 	UMaterialInstanceConstant* MaterialToUpdate = nullptr;
 	int32 ParameterIndex = InIndex;
 	if (InAssociation == EMaterialParameterAssociation::LayerParameter)
@@ -1078,34 +1181,39 @@ void SMaterialSubstrateTree::UpdateThumbnailMaterial(TEnumAsByte<EMaterialParame
 		MaterialToUpdate = MaterialEditorInstance->StoredBlendPreviews[ParameterIndex];
 	}
 
-	TArray<FEditorParameterGroup> ParameterGroups;
-	for (TSharedPtr<FSortedParamData> AssetChild : AssetChildren)
+	if (MaterialToUpdate != nullptr)
 	{
-		for (TSharedPtr<FSortedParamData> Group : AssetChild->Children)
+		// Need to invert index b/c layer properties is generated in reverse order
+	/*	TArray<TSharedPtr<FSortedParamData>> AssetChildren = LayerProperties[LayerProperties.Num() - 1 - InIndex]->Children;
+
+		TArray<FEditorParameterGroup> ParameterGroups;
+		for (TSharedPtr<FSortedParamData> AssetChild : AssetChildren)
 		{
-			if (Group->ParameterInfo.Association == InAssociation)
+			for (TSharedPtr<FSortedParamData> Group : AssetChild->Children)
 			{
-				FEditorParameterGroup DuplicatedGroup = FEditorParameterGroup();
-				DuplicatedGroup.GroupAssociation = Group->Group.GroupAssociation;
-				DuplicatedGroup.GroupName = Group->Group.GroupName;
-				DuplicatedGroup.GroupSortPriority = Group->Group.GroupSortPriority;
-				for (UDEditorParameterValue* Parameter : Group->Group.Parameters)
+				if (Group->ParameterInfo.Association == InAssociation)
 				{
-					if (Parameter->ParameterInfo.Index == ParameterIndex)
+					FEditorParameterGroup DuplicatedGroup = FEditorParameterGroup();
+					DuplicatedGroup.GroupAssociation = Group->Group.GroupAssociation;
+					DuplicatedGroup.GroupName = Group->Group.GroupName;
+					DuplicatedGroup.GroupSortPriority = Group->Group.GroupSortPriority;
+					for (UDEditorParameterValue* Parameter : Group->Group.Parameters)
 					{
-						DuplicatedGroup.Parameters.Add(Parameter);
+						if (Parameter->ParameterInfo.Index == ParameterIndex)
+						{
+							DuplicatedGroup.Parameters.Add(Parameter);
+						}
 					}
+					ParameterGroups.Add(DuplicatedGroup);
 				}
-				ParameterGroups.Add(DuplicatedGroup);
 			}
-		}
 
 	
 
-	}
-	if (MaterialToUpdate != nullptr)
-	{
+		}
+
 		FMaterialPropertyHelpers::TransitionAndCopyParameters(MaterialToUpdate, ParameterGroups, true);
+*/
 	}
 }
 
diff --git a/Engine/Source/Editor/MaterialEditor/Private/SMaterialSubstrateTree.h b/Engine/Source/Editor/MaterialEditor/Private/SMaterialSubstrateTree.h
--- a/Engine/Source/Editor/MaterialEditor/Private/SMaterialSubstrateTree.h
+++ b/Engine/Source/Editor/MaterialEditor/Private/SMaterialSubstrateTree.h
@@ -23,6 +23,7 @@ class SMaterialSubstrateTree;
 class SMaterialLayersFunctionsInstanceWrapper;
 class UDEditorParameterValue;
 class UMaterialEditorInstanceConstant;
+struct FRecursiveCreateWidgetsContext;
 
 typedef TSharedPtr<FSortedParamData> FSortedParamDataPtr;
 class SMaterialSubstrateTree : public STreeView<FSortedParamDataPtr>
@@ -50,16 +51,20 @@ public:
 	TWeakObjectPtr<class UDEditorParameterValue> FunctionParameter;
 	struct FMaterialLayersFunctions* FunctionInstance;
 	TSharedPtr<IPropertyHandle> FunctionInstanceHandle;
-	void RefreshOnAssetChange(const struct FAssetData& InAssetData, int32 Index, EMaterialParameterAssociation MaterialType);
+	void RefreshOnAssetChange(const struct FAssetData& InAssetData, int32 InNodeId, EMaterialParameterAssociation MaterialType);
 	void ResetAssetToDefault(TSharedPtr<FSortedParamData> InData);
-	void AddLayer();
-	void InsertLayerFromAsset(const struct FAssetData& InAssetData, int32 Index, EMaterialParameterAssociation MaterialType);
-	void RemoveLayer(int32 Index);
+	
+	void AddRootNodeLayer() { AddNodeLayer(); }
+	void AddNodeLayer(int32 InParent = -1);
+	void RemoveNodeLayer(int32 InNodeId);
+
 	FReply UnlinkLayer(int32 Index);
 	FReply RelinkLayersToParent();
 	EVisibility GetUnlinkLayerVisibility(int32 Index) const;
 	EVisibility GetRelinkLayersToParentVisibility() const;
 	FReply ToggleLayerVisibility(int32 Index);
+	bool IsLayerVisible(int32 Index) const;
+
 	TSharedPtr<class FAssetThumbnailPool> GetTreeThumbnailPool();
 
 	/** Object that stores all of the possible parameters we can edit */
@@ -67,8 +72,9 @@ public:
 
 	/** Builds the custom parameter groups category */
 	void CreateGroupsWidget();
+	using FNodeId = int32;
+	void RecursiveCreateWidgets(struct FRecursiveCreateWidgetsContext* Context, FNodeId InNodeId, TArray<TSharedPtr<FSortedParamData>>& InParentContainer, bool GenerateChildren, bool IsBackgroundItem);
 
-	bool IsLayerVisible(int32 Index) const;
 
 	SMaterialLayersFunctionsInstanceWrapper* GetWrapper() { return Wrapper; }
 
@@ -136,14 +142,20 @@ public:
 	bool GetFilterState(SMaterialSubstrateTree* InTree, TSharedPtr<FSortedParamData> InStackData) const;
 	void FilterClicked(const ECheckBoxState NewCheckedState, SMaterialSubstrateTree* InTree, TSharedPtr<FSortedParamData> InStackData);
 	ECheckBoxState GetFilterChecked(SMaterialSubstrateTree* InTree, TSharedPtr<FSortedParamData> InStackData) const;
-	FText GetLayerName(SMaterialSubstrateTree* InTree, int32 Counter) const;
-	void OnNameChanged(const FText& InText, ETextCommit::Type CommitInfo, SMaterialSubstrateTree* InTree, int32 Counter);
+	FText GetLayerName() const;
+	void OnNameChanged(const FText& InText, ETextCommit::Type CommitInfo);
+	
+	FReply ToggleLayerVisibility();
+	bool IsLayerVisible() const;
+
+	FReply UnlinkLayer();
+	EVisibility GetUnlinkLayerVisibility() const;
 
 	TOptional<EItemDropZone> CanAcceptDrop(const FDragDropEvent& DragDropEvent, EItemDropZone DropZone, FSortedParamDataPtr Item);
 	FReply OnAcceptDrop(const FDragDropEvent& DragDropEvent, EItemDropZone DropZone, FSortedParamDataPtr TargetItem);
 	void OnLayerDragEnter(const FDragDropEvent& DragDropEvent) override
 	{
-		if (StackParameterData->ParameterInfo.Index != 0)
+		//if (StackParameterData->ParameterInfo.Index != 0)
 		{
 			bIsHoveredDragTarget = true;
 		}
@@ -162,7 +174,7 @@ public:
 	FReply OnLayerDrop(const FDragDropEvent& DragDropEvent, EItemDropZone DropZone, FSortedParamDataPtr TargetItem);
 	void OnOverrideParameter(bool NewValue, class UDEditorParameterValue* Parameter);
 	void OnOverrideParameter(bool NewValue, TObjectPtr<UDEditorParameterValue> Parameter);
-	void AddSubLayer();
+
 	/**
 	* Construct the widget
 	*
diff --git a/Engine/Source/Runtime/Engine/Private/Materials/MaterialExpressions.cpp b/Engine/Source/Runtime/Engine/Private/Materials/MaterialExpressions.cpp
--- a/Engine/Source/Runtime/Engine/Private/Materials/MaterialExpressions.cpp
+++ b/Engine/Source/Runtime/Engine/Private/Materials/MaterialExpressions.cpp
@@ -17337,6 +17337,340 @@ void FMaterialLayersFunctionsID::Append(FShaderKeyGenerator& KeyGen) const
 
 const FGuid FMaterialLayersFunctions::BackgroundGuid(2u, 0u, 0u, 0u);
 
+
+bool FMaterialLayersFunctionsTree::Serialize(FArchive& Ar)
+{
+	Ar << Nodes;
+	Ar << Payloads;
+	Ar << Root;
+	return true;
+}
+
+using FLayerNodeId = FMaterialLayersFunctionsTree::FNodeId;
+using FLayerNode = FMaterialLayersFunctionsTree::FNode;
+using FLayerNodePayload = FMaterialLayersFunctionsTree::FPayload;
+using FLayerNodeIdArray = TArray<FLayerNodeId>;
+
+FLayerNodeId FMaterialLayersFunctionsTree::AllocNode(const FPayload& InPayload)
+{
+	Nodes.AddDefaulted();
+	Payloads.Add(InPayload);
+	return (FNodeId)Nodes.Num() - 1;
+}
+
+
+// Find and return the node at the specified Id or null if invalid 
+FLayerNode* FMaterialLayersFunctionsTree::GetNode(FNodeId InNodeId) const
+{
+	if (IsValidId(InNodeId))
+	{
+		return const_cast<FNode*>(Nodes.GetData() + InNodeId);
+	}
+
+	return nullptr; // invalid index => return null
+}
+
+// Find and return the node payload at the specified Id or null if invalid 
+FLayerNodePayload* FMaterialLayersFunctionsTree::GetPayload(FNodeId InNodeId) const
+{
+	if (IsValidId(InNodeId))
+	{
+		return const_cast<FPayload*>(Payloads.GetData() + InNodeId);
+	}
+
+	return nullptr; // invalid index => return null
+}
+
+// Find and return the parent node at the specified Id or null if invalid 
+FLayerNode* FMaterialLayersFunctionsTree::GetParent(FNodeId InNodeId) const
+{
+	FNode* Node = GetNode(InNodeId);
+	if (Node && IsValidId(Node->Parent))
+	{
+		return GetNode(Node->Parent);
+	}
+
+	return nullptr; // invalid parent => return null
+}
+
+// Return the depth of the node in the tree, the root is depth 0
+// the first valid node in the tree under the root is depth 1
+FLayerNodeIdArray FMaterialLayersFunctionsTree::GetChildrenIds(FNodeId InNodeId) const
+{
+	FLayerNodeIdArray ChildrendIds;
+
+	FNodeId NextChildId = Root;
+	FNode* Node = GetNode(InNodeId);
+	if (Node)
+		NextChildId = Node->ChildrenHead;
+
+	while (IsValidId(NextChildId))
+	{
+		ChildrendIds.Add(NextChildId);
+		Node = GetNode(NextChildId);
+		if (Node)
+			NextChildId = Node->NextSibling;
+		else
+			NextChildId = -1;
+	}
+
+	return ChildrendIds;
+}
+
+int32 FMaterialLayersFunctionsTree::GetDepth(FNodeId InNodeId) const
+{
+	int32 Depth = 0;
+	FNode* Node = GetNode(InNodeId);
+	while (Node)
+	{
+		Node = GetNode(Node->Parent);
+		Depth++;
+	}
+
+	return Depth;
+}
+
+
+// Find the node which is the head of the list of siblings where InNodeID belongs
+FLayerNodeId FMaterialLayersFunctionsTree::GetSiblingHeadId(FNodeId InNodeId) const
+{
+	if (!IsValidId(InNodeId))
+		return InvalidId;
+
+	FNode* TargetNode = const_cast<FNode*>(Nodes.GetData() + InNodeId);
+
+	FNodeId HeadNodeId = Root;
+	if (TargetNode->Parent != InvalidId)
+	{
+		HeadNodeId =Nodes[TargetNode->Parent].ChildrenHead;
+	}
+
+	return HeadNodeId;
+}
+
+FLayerNodeId FMaterialLayersFunctionsTree::AddNode(const FPayload& InPayload, FNodeId InParent, int32 InAtChildIndex)
+{
+	// Check Parent id
+	if (!IsValidId(InParent))
+	{
+		// Bad parent id, allocate at root at the end
+		InParent = -1;
+	}
+
+	// Allocate the new node setting the payload value
+	FNodeId NodeId = AllocNode(InPayload);
+
+	// Let's connect the new node in the hierarchy
+	FNode* NewNode = Nodes.GetData() + NodeId;
+	NewNode->Parent = InParent;
+
+	// First find the parent node and from that the head node for the list in which we will add the new node
+	// Initialize the case if parent is the tree root
+	FNodeId HeadSiblingId = Root;
+	FNodeId* ParentChildrenHeaddPtr = &Root;
+	// Parent is actually a node and not the root:
+	if (InParent >= 0)
+	{
+		FNode* ParentNode = Nodes.GetData() + InParent;
+		HeadSiblingId = ParentNode->ChildrenHead;
+		ParentChildrenHeaddPtr = &ParentNode->ChildrenHead;
+	}
+
+	// When Insert at head or no siblings yet then update parent sub
+	if ((InAtChildIndex == 0) || (HeadSiblingId < 0))
+	{
+		(*ParentChildrenHeaddPtr) = NodeId;
+		NewNode->NextSibling = HeadSiblingId;
+
+		// Done
+		return NodeId;
+	}
+
+	// Second insert the new node among the siblings AFTER head
+	FNode* SiblingNode = Nodes.GetData() + HeadSiblingId;
+
+	int32 SiblingNextIdx = 1;
+	while (SiblingNode->NextSibling >= 0)
+	{
+		// if the next sibling is where we want to insert then insert
+		if (InAtChildIndex == SiblingNextIdx)
+		{
+			NewNode->NextSibling = SiblingNode->NextSibling;
+			SiblingNode->NextSibling = NodeId;
+
+			// Done
+			return NodeId;
+		}
+
+		// next sibling
+		SiblingNode = Nodes.GetData() + SiblingNode->NextSibling;
+		SiblingNextIdx++;
+	}
+
+	// we went through the siblings and now hold the tail node, let's add the new node as the new tail
+	SiblingNode->NextSibling = NodeId;
+
+	// Done
+	return NodeId;
+}	
+
+
+
+FLayerNodeIdArray FMaterialLayersFunctionsTree::RemoveNode(FNodeId RemovedNodeId)
+{
+	FLayerNodeIdArray RemovedIds;
+
+	FNode* RemovedNode = GetNode(RemovedNodeId);
+	if (!RemovedNode)
+		return RemovedIds;
+
+	FNodeId NextId = RemovedNode->NextSibling;
+
+	FNodeId HeadId = GetSiblingHeadId(RemovedNodeId);
+	FNode* HeadSibling = GetNode(HeadId);
+
+	// walk from head sibling to removed node
+	// then detach from tree
+	FNodeId NextPrevNodeId = HeadId;
+	FNode* PrevNode = HeadSibling; // prev node is the head sibling at first
+	FNode* NextPrevNode = HeadSibling; // next prev node is the head sibling at first
+	while (NextPrevNode)
+	{
+		if (NextPrevNodeId == RemovedNodeId)
+		{
+			break;
+		}
+		PrevNode = NextPrevNode;
+		NextPrevNodeId = PrevNode->NextSibling;
+		NextPrevNode = GetNode(PrevNode->NextSibling);
+	}
+
+	// Removed node is the sibling head
+	// adjust parent link to children
+	if (NextPrevNode == HeadSibling)
+	{
+		FNodeId OldChildHeadId = Root;
+
+		if (RemovedNode->Parent != InvalidId)
+		{
+			FNode* ParentNode = Nodes.GetData() + RemovedNode->Parent;
+			OldChildHeadId = ParentNode->ChildrenHead;
+			ParentNode->ChildrenHead = NextId;
+		}
+		else // Parent is the root so relink root
+		{
+			Root = NextId;
+		}
+	}
+	// Removed node in the middle of the siblings
+	else
+	{
+		PrevNode->NextSibling = NextId;
+	}
+
+	Nodes[RemovedNodeId] = FNode();
+	Payloads[RemovedNodeId] = FPayload();
+	RemovedIds.Add(RemovedNodeId);
+
+	return RemovedIds;
+}
+
+FString FMaterialLayersFunctionsTree::Log(FString InTab) const
+{
+	FString NewLine = TEXT("\r\n");
+	FString LogMessage;
+
+	auto Visitor = [&](FNodeId InNodeId, int32 InDepth, int32 InSiblingNum, FPayload InPayload)
+		{
+			FString Tab = InTab;
+			for (int d = 0; d < InDepth; d++)
+				Tab += "  ";
+			LogMessage += Tab + FString::Printf(TEXT("%i %i - L%i B%i"), InNodeId, InSiblingNum, InPayload.Layer, InPayload.Blend) + NewLine;
+		};
+
+	Traverse(Visitor);
+
+	return LogMessage;
+}
+
+namespace MLFT {
+
+	void RemovePayloadIndex(FMaterialLayersFunctionsTree& Tree, int32 InRemovedIdx)
+	{
+		for (auto& Payload : Tree.Payloads)
+		{
+			if (Payload.Layer > InRemovedIdx)
+			{
+				Payload.Layer--;
+			}
+			if (Payload.Blend > InRemovedIdx - 1)
+			{
+				Payload.Blend--;
+			}
+		}
+	}
+
+
+	FString Print(const FMaterialLayersFunctionsRuntimeData* LFRD, FString InTab)
+	{
+		FString NewLine = TEXT("\r\n");
+		FString LogMessage; // = InTab + TEXT("FMaterialLayersFunctionsRuntimeData") + NewLine;
+
+		InTab += "  ";
+
+		// Layers
+		LogMessage += InTab + TEXT("Layers:") + NewLine;
+		for (int32 i = 0; i < LFRD->Layers.Num(); ++i)
+		{
+			FString LayerHead = InTab + FString::Printf(TEXT("%i - "), i);
+			auto Layer = LFRD->Layers[i];
+			if (Layer)
+			{
+				LayerHead += FString::Printf(TEXT("%s"), *Layer->GetName());
+			}
+			LogMessage += LayerHead + NewLine;
+		}
+		// Blends
+		LogMessage += InTab + TEXT("Blends:") + NewLine;
+		for (int32 i = 0; i < LFRD->Blends.Num(); ++i)
+		{
+			FString BlendHead = InTab + FString::Printf(TEXT("%i - "), i);
+			auto Blend = LFRD->Blends[i];
+			if (Blend)
+			{
+				BlendHead += FString::Printf(TEXT("%s"), *Blend->GetName());
+			}
+			LogMessage += BlendHead + NewLine;
+		}
+		// Tree
+		LogMessage += InTab + TEXT("Tree:") + NewLine;
+		LogMessage += LFRD->Tree.Log(InTab);
+		return LogMessage;
+	}
+
+
+	void Log(const FMaterialLayersFunctionsRuntimeData* LFRD, FString CallingSite = FString())
+	{
+		FString NewLine = TEXT("\r\n");
+		FString LogMessage;
+		FString Tab = " ";
+
+		LogMessage += FString::Printf(TEXT("**** FMaterialLayersFunctionsRuntimeData from <%s> ****"), *CallingSite) + NewLine;
+		LogMessage += Print(LFRD, Tab);
+		LogMessage += TEXT("**** *********** ****") + NewLine;
+
+		TArray<FString> Lines;
+		LogMessage.ParseIntoArray(Lines, TEXT("\n"));
+		for (const FString& Line : Lines)
+		{
+			UE_LOG(LogMaterial, Log, TEXT("%s"), *Line);
+		}
+	}
+
+
+}
+
+
 FMaterialLayersFunctionsRuntimeData::~FMaterialLayersFunctionsRuntimeData()
 {
 #if WITH_EDITORONLY_DATA
@@ -17364,6 +17698,49 @@ bool FMaterialLayersFunctionsRuntimeData::SerializeFromMismatchedTag(const FProp
 	return false;
 }
 
+
+bool FMaterialLayersFunctionsRuntimeData::Serialize(FArchive& Ar)
+{
+#if ENABLE_MATERIAL_LAYER_PROTOTYPE
+	// TODO: We d like to do rely on the custom version of the branch or something:
+	// Ar.UsingCustomVersion(FRenderingObjectVersion::GUID);
+	// we are using our custom version indicator for now in the tree
+	if (!Ar.IsLoading())
+	{
+		Tree.EnableMaterialLayersFunctionsTree();
+	}
+#endif
+
+	// Use default serialization for most properties for
+	// archive configurations supported by UStruct::SerializeVersionedTaggedProperties 
+	if (Ar.IsLoading() || Ar.IsSaving() || Ar.IsCountingMemory() || Ar.IsObjectReferenceCollector())
+	{
+		StaticStruct()->SerializeTaggedProperties(Ar, (uint8*)this, StaticStruct(), nullptr);
+	}
+
+#if ENABLE_MATERIAL_LAYER_PROTOTYPE
+	// For PRE substrate tree layer version:
+	// TODO: As commented above, we could use:
+	// 	if (Ar.CustomVer(FRenderingObjectVersion::GUID) < FRenderingObjectVersion::AddedMaterialLayersSubstrateSupport)
+	if (!Tree.IsMaterialLayersFunctionsTreeEnabled())
+	{
+		// When loading, The tree need to be recreated from the list
+		// of layers and blends assuming this is a flat hierarchy of layers:
+		if (Ar.IsLoading())
+		{
+			Tree.Empty();
+			for (int32 l = 0; l < Layers.Num(); ++l)
+			{
+				Tree.AddNode({ l, l - 1 }, -1);
+			}
+		}
+	}
+#endif
+
+	return true;
+}
+
+
 #if WITH_EDITOR
 const FMaterialLayersFunctionsID FMaterialLayersFunctionsRuntimeData::GetID(const FMaterialLayersFunctionsEditorOnlyData& EditorOnly) const
 {
@@ -17485,14 +17862,40 @@ void FMaterialLayersFunctions::PostSerialize(const FArchive& Ar)
 				}
 			}
 		}
+
+
 	}
-#endif // WITH_EDITORONLY_DATA
+#endif // WITH_EDITORONLY_DATA	
 }
+
+
 #if WITH_EDITOR
+
+void FMaterialLayersFunctions::AddDefaultBackgroundLayer()
+{
+#if ENABLE_MATERIAL_LAYER_PROTOTYPE
+	// This call is only valid for the very first layer created in the Layer tree.
+	check(Layers.Num() == 0 && Tree.IsEmpty());
+	Tree.AddNode({ 0, -1 }, -1);
+#endif
+
+	// Default to a non-blended "background" layer
+	Layers.AddDefaulted();	
+	EditorOnly.LayerStates.Add(true);
+	FText LayerName = FText(LOCTEXT("Background", "Background"));
+	EditorOnly.LayerNames.Add(LayerName);
+	EditorOnly.RestrictToLayerRelatives.Add(false);
+	// Use a consistent Guid for the background layer
+	// Default constructor assigning different guids will break FStructUtils::AttemptToFindUninitializedScriptStructMembers
+	EditorOnly.LayerGuids.Add(BackgroundGuid);
+	EditorOnly.LayerLinkStates.Add(EMaterialLayerLinkState::NotFromParent);
+}
+
 int32 FMaterialLayersFunctions::AppendBlendedLayer()
 {
 	const int32 LayerIndex = Layers.AddDefaulted();
 	Blends.AddDefaulted();
+
 	EditorOnly.LayerStates.Add(true);
 	FText LayerName = FText::Format(LOCTEXT("LayerPrefix", "Layer {0}"), Layers.Num() - 1);
 	EditorOnly.LayerNames.Add(LayerName);
@@ -17500,6 +17903,7 @@ int32 FMaterialLayersFunctions::AppendBlendedLayer()
 	EditorOnly.RestrictToBlendRelatives.Add(false);
 	EditorOnly.LayerGuids.Add(FGuid::NewGuid());
 	EditorOnly.LayerLinkStates.Add(EMaterialLayerLinkState::NotFromParent);
+
 	return LayerIndex;
 }
 
@@ -17517,7 +17921,7 @@ int32 FMaterialLayersFunctions::AddLayerCopy(const FMaterialLayersFunctionsRunti
 	{
 		Blends.Add(Source.Blends[SourceLayerIndex - 1]);
 	}
-	
+
 	EditorOnly.LayerStates.Add(bVisible);
 	EditorOnly.LayerNames.Add(SourceEditorOnly.LayerNames[SourceLayerIndex]);
 	EditorOnly.RestrictToLayerRelatives.Add(SourceEditorOnly.RestrictToLayerRelatives[SourceLayerIndex]);
@@ -17527,6 +17931,7 @@ int32 FMaterialLayersFunctions::AddLayerCopy(const FMaterialLayersFunctionsRunti
 	}
 	EditorOnly.LayerGuids.Add(SourceEditorOnly.LayerGuids[SourceLayerIndex]);
 	EditorOnly.LayerLinkStates.Add(LinkState);
+
 	return LayerIndex;
 }
 
@@ -17540,7 +17945,7 @@ void FMaterialLayersFunctions::InsertLayerCopy(const FMaterialLayersFunctionsRun
 	check(LayerIndex > 0);
 	Layers.Insert(Source.Layers[SourceLayerIndex], LayerIndex);
 	Blends.Insert(Source.Blends[SourceLayerIndex - 1], LayerIndex - 1);
-	
+
 	EditorOnly.LayerStates.Insert(SourceEditorOnly.LayerStates[SourceLayerIndex], LayerIndex);
 	EditorOnly.LayerNames.Insert(SourceEditorOnly.LayerNames[SourceLayerIndex], LayerIndex);
 	EditorOnly.RestrictToLayerRelatives.Insert(SourceEditorOnly.RestrictToLayerRelatives[SourceLayerIndex], LayerIndex);
@@ -17701,7 +18106,7 @@ bool FMaterialLayersFunctions::ResolveParent(const FMaterialLayersFunctionsRunti
 	FMaterialLayersFunctionsEditorOnlyData& EditorOnly,
 	TArray<int32>& OutRemapLayerIndices)
 {
-	check(EditorOnly.LayerGuids.Num() == Runtime.Layers.Num());
+ 	check(EditorOnly.LayerGuids.Num() == Runtime.Layers.Num());
 	check(EditorOnly.LayerLinkStates.Num() == Runtime.Layers.Num());
 
 	FMaterialLayersFunctions ResolvedLayers;
@@ -17873,9 +18278,40 @@ bool FMaterialLayersFunctions::ResolveParent(const FMaterialLayersFunctionsRunti
 		}
 	}
 
+#if ENABLE_MATERIAL_LAYER_PROTOTYPE
+	struct Traverser
+	{
+		void Recursive(const FMaterialLayersFunctionsTree& Source, FLayerNodeId InSourceId, FMaterialLayersFunctionsTree& Dest, FLayerNodeId InDestParentId)
+		{
+
+			FLayerNode* SourceNode = Source.GetNode(InSourceId);
+			FLayerNodeId NewNodeId = FMaterialLayersFunctionsTree::InvalidId;
+			FLayerNodePayload Payload;
+			if (SourceNode)
+			{
+				Payload = Source.Payloads[InSourceId];
+				NewNodeId = Dest.AddNode(Payload, InDestParentId);
+			}
+
+			auto ChildrenIds = Source.GetChildrenIds(InSourceId);
+			for (auto SourceSubId : ChildrenIds)
+			{
+				Recursive(Source, SourceSubId, Dest, NewNodeId);
+			}
+
+		}
+	} ReplicateTree;
+	
+	ReplicateTree.Recursive(Runtime.Tree, -1, ResolvedLayers.Tree, -1);
+
+	// NOTE: For debug purpose at the moment, we want to monitor the end result tree data state.
+	MLFT::Log(&Runtime, __FUNCTION__);
+#endif
+
 	Runtime = MoveTemp(static_cast<FMaterialLayersFunctionsRuntimeData&>(ResolvedLayers));
 	EditorOnly = MoveTemp(ResolvedLayers.EditorOnly);
 
+
 	return bUpdatedLayerIndices;
 }
 
@@ -17891,6 +18327,115 @@ void FMaterialLayersFunctions::Validate(const FMaterialLayersFunctionsRuntimeDat
 	}
 }
 
+
+
+
+FLayerNodeId FMaterialLayersFunctions::GetNodeParent(FLayerNodeId InNodeId) const
+{
+	FLayerNode* Node = Tree.GetNode(InNodeId);
+	if (Node)
+		return Node->Parent;
+
+	return FMaterialLayersFunctionsTree::InvalidId;
+}
+
+FMaterialLayersFunctions::FLayerNodeIdArray FMaterialLayersFunctions::GetNodeChildren(FLayerNodeId InNodeId) const
+{
+	return Tree.GetChildrenIds(InNodeId);
+}
+
+FMaterialLayersFunctions::FLayerNodePayload FMaterialLayersFunctions::GetNodePayload(FLayerNodeId InNodeId) const
+{
+	FLayerNodePayload* Payload = Tree.GetPayload(InNodeId);
+	if (Payload)
+		return *Payload;
+
+	return {};
+}
+
+int32 FMaterialLayersFunctions::GetLayerFuncIndex(FLayerNodeId InNodeId) const
+{
+	return GetNodePayload(InNodeId).Layer;
+}
+
+int32 FMaterialLayersFunctions::GetBlendFuncIndex(FLayerNodeId InNodeId) const
+{
+	return GetNodePayload(InNodeId).Blend;
+}
+
+int32 FMaterialLayersFunctions::GetNodeDepth(FLayerNodeId InNodeId) const
+{
+	return Tree.GetDepth(InNodeId);
+}
+
+bool FMaterialLayersFunctions::CanAppendLayerNode(FLayerNodeId InParent) const
+{
+	// NOTE: Current design only allows for 2 levels depth in the tree
+	return  !(GetNodeDepth(InParent) > 1);
+}
+
+FLayerNodeId FMaterialLayersFunctions::AppendLayerNode(FLayerNodeId InParent, int32 InSiblingIndex)
+{
+#if ENABLE_MATERIAL_LAYER_PROTOTYPE
+	if (!CanAppendLayerNode(InParent))
+		return FMaterialLayersFunctionsTree::InvalidId;
+
+	const int32 LayerIndex = AppendBlendedLayer();
+
+	// InPayload
+	FLayerNodeId NewNodeId = Tree.AddNode({ LayerIndex, LayerIndex - 1 }, InParent, InSiblingIndex);
+
+	MLFT::Log(this, __FUNCTION__);
+
+	return NewNodeId;
+#else
+	return FMaterialLayersFunctionsTree::InvalidId;
+#endif
+}
+
+bool FMaterialLayersFunctions::CanRemoveLayerNode(FLayerNodeId InNodeId) const
+{
+	// NOTE: Current design prevent to delete 1st default node
+	if (Tree.IsValidId(InNodeId))
+	{
+		if (InNodeId != 0)
+		{
+			return true;
+		}
+	}
+	return false;
+}
+
+void FMaterialLayersFunctions::RemoveLayerNodeAt(FLayerNodeId InNodeId)
+{
+#if ENABLE_MATERIAL_LAYER_PROTOTYPE
+	if (!CanRemoveLayerNode(InNodeId))
+		return;
+
+	if (Tree.IsValidId(InNodeId))
+	{
+		FLayerNodePayload* Payload = Tree.GetPayload(InNodeId);
+		if (Payload && (Payload->Layer >= 0 && Payload->Layer < Layers.Num()))
+		{
+			RemoveBlendedLayerAt(Payload->Layer);
+
+			MLFT::RemovePayloadIndex(Tree, Payload->Layer);
+		}
+
+		Tree.RemoveNode(InNodeId);
+	}
+
+	MLFT::Log(this, __FUNCTION__);
+#endif
+}
+	
+
+
+void FMaterialLayersFunctions::MoveLayerNode(FLayerNodeId InNodeId, FLayerNodeId DstParentId, int32 InSiblingIndex)
+{
+
+}
+
 #endif // WITH_EDITOR
 
 ///////////////////////////////////////////////////////////////////////////////
diff --git a/Engine/Source/Runtime/Engine/Public/Materials/MaterialLayersFunctions.h b/Engine/Source/Runtime/Engine/Public/Materials/MaterialLayersFunctions.h
--- a/Engine/Source/Runtime/Engine/Public/Materials/MaterialLayersFunctions.h
+++ b/Engine/Source/Runtime/Engine/Public/Materials/MaterialLayersFunctions.h
@@ -4,6 +4,7 @@
 
 #include "CoreMinimal.h"
 #include "MaterialTypes.h"
+#include "MaterialExpression.h"
 
 #include "MaterialLayersFunctions.generated.h"
 
@@ -116,6 +117,235 @@ struct FMaterialLayersFunctionsEditorOnlyData
 #endif // WITH_EDITOR
 };
 
+
+
+USTRUCT()
+struct FMaterialLayersFunctionsTree
+{
+	GENERATED_BODY()
+
+public:
+	using FNodeId = int32;
+	using FPayloadId = int32;
+
+	static const FNodeId InvalidId = -1; // Invalid id is always <0
+
+	struct FNode
+	{
+		FNodeId 		Parent = InvalidId;
+		FNodeId 		NextSibling = InvalidId;
+		FNodeId 		ChildrenHead = InvalidId;
+		FNodeId			Spare = InvalidId;
+
+		FORCEINLINE bool operator==(const FNode& Other) const
+		{
+			if (Parent != Other.Parent || NextSibling != Other.NextSibling || ChildrenHead != Other.ChildrenHead)
+			{
+				return false;
+			}
+			return true;
+		}
+
+		FORCEINLINE bool operator!=(const FNode& Other) const
+		{
+			return !operator==(Other);
+		}
+
+		void Clear()
+		{
+			Parent = NextSibling = Spare = InvalidId;
+		}
+	};
+
+	struct FPayload
+	{
+		FPayloadId Layer = InvalidId;
+		FPayloadId Blend = InvalidId;
+
+		FORCEINLINE bool operator==(const FPayload& Other) const
+		{
+			if (Layer != Other.Layer || Blend != Other.Blend)
+			{
+				return false;
+			}
+			return true;
+		}
+
+		FORCEINLINE bool operator!=(const FPayload& Other) const
+		{
+			return !operator==(Other);
+		}
+	};
+
+	TArray<FNode>		Nodes;
+	TArray<FPayload>	Payloads;
+	FNodeId				Root = -2; // the index of the head of the top level siblings
+							   // and temporarly, capture the information if the LayersFunctionsTree is enabled in this asset
+							   // -2 means is is disabled
+
+	// Custom version system until we can rely on the archive custom version system
+	void			EnableMaterialLayersFunctionsTree() { Root = (Root == -2 ? -1 : Root); }
+	bool			IsMaterialLayersFunctionsTreeEnabled() const { return Root != -2; }
+
+	FMaterialLayersFunctionsTree() = default;
+	FMaterialLayersFunctionsTree(const FMaterialLayersFunctionsTree& Rhs)
+		:	Nodes(Rhs.Nodes),
+			Payloads(Rhs.Payloads),
+			Root(Rhs.Root)
+			 
+	{}
+	FMaterialLayersFunctionsTree(const FMaterialLayersFunctions& Rhs) = delete;
+
+	FMaterialLayersFunctionsTree& operator=(const FMaterialLayersFunctionsTree& Rhs)
+	{
+		Nodes = Rhs.Nodes;
+		Payloads = Rhs.Payloads;
+		Root = Rhs.Root;
+
+		return *this;
+	}
+
+	FORCEINLINE bool operator==(const FMaterialLayersFunctionsTree& Other) const
+	{
+		if (Nodes != Other.Nodes || Payloads != Other.Payloads || Root != Other.Root)
+		{
+			return false;
+		}
+		return true;
+	}
+
+	FORCEINLINE bool operator!=(const FMaterialLayersFunctionsTree& Other) const
+	{
+		return !operator==(Other);
+	}
+
+	void Empty()
+	{
+		Nodes.Empty();
+		Payloads.Empty();
+		Root = -1;
+	}
+
+	FORCEINLINE bool IsEmpty() const
+	{
+		return Nodes.IsEmpty();
+	}
+
+	bool Serialize( FArchive& Ar );
+
+	friend FArchive& operator<<(FArchive& Ar, FNode& T)
+	{
+		return Ar << T.Parent << T.NextSibling << T.ChildrenHead << T.Spare;
+	}
+	friend FArchive& operator<<(FArchive& Ar, FPayload& T)
+	{
+		return Ar << T.Layer << T.Blend;
+	}
+
+	/// Authoring the Tree
+
+	FORCEINLINE bool IsValidId(FNodeId Id) const
+	{
+		return (Id >= 0 && Id < Nodes.Num());
+	}
+
+	// Find and return the node at the specified Id or null if invalid 
+	FNode* GetNode(FNodeId InNodeId) const;
+
+	// Find and return the node payload at the specified Id or null if invalid 
+	FPayload* GetPayload(FNodeId InNodeId) const;
+
+	// Find and return the parent node at the specified Id or null if invalid 
+	FNode* GetParent(FNodeId InNodeId) const;
+
+	// Find the node which is the head of the list of siblings where InNodeID belongs
+	FNodeId GetSiblingHeadId(FNodeId InNodeId) const;
+
+	// Gather the list of children for this node
+	TArray<FNodeId> GetChildrenIds(FNodeId InNodeId) const;
+
+	// Return the depth of the node in the tree, the root is depth 0
+	// the first valid node in the tree under the root is depth 1
+	int32 GetDepth(FNodeId InNodeId) const;
+
+	// Add a new node with the specified Payload under the InPArent at the specified child index
+	// the new node id is returned
+	FNodeId AddNode(const FPayload& InPayload, FNodeId InParent, int32 InAtChildIndex = -1);
+
+	// Remove this node
+	// detach the node from the tree hierarchy as well as all the children recursively
+	// the nodes are all cleared in the array of nodes
+	// the payloads are NOT cleared and still hold on their value
+	// return the array of Id of the removed nodes
+	TArray<FNodeId> RemoveNode(FNodeId RemovedNodeId);
+
+	// Traverser Utility
+	using FNodeVisitor = void (*)(FNodeId InNodeId, int32 InDepth, int32 InSiblingNum, const FPayload& InPayload);
+	template <typename V>
+	FNodeId TraverseNode(V InVisitor, FNodeId InNodeId, int32 InDepth, int32 InSiblingNum) const
+	{
+		// First plan for the next child node and sibling we will be visiting
+		// default with the InNodeId is the root case
+		FNodeId NextChildId = Root;	
+		FNodeId NextSiblingId = InvalidId;
+
+		// Traverse this node if not the root tree
+		if (InNodeId != InvalidId)
+		{
+			if (!IsValidId(InNodeId))
+				return InvalidId;
+
+			const FNode& Node = Nodes[InNodeId];
+			const FPayload& Payload = Payloads[InNodeId];
+			InVisitor(InNodeId, InDepth, InSiblingNum, Payload);
+
+			// And since we have visited a valid node, lets visit the children next and the next sibling after
+			NextChildId = Node.ChildrenHead;
+			NextSiblingId = Node.NextSibling;
+		}
+
+		// Traverse children
+		{
+			InDepth++;
+			int32 SiblingNum = 0;
+			while (IsValidId(NextChildId))
+			{
+				NextChildId = TraverseNode(InVisitor, NextChildId, InDepth, SiblingNum);
+				SiblingNum++;
+			}
+			InDepth--;
+		}
+
+		return NextSiblingId;
+	}
+
+	template <typename V>
+	void Traverse(V InVisitor, FNodeId InRootNodeId = InvalidId) const
+	{
+		TraverseNode(InVisitor, InRootNodeId, 0, 0);
+	}
+
+	// Logging the tree data structure, useful for debugging
+	// generate a multiline string representing the tree nodes hierarchy with the payload info per node
+	FString Log(FString InTab) const;
+
+private:
+	FNodeId AllocNode(const FPayload& InPayload);
+
+};
+
+template<>
+struct TStructOpsTypeTraits<FMaterialLayersFunctionsTree>
+: public TStructOpsTypeTraitsBase2<FMaterialLayersFunctionsTree>
+{
+	enum
+	{
+		WithSerializer = true,
+		WithCopy = true,
+		WithIdenticalViaEquality = true
+	};
+};
+
 USTRUCT()
 struct FMaterialLayersFunctionsRuntimeData
 {
@@ -127,10 +357,17 @@ struct FMaterialLayersFunctionsRuntimeData
 	UPROPERTY(EditAnywhere, Category = MaterialLayers)
 	TArray<TObjectPtr<class UMaterialFunctionInterface>> Blends;
 
+	// ENABLE_MATERIAL_LAYER_PROTOTYPE : UProperty describing the tree data structure
+	UPROPERTY(EditAnywhere, Category = MaterialLayers)
+	FMaterialLayersFunctionsTree Tree;
+
 	FMaterialLayersFunctionsRuntimeData() = default;
 	FMaterialLayersFunctionsRuntimeData(const FMaterialLayersFunctionsRuntimeData& Rhs)
 		: Layers(Rhs.Layers)
 		, Blends(Rhs.Blends)
+#if ENABLE_MATERIAL_LAYER_PROTOTYPE
+		, Tree(Rhs.Tree)
+#endif
 	{}
 
 	FMaterialLayersFunctionsRuntimeData(const FMaterialLayersFunctions& Rhs) = delete;
@@ -139,6 +376,9 @@ struct FMaterialLayersFunctionsRuntimeData
 	{
 		Layers = Rhs.Layers;
 		Blends = Rhs.Blends;
+#if ENABLE_MATERIAL_LAYER_PROTOTYPE
+		Tree = Rhs.Tree;
+#endif
 		return *this;
 	}
 
@@ -150,11 +390,18 @@ struct FMaterialLayersFunctionsRuntimeData
 	{
 		Layers.Empty();
 		Blends.Empty();
+#if ENABLE_MATERIAL_LAYER_PROTOTYPE
+		Tree.Empty();
+#endif
 	}
 
 	FORCEINLINE bool operator==(const FMaterialLayersFunctionsRuntimeData& Other) const
 	{
-		if (Layers != Other.Layers || Blends != Other.Blends)
+		if (Layers != Other.Layers || Blends != Other.Blends
+#if ENABLE_MATERIAL_LAYER_PROTOTYPE
+			|| Tree != Other.Tree
+#endif
+		)
 		{
 			return false;
 		}
@@ -167,6 +414,7 @@ struct FMaterialLayersFunctionsRuntimeData
 	}
 
 	bool SerializeFromMismatchedTag(const FPropertyTag& Tag, FStructuredArchive::FSlot Slot);
+	bool Serialize(FArchive& Ar);
 
 #if WITH_EDITOR
 	const FMaterialLayersFunctionsID GetID(const FMaterialLayersFunctionsEditorOnlyData& EditorOnly) const;
@@ -177,6 +425,8 @@ private:
 	/** FMaterialLayersFunctionsRuntimeData can be deserialized from an FMaterialLayersFunctions property, will store the editor-only portion here */
 	TUniquePtr<FMaterialLayersFunctionsEditorOnlyData> LegacySerializedEditorOnlyData;
 
+
+
 	friend struct FStaticParameterSet;
 #endif // WITH_EDITORONLY_DATA
 };
@@ -185,6 +435,7 @@ template<>
 struct TStructOpsTypeTraits<FMaterialLayersFunctionsRuntimeData> : TStructOpsTypeTraitsBase2<FMaterialLayersFunctionsRuntimeData>
 {
 	enum { WithStructuredSerializeFromMismatchedTag = true };
+	enum { WithSerializer = true };
 };
 
 USTRUCT()
@@ -218,19 +469,7 @@ struct FMaterialLayersFunctions : public FMaterialLayersFunctionsRuntimeData
 	inline bool IsEmpty() const { return Layers.Num() == 0; }
 
 #if WITH_EDITOR
-	void AddDefaultBackgroundLayer()
-	{
-		// Default to a non-blended "background" layer
-		Layers.AddDefaulted();
-		EditorOnly.LayerStates.Add(true);
-		FText LayerName = FText(LOCTEXT("Background", "Background"));
-		EditorOnly.LayerNames.Add(LayerName);
-		EditorOnly.RestrictToLayerRelatives.Add(false);
-		// Use a consistent Guid for the background layer
-		// Default constructor assigning different guids will break FStructUtils::AttemptToFindUninitializedScriptStructMembers
-		EditorOnly.LayerGuids.Add(BackgroundGuid);
-		EditorOnly.LayerLinkStates.Add(EMaterialLayerLinkState::NotFromParent);
-	}
+	ENGINE_API void AddDefaultBackgroundLayer();
 
 	ENGINE_API int32 AppendBlendedLayer();
 
@@ -339,7 +578,7 @@ struct FMaterialLayersFunctions : public FMaterialLayersFunctionsRuntimeData
 	}
 
 	ENGINE_API void SerializeLegacy(FArchive& Ar);
-#endif // WITH_EDITOR
+	#endif // WITH_EDITOR
 
 	ENGINE_API void PostSerialize(const FArchive& Ar);
 
@@ -363,6 +602,30 @@ struct FMaterialLayersFunctions : public FMaterialLayersFunctionsRuntimeData
 		return !operator==(Other);
 	}
 
+	// Define the Tree authoring of the layers
+#if WITH_EDITOR
+	using FLayerNodeId = FMaterialLayersFunctionsTree::FNodeId;
+	using FLayerNodeIdArray = TArray<FLayerNodeId>;
+	using FLayerNodePayload = FMaterialLayersFunctionsTree::FPayload; 
+
+	ENGINE_API FLayerNodeId GetNodeParent(FLayerNodeId InNode) const;
+	ENGINE_API FLayerNodeIdArray GetNodeChildren(FLayerNodeId InNode) const;
+	ENGINE_API FLayerNodePayload GetNodePayload(FLayerNodeId InNodeId) const;
+	ENGINE_API int32 GetNodeDepth(FLayerNodeId InNodeId) const;
+
+	ENGINE_API int32 GetLayerFuncIndex(FLayerNodeId InNodeId) const;
+	ENGINE_API int32 GetBlendFuncIndex(FLayerNodeId InNodeId) const;
+
+	ENGINE_API bool	 CanAppendLayerNode(FLayerNodeId InParent) const;
+	ENGINE_API FLayerNodeId	AppendLayerNode(FLayerNodeId InParent, int32 InSiblingIndex = -1);
+
+	ENGINE_API bool	 CanRemoveLayerNode(FLayerNodeId InNodeId) const;
+	ENGINE_API void	 RemoveLayerNodeAt(FLayerNodeId InNodeId);
+
+	ENGINE_API void MoveLayerNode(FLayerNodeId InNodeId, FLayerNodeId DstParentId, int32 InSiblingIndex = -1);
+
+#endif // WITH_EDITOR
+
 private:
 	UPROPERTY()
 	TArray<bool> LayerStates_DEPRECATED;
