Missing function change

diff --git a/Engine/Source/Editor/MaterialEditor/Private/SMaterialLayersFunctionsTree.cpp b/Engine/Source/Editor/MaterialEditor/Private/SMaterialLayersFunctionsTree.cpp
--- a/Engine/Source/Editor/MaterialEditor/Private/SMaterialLayersFunctionsTree.cpp
+++ b/Engine/Source/Editor/MaterialEditor/Private/SMaterialLayersFunctionsTree.cpp
@@ -1826,7 +1826,7 @@ void SMaterialLayersFunctionsInstanceWrapper::Refresh()
 				.VAlign(VAlign_Center)
 				[
 #if ENABLE_MATERIAL_LAYER_PROTOTYPE
-				PropertyCustomizationHelpers::MakeAddButton(FSimpleDelegate::CreateSP(NestedTree.Get(), &SMaterialSubstrateTree::AddLayer))
+				PropertyCustomizationHelpers::MakeAddButton(FSimpleDelegate::CreateSP(NestedTree.Get(), &SMaterialSubstrateTree::AddRootNodeLayer))
 #else
 					PropertyCustomizationHelpers::MakeAddButton(FSimpleDelegate::CreateSP(NestedTree.Get(), &SMaterialLayersFunctionsInstanceTree::AddLayer))
 #endif
