Creating the  ViewModel as a datastructure of Sorted and UnsortedParamData that matches the LayersFunctionsTree

diff --git a/Engine/Source/Editor/MaterialEditor/Private/SMaterialSubstrateTree.cpp b/Engine/Source/Editor/MaterialEditor/Private/SMaterialSubstrateTree.cpp
--- a/Engine/Source/Editor/MaterialEditor/Private/SMaterialSubstrateTree.cpp
+++ b/Engine/Source/Editor/MaterialEditor/Private/SMaterialSubstrateTree.cpp
@@ -950,9 +950,11 @@ void SMaterialSubstrateTree::RecursiveCreateWidgets(FRecursiveCreateWidgetsConte
 		}
 
 		StackProperty->Children.Add(ChildProperty);
+		ShowSubParameters(ChildProperty);
+
 	}
 
-	if (Payload.Blend != -1)
+	if (!bIsBackgroundItem && Payload.Blend != -1)
 	{
 		TSharedRef<FSortedParamData> ChildProperty = MakeShared<FSortedParamData>();
 		ChildProperty->StackDataType = EStackDataType::Asset;
@@ -980,6 +982,7 @@ void SMaterialSubstrateTree::RecursiveCreateWidgets(FRecursiveCreateWidgetsConte
 		}
 
 		StackProperty->Children.Add(ChildProperty);
+		ShowSubParameters(ChildProperty);
 	}
 
 	InParentContainer.Add(StackProperty);
@@ -992,6 +995,8 @@ void SMaterialSubstrateTree::CreateGroupsWidget()
 	NonLayerProperties.Empty();
 	LayerProperties.Empty();
 	FunctionParameter = nullptr;
+	TSharedPtr<IPropertyHandle> FunctionParameterHandle;
+
 	FPropertyEditorModule& Module = FModuleManager::LoadModuleChecked<FPropertyEditorModule>("PropertyEditor");
 	if (!Generator.IsValid())
 	{
@@ -1015,6 +1020,7 @@ void SMaterialSubstrateTree::CreateGroupsWidget()
 		Generator->SetObjects(Objects);
 	}
 
+
 	TSharedPtr<IDetailTreeNode> ParameterGroups = FindParameterGroupsNode(Generator);
 	if (ParameterGroups.IsValid())
 	{
@@ -1041,6 +1047,7 @@ void SMaterialSubstrateTree::CreateGroupsWidget()
 
 				if (Cast<UDEditorMaterialLayersParameterValue>(Parameter))
 				{
+					FunctionParameterHandle = ChildHandle;
 					if (FunctionParameter == nullptr)
 					{
 						FunctionParameter = Parameter;
@@ -1050,44 +1057,6 @@ void SMaterialSubstrateTree::CreateGroupsWidget()
 					auto It = StructPtrs.CreateConstIterator();
 					FunctionInstance = reinterpret_cast<FMaterialLayersFunctions*>(*It);
 					FunctionInstanceHandle = ParameterValueProperty;
-
-
-					TSharedPtr<IPropertyHandle>	LayerHandle = ChildHandle->GetChildHandle("Layers").ToSharedRef();
-					TSharedPtr<IPropertyHandle> BlendHandle = ChildHandle->GetChildHandle("Blends").ToSharedRef();
-					uint32 NumLayerChildren;
-					LayerHandle->GetNumChildren(NumLayerChildren);
-					uint32 NumBlendChildren;
-					BlendHandle->GetNumChildren(NumBlendChildren);
-					if (MaterialEditorInstance->StoredLayerPreviews.Num() != NumLayerChildren)
-					{
-						MaterialEditorInstance->StoredLayerPreviews.Empty();
-						MaterialEditorInstance->StoredLayerPreviews.AddDefaulted(NumLayerChildren);
-					}
-					if (MaterialEditorInstance->StoredBlendPreviews.Num() != NumBlendChildren)
-					{
-						MaterialEditorInstance->StoredBlendPreviews.Empty();
-						MaterialEditorInstance->StoredBlendPreviews.AddDefaulted(NumBlendChildren);
-					}
-
-
-#ifdef ENABLE_MATERIAL_LAYER_PROTOTYPE					
-					{
-						TSharedPtr<IPropertyHandle>	TreeHandle = ChildHandle->GetChildHandle("Tree").ToSharedRef();
-						
-						FRecursiveCreateWidgetsContext Context {
-							.Parameter = Parameter,
-							.LayerHandle = LayerHandle,
-							.BlendHandle = BlendHandle
-						};
-
-						auto RootChildren = FunctionInstance->GetNodeChildren(-1);
-						for (int i = 0; i < RootChildren.Num(); ++i)
-						{
-							RecursiveCreateWidgets(&Context, RootChildren[i], LayerProperties, true, false);
-						}	
-					}
-#endif // ENABLE_MATERIAL_LAYER_PROTOTYPE
-
 				}
 				else
 				{
@@ -1125,15 +1094,45 @@ void SMaterialSubstrateTree::CreateGroupsWidget()
 		DeferredResults.Empty();
 		DeferredSearches.Empty();
 
-		for (int32 LayerIdx = 0; LayerIdx < LayerProperties.Num(); LayerIdx++)
+		// Create the hierarchy of Sorted items recursivelly following the LayerFunctions Tree
+#ifdef ENABLE_MATERIAL_LAYER_PROTOTYPE
+		if (FunctionParameterHandle)
 		{
-			for (int32 ChildIdx = 0; ChildIdx < LayerProperties[LayerIdx]->Children.Num(); ChildIdx++)
+			TSharedPtr<IPropertyHandle>	LayerHandle = FunctionParameterHandle->GetChildHandle("Layers").ToSharedRef();
+			TSharedPtr<IPropertyHandle> BlendHandle = FunctionParameterHandle->GetChildHandle("Blends").ToSharedRef();
+			uint32 NumLayerChildren;
+			LayerHandle->GetNumChildren(NumLayerChildren);
+			uint32 NumBlendChildren;
+			BlendHandle->GetNumChildren(NumBlendChildren);
+			if (MaterialEditorInstance->StoredLayerPreviews.Num() != NumLayerChildren)
 			{
-				ShowSubParameters(LayerProperties[LayerIdx]->Children[ChildIdx]);
+				MaterialEditorInstance->StoredLayerPreviews.Empty();
+				MaterialEditorInstance->StoredLayerPreviews.AddDefaulted(NumLayerChildren);
+			}
+			if (MaterialEditorInstance->StoredBlendPreviews.Num() != NumBlendChildren)
+			{
+				MaterialEditorInstance->StoredBlendPreviews.Empty();
+				MaterialEditorInstance->StoredBlendPreviews.AddDefaulted(NumBlendChildren);
+			}
+
+			// root 
+			auto StrongFunctionParameter = FunctionParameter.Pin();
+
+			FRecursiveCreateWidgetsContext Context{
+				.Parameter = StrongFunctionParameter.Get(),
+				.LayerHandle = LayerHandle,
+				.BlendHandle = BlendHandle
+			};
+
+			auto RootChildren = FunctionInstance->GetNodeChildren(-1);
+			for (int i = 0; i < RootChildren.Num(); ++i)
+			{
+				RecursiveCreateWidgets(&Context, RootChildren[i], LayerProperties, true, false);
 			}
 		}
+#endif // ENABLE_MATERIAL_LAYER_PROTOTYPE
 	}
-
+	
 	SetParentsExpansionState();
 }
 
