Add Nanite Vertex factory and Compute shaders to the platform stats in the material editor

- This way we can see the generated shader code for these shaders in the ShaderCode menu for materials.

diff --git a/Engine/Source/Editor/MaterialEditor/Private/MaterialStatsCommon.cpp b/Engine/Source/Editor/MaterialEditor/Private/MaterialStatsCommon.cpp
--- a/Engine/Source/Editor/MaterialEditor/Private/MaterialStatsCommon.cpp
+++ b/Engine/Source/Editor/MaterialEditor/Private/MaterialStatsCommon.cpp
@@ -195,6 +195,10 @@ FString FMaterialStatsUtils::RepresentativeShaderTypeToString(const ERepresentat
 			return TEXT("Skinned Cloth");
 		break;
 
+		case ERepresentativeShader::NaniteMesh:
+			return TEXT("Nanite Mesh");
+		break;
+
 		case ERepresentativeShader::UIDefaultFragmentShader:
 			return TEXT("UI Pixel Shader");
 		break;
@@ -289,6 +293,7 @@ void FMaterialStatsUtils::GetRepresentativeShaderTypesAndDescriptions(TMap<FName
 	static const FName FLocalVertexFactoryName = FLocalVertexFactory::StaticType.GetFName();
 	static const FName FGPUFactoryName = TEXT("TGPUSkinVertexFactoryDefault");
 	static const FName FClothVertexFactoryName = TEXT("TGPUSkinAPEXClothVertexFactoryDefault");
+	static const FName FNaniteVertexFactoryName = TEXT("FNaniteVertexFactory");
 
 	if (TargetMaterial->IsUIMaterial())
 	{
@@ -350,6 +355,16 @@ void FMaterialStatsUtils::GetRepresentativeShaderTypesAndDescriptions(TMap<FName
 				.Add(FRepresentativeShaderInfo(ERepresentativeShader::SkinnedCloth, TBasePassVSFNoLightMapPolicyName, TEXT("Base pass vertex shader")));
 		}
 
+		if (TargetMaterial->GetFeatureLevel() >= ERHIFeatureLevel::SM6)
+		{
+			if (TargetMaterial->IsUsedWithNanite())
+			{
+				static FName TBasePassCSFNoLightMapPolicyName = TEXT("TBasePassCSFNoLightMapPolicy");
+				ShaderTypeNamesAndDescriptions.FindOrAdd(FNaniteVertexFactoryName)
+					.Add(FRepresentativeShaderInfo(ERepresentativeShader::NaniteMesh, TBasePassCSFNoLightMapPolicyName, TEXT("Nanite Compute Shader")));
+			}
+		}
+
 		// Add the shader type with the most sampler usages so we can accurately report the worst case scenario.
 		// This is ad-hoc, and ideally we have a better way for finding this shader type in the future.
 		ShaderTypeNamesAndDescriptions.FindOrAdd(FLocalVertexFactoryName)
diff --git a/Engine/Source/Editor/MaterialEditor/Private/MaterialStatsGrid.cpp b/Engine/Source/Editor/MaterialEditor/Private/MaterialStatsGrid.cpp
--- a/Engine/Source/Editor/MaterialEditor/Private/MaterialStatsGrid.cpp
+++ b/Engine/Source/Editor/MaterialEditor/Private/MaterialStatsGrid.cpp
@@ -177,6 +177,10 @@ FStatsGridRow_Shaders::EShaderClass FStatsGridRow_Shaders::GetShaderClass(const
 		case ERepresentativeShader::UIInstancedVertexShader:
 			return EShaderClass::VertexShader;
 		break;
+
+		case ERepresentativeShader::NaniteMesh:
+			return EShaderClass::ComputeShader;
+		break;
 	}
 
 	return EShaderClass::VertexShader;
@@ -190,8 +194,21 @@ void FStatsGridRow_Shaders::CreateRow(TSharedPtr<FMaterialStats> StatsManager)
 	if (bIsHeaderRow)
 	{
 		EShaderClass ShaderClass = GetShaderClass(ShaderType);
-		FString HeaderContent = ShaderClass == EShaderClass::VertexShader ? TEXT("Vertex Shader") : TEXT("Pixel Shader");
 
+		FString HeaderContent;
+		if (ShaderClass == EShaderClass::VertexShader)
+		{
+			HeaderContent = TEXT("Vertex Shader");
+		}
+		else if (ShaderClass == EShaderClass::FragmentShader)
+		{
+			HeaderContent = TEXT("Pixel Shader");
+		}
+		else if (ShaderClass == EShaderClass::ComputeShader)
+		{
+			HeaderContent = TEXT("Compute Shader");
+		}
+		
 		HeaderCell = MakeShareable(new FGridCell_StaticString(HeaderContent, HeaderContent));
 		HeaderCell->SetContentBold(true);
 		HeaderCell->SetColor(FStyleColors::Foreground);
diff --git a/Engine/Source/Editor/MaterialEditor/Private/MaterialStatsGrid.h b/Engine/Source/Editor/MaterialEditor/Private/MaterialStatsGrid.h
--- a/Engine/Source/Editor/MaterialEditor/Private/MaterialStatsGrid.h
+++ b/Engine/Source/Editor/MaterialEditor/Private/MaterialStatsGrid.h
@@ -182,7 +182,8 @@ public:
 	enum class EShaderClass
 	{
 		VertexShader,
-		FragmentShader
+		FragmentShader,
+		ComputeShader
 	};
 
 private:
diff --git a/Engine/Source/Editor/MaterialEditor/Public/MaterialStatsCommon.h b/Engine/Source/Editor/MaterialEditor/Public/MaterialStatsCommon.h
--- a/Engine/Source/Editor/MaterialEditor/Public/MaterialStatsCommon.h
+++ b/Engine/Source/Editor/MaterialEditor/Public/MaterialStatsCommon.h
@@ -65,6 +65,8 @@ enum class ERepresentativeShader
 	UIInstancedVertexShader,
 	LastVertexShader = UIInstancedVertexShader,
 
+	NaniteMesh,
+
 	Num
 };
 
