fix a bug with UI when deleting a layer

diff --git a/Engine/Source/Editor/MaterialEditor/Private/MaterialEditorInstanceDetailCustomization.cpp b/Engine/Source/Editor/MaterialEditor/Private/MaterialEditorInstanceDetailCustomization.cpp
--- a/Engine/Source/Editor/MaterialEditor/Private/MaterialEditorInstanceDetailCustomization.cpp
+++ b/Engine/Source/Editor/MaterialEditor/Private/MaterialEditorInstanceDetailCustomization.cpp
@@ -125,20 +125,24 @@ void FMaterialInstanceParameterDetails::CustomizeDetails(IDetailLayoutBuilder& D
 								{
 									return Group.GroupName == GroupParamData->Group.GroupName;
 								});
-						
-						FEditorParameterGroup& ParameterGroup = GroupParamData->Group;
-						IDetailGroup& DetailGroup = GroupsCategory.AddGroup(ParameterGroup.GroupName, FText::FromName(ParameterGroup.GroupName), false, true);
-						TSharedPtr<IPropertyHandle> GroupPropertyHandle = ParameterGroupsProperty->GetChildHandle(GroupIdx);
-						
-						CreateSingleGroupWidget(ParameterGroup, GroupPropertyHandle, DetailGroup, GroupParamData->ParameterInfo.Index, true);
-						
-						FSimpleDelegate UpdateThumbnails = FSimpleDelegate::CreateLambda([=, this]()
+
+						if (GroupIdx >= 0) // Only if found valid group index
 						{
-							this->MaterialLayersFunctionsInstance->NestedTree->UpdateThumbnailMaterial(ChildAsset->ParameterInfo.Association, ChildAsset->ParameterInfo.Index);
-						});
-						GroupPropertyHandle->SetOnPropertyValueChanged(UpdateThumbnails);
-						GroupPropertyHandle->SetOnChildPropertyValueChanged(UpdateThumbnails);
 						
+							FEditorParameterGroup& ParameterGroup = GroupParamData->Group;
+							IDetailGroup& DetailGroup = GroupsCategory.AddGroup(ParameterGroup.GroupName, FText::FromName(ParameterGroup.GroupName), false, true);
+
+							TSharedPtr<IPropertyHandle> GroupPropertyHandle = ParameterGroupsProperty->GetChildHandle(GroupIdx);
+
+							CreateSingleGroupWidget(ParameterGroup, GroupPropertyHandle, DetailGroup, GroupParamData->ParameterInfo.Index, true);
+
+							FSimpleDelegate UpdateThumbnails = FSimpleDelegate::CreateLambda([=, this]()
+								{
+									this->MaterialLayersFunctionsInstance->NestedTree->UpdateThumbnailMaterial(ChildAsset->ParameterInfo.Association, ChildAsset->ParameterInfo.Index);
+								});
+							GroupPropertyHandle->SetOnPropertyValueChanged(UpdateThumbnails);
+							GroupPropertyHandle->SetOnChildPropertyValueChanged(UpdateThumbnails);
+						}
 					}
 				}
 			}
