Material Layer visual updates

diff --git a/Engine/Source/Editor/MaterialEditor/Private/MaterialEditorInstanceDetailCustomization.cpp b/Engine/Source/Editor/MaterialEditor/Private/MaterialEditorInstanceDetailCustomization.cpp
--- a/Engine/Source/Editor/MaterialEditor/Private/MaterialEditorInstanceDetailCustomization.cpp
+++ b/Engine/Source/Editor/MaterialEditor/Private/MaterialEditorInstanceDetailCustomization.cpp
@@ -69,6 +69,19 @@ FMaterialInstanceParameterDetails::FMaterialInstanceParameterDetails(UMaterialEd
 {
 }
 
+void FMaterialInstanceParameterDetails::CollectStackItemsRecursively(TSharedPtr<FSortedParamData> Item, TArray<TSharedPtr<FSortedParamData>>& OutGroupsContainer)
+{
+	for(TSharedPtr<FSortedParamData> Child : Item->Children)
+	{
+		if (Child->StackDataType == EStackDataType::Stack)
+		{
+			OutGroupsContainer.Add(Child);
+		}
+		
+		CollectStackItemsRecursively(Child, OutGroupsContainer);
+	}
+}
+
 TOptional<float> FMaterialInstanceParameterDetails::OnGetValue(TSharedRef<IPropertyHandle> PropertyHandle)
 {
 	float Value = 0.0f;
@@ -110,12 +123,26 @@ void FMaterialInstanceParameterDetails::CustomizeDetails(IDetailLayoutBuilder& D
 	if (MaterialLayersFunctionsInstance != nullptr && MaterialLayersFunctionsInstance->NestedTree->GetNumItemsSelected() > 0)
 	{
 		// for each selected FSortedParamData item (type stack)
-		for (TSharedPtr<FSortedParamData> SelectedItem : this->MaterialLayersFunctionsInstance->NestedTree->GetSelectedItems())
+		TSharedPtr<FSortedParamData> SelectedItem = MaterialLayersFunctionsInstance->NestedTree->GetSelectedItems().Last();
+		
+			// make sure we selected a stack item
+		check (SelectedItem->StackDataType == EStackDataType::Stack)
+
+		// we should now gather all sub-stack items to loop through all together
+		TArray<TSharedPtr<FSortedParamData>> StacksCollection;
+		CollectStackItemsRecursively(SelectedItem, StacksCollection);
+
+		if (StacksCollection.IsEmpty())
 		{
-			// we go through list of assets
-			for(TSharedPtr<FSortedParamData> ChildAsset : SelectedItem->Children)
+			// we add the selected item in this case, as this indicates an inner stack item was selected
+			StacksCollection.Add(SelectedItem);
+		}
+		
+		for (TSharedPtr<FSortedParamData> StackItem : StacksCollection)
+		{
+			// // we go through list of assets
+			for(TSharedPtr<FSortedParamData> ChildAsset : StackItem->Children)
 			{
-				// we look for parameter group children
 				for(TSharedPtr<FSortedParamData> GroupParamData : ChildAsset->Children)
 				{
 					if (GroupParamData->StackDataType == EStackDataType::Group)
@@ -125,21 +152,18 @@ void FMaterialInstanceParameterDetails::CustomizeDetails(IDetailLayoutBuilder& D
 								{
 									return Group.GroupName == GroupParamData->Group.GroupName;
 								});
-
-						if (GroupIdx >= 0) // Only if found valid group index
+						if (GroupIdx != INDEX_NONE)
 						{
-						
 							FEditorParameterGroup& ParameterGroup = GroupParamData->Group;
 							IDetailGroup& DetailGroup = GroupsCategory.AddGroup(ParameterGroup.GroupName, FText::FromName(ParameterGroup.GroupName), false, true);
-
 							TSharedPtr<IPropertyHandle> GroupPropertyHandle = ParameterGroupsProperty->GetChildHandle(GroupIdx);
-
+						
 							CreateSingleGroupWidget(ParameterGroup, GroupPropertyHandle, DetailGroup, GroupParamData->ParameterInfo.Index, true);
-
+						
 							FSimpleDelegate UpdateThumbnails = FSimpleDelegate::CreateLambda([=, this]()
-								{
-									this->MaterialLayersFunctionsInstance->NestedTree->UpdateThumbnailMaterial(ChildAsset->ParameterInfo.Association, ChildAsset->ParameterInfo.Index);
-								});
+							{
+								this->MaterialLayersFunctionsInstance->NestedTree->UpdateThumbnailMaterial(ChildAsset->ParameterInfo.Association, ChildAsset->ParameterInfo.Index);
+							});
 							GroupPropertyHandle->SetOnPropertyValueChanged(UpdateThumbnails);
 							GroupPropertyHandle->SetOnChildPropertyValueChanged(UpdateThumbnails);
 						}
@@ -147,6 +171,7 @@ void FMaterialInstanceParameterDetails::CustomizeDetails(IDetailLayoutBuilder& D
 				}
 			}
 		}
+		
 		DetailLayout.HideCategory("MaterialEditorInstanceConstant");
 		DetailLayout.HideProperty("Parent");
 		DetailLayout.HideProperty("PostProcessOverrides");
diff --git a/Engine/Source/Editor/MaterialEditor/Private/MaterialEditorInstanceDetailCustomization.h b/Engine/Source/Editor/MaterialEditor/Private/MaterialEditorInstanceDetailCustomization.h
--- a/Engine/Source/Editor/MaterialEditor/Private/MaterialEditorInstanceDetailCustomization.h
+++ b/Engine/Source/Editor/MaterialEditor/Private/MaterialEditorInstanceDetailCustomization.h
@@ -32,6 +32,7 @@ public:
 	/** Constructor */
 	FMaterialInstanceParameterDetails(UMaterialEditorInstanceConstant* MaterialInstance, SMaterialLayersFunctionsInstanceWrapper* MaterialLayersFunctionsInstance, FGetShowHiddenParameters InShowHiddenDelegate);
 
+	void CollectStackItemsRecursively(TSharedPtr<FSortedParamData> Item, TArray<TSharedPtr<FSortedParamData>>& OutGroupsContainer);
 	/** IDetailCustomization interface */
 	virtual void CustomizeDetails(IDetailLayoutBuilder& DetailLayout) override;
 
diff --git a/Engine/Source/Editor/MaterialEditor/Private/MaterialEditorStyle.cpp b/Engine/Source/Editor/MaterialEditor/Private/MaterialEditorStyle.cpp
--- a/Engine/Source/Editor/MaterialEditor/Private/MaterialEditorStyle.cpp
+++ b/Engine/Source/Editor/MaterialEditor/Private/MaterialEditorStyle.cpp
@@ -146,8 +146,8 @@ TSharedRef<FSlateStyleSet> FSubstrateMaterialEditorStyle::Create()
 	Style->Set("Icons.Stage.ChainUnlinked.Vertical", new IMAGE_BRUSH_SVG("Icons/EditorIcons/ChainUnlinked_Vertical", Icon24x24));
 	
 	Style->Set("ImageBorder", new FSlateRoundedBoxBrush(
-		FLinearColor::Transparent, 10.0f, 
-		FStyleColors::Panel.GetSpecifiedColor(), 2.0f));
+		FLinearColor::Transparent, 0.0f, 
+		FStyleColors::Panel.GetSpecifiedColor(), 6.0f));
 	
 	Style->Set("Border.SinglePixel", new BORDER_BRUSH(TEXT("Images/Borders/Border_SinglePixel"), FMargin(1.0f / 4.0f)));
 	Style->Set("Border.LeftTopRight", new BORDER_BRUSH(TEXT("Images/Borders/Border_LeftTopRight"), FMargin(1.0f / 4.0f, 1.0f / 2.0f)));
@@ -199,27 +199,21 @@ TSharedRef<FSlateStyleSet> FSubstrateMaterialEditorStyle::Create()
 void FSubstrateMaterialEditorStyle::SetupLayerViewStyles(const TSharedRef<FSlateStyleSet>& Style)
 {
 	Style->Set("LayerView.Background", new FSlateRoundedBoxBrush(
-		FStyleColors::Panel.GetSpecifiedColor()/*FLinearColor(0, 0, 0, 0.25f)*/, 6.0f,
-		FStyleColors::Header.GetSpecifiedColor()/*FLinearColor(1, 1, 1, 0.2f)*/, 1.0f));
+		FStyleColors::Recessed/*FLinearColor(0, 0, 0, 0.25f)*/, 6.0f,
+		/*FStyleColors::Header.GetSpecifiedColor()*/FStyleColors::Recessed, 0.0f));
 
 	Style->Set("LayerView.Details.Background", new FSlateRoundedBoxBrush(
-		FStyleColors::Recessed.GetSpecifiedColor()/*FLinearColor(0, 0, 0, 0.25f)*/, 6.0f,
-		FStyleColors::Header.GetSpecifiedColor()/*FLinearColor(1, 1, 1, 0.2f)*/, 1.0f));
-
-	/**
-	 * SListView and and FTableViewStyle have no support for adding padding between the background brush
-	 * and the SListView widget, so we are not using this style for the SDMSlotLayerView. Instead, we add a
-	 * SBorder around the SDMBLayerView and style that.
-	 */
+		/*FStyleColors::Recessed.GetSpecifiedColor()*/FLinearColor(FColor::FromHex(TEXT("#575757"))), 6.0f,
+		FStyleColors::Header.GetSpecifiedColor()/*FLinearColor(1, 1, 1, 0.2f)*/, 20.0f));
+
 	Style->Set("LayerView", FTableViewStyle()
 		.SetBackgroundBrush(*Style->GetBrush("LayerView.Background"))
 	);
-
+	
 	const float LayerViewItemCornerRadius = 10.0f;
-	const float LayerViewItemBorderWidth = 5.0f;
-	const float LayerViewItemBorderWidthSelected = 1.0f;
+	const float LayerViewItemBorderWidth = 1.0f;
 
-	const FLinearColor LayerViewItemFillColor = FLinearColor::Transparent;
+	const FLinearColor LayerViewItemFillColor = FLinearColor(FColor::FromHex(TEXT("#383838")));//FLinearColor::Transparent;
 	constexpr FLinearColor LayerViewItemBorderColor = FLinearColor(0, 0, 0, 1.f);
 
 	const FLinearColor LayerItemHoverFillColor = FStyleColors::Recessed.GetSpecifiedColor();
@@ -228,33 +222,30 @@ void FSubstrateMaterialEditorStyle::SetupLayerViewStyles(const TSharedRef<FSlate
 	const FLinearColor LayerItemSelectFillColor = FStyleColors::Header.GetSpecifiedColor();
 	const FLinearColor LayerItemSelectBorderColor = ReplaceColorAlpha(FStyleColors::Select.GetSpecifiedColor(), 0.9f);
 
-	//constexpr FLinearColor LayerItemSelectHoverFillColor = FLinearColor(1, 1, 1, 0.25f);
-	//const FLinearColor LayerItemSelectHoverBorderColor = FStyleColors::Select.GetSpecifiedColor();
-
 	Style->Set("LayerView.Row.Item", new FSlateRoundedBoxBrush(
 		LayerViewItemFillColor, LayerViewItemCornerRadius,
 		LayerViewItemBorderColor, LayerViewItemBorderWidth));
 
 	Style->Set("LayerView.Row.Hovered", new FSlateRoundedBoxBrush(
 		LayerItemHoverFillColor, LayerViewItemCornerRadius,
-		LayerItemHoverBorderColor, LayerViewItemBorderWidthSelected));
+		LayerItemHoverBorderColor, LayerViewItemBorderWidth));
 
 	Style->Set("LayerView.Row.Selected", new FSlateRoundedBoxBrush(
 		LayerItemSelectFillColor, LayerViewItemCornerRadius,
-		LayerItemSelectBorderColor, LayerViewItemBorderWidthSelected));
+		LayerItemSelectBorderColor, LayerViewItemBorderWidth));
 
 	Style->Set("LayerView.Row.ActiveBrush", new FSlateRoundedBoxBrush(
 		LayerItemSelectFillColor, LayerViewItemCornerRadius,
-		LayerItemSelectBorderColor, LayerViewItemBorderWidthSelected));
+		LayerItemSelectBorderColor, LayerViewItemBorderWidth));
 	Style->Set("LayerView.Row.ActiveHoveredBrush", new FSlateRoundedBoxBrush(
 		LayerItemSelectFillColor, LayerViewItemCornerRadius,
-		LayerItemSelectBorderColor, LayerViewItemBorderWidthSelected));
+		LayerItemSelectBorderColor, LayerViewItemBorderWidth));
 	Style->Set("LayerView.Row.InactiveBrush", new FSlateRoundedBoxBrush(
 		LayerItemSelectFillColor, LayerViewItemCornerRadius,
-		LayerItemSelectBorderColor, LayerViewItemBorderWidthSelected));
+		LayerItemSelectBorderColor, LayerViewItemBorderWidth));
 	Style->Set("LayerView.Row.InactiveHoveredBrush", new FSlateRoundedBoxBrush(
 		LayerItemSelectFillColor, LayerViewItemCornerRadius,
-		LayerItemSelectBorderColor, LayerViewItemBorderWidthSelected));
+		LayerItemSelectBorderColor, LayerViewItemBorderWidth));
 
 	const float DropZoneMargin = 0.25f; 
 	Style->Set("LayerView.Row", FTableRowStyle()
@@ -269,13 +260,9 @@ void FSubstrateMaterialEditorStyle::SetupLayerViewStyles(const TSharedRef<FSlate
 		.SetActiveHoveredBrush(*Style->GetBrush(TEXT("LayerView.Row.ActiveHoveredBrush")))
 		.SetInactiveBrush(*Style->GetBrush(TEXT("LayerView.Row.InactiveBrush")))
 		.SetInactiveHoveredBrush(*Style->GetBrush(TEXT("LayerView.Row.InactiveHoveredBrush")))
-		//.SetSelectorFocusedBrush(BORDER_BRUSH("Common/Selector", FMargin(4.f / 16.f), Style->GetColor(TEXT("Color.Select.Hover"))))
-		//.SetSelectorFocusedBrush(BOX_BRUSH("Common/DropZoneIndicator_Onto", FMargin(4.f / 16.f), Style->GetColor(TEXT("Color.Select.Hover"))))
 		.SetDropIndicator_Onto(BOX_BRUSH("Common/DropZoneIndicator_Onto", FMargin(4.0f / 16.0f), Style->GetColor(TEXT("Color.Select.Hover"))))
-		//.SetDropIndicator_Above(BORDER_BRUSH("Common/DropZoneIndicatorDashed_Above", FMargin(DropZoneMargin, DropZoneMargin, 0.f, 0.f), Style->GetColor(TEXT("Color.Select.Hover"))))
-		.SetDropIndicator_Above(BOX_BRUSH("Common/VerticalBoxDropZoneIndicator_Above", FMargin(10.0f / 16.0f, 10.0f / 16.0f, 0, 0), Style->GetColor(TEXT("Color.Select.Hover"))))
-		// .SetDropIndicator_Below(BORDER_BRUSH("Common/DropZoneIndicatorDashed_Below", FMargin(DropZoneMargin, 0, 0, DropZoneMargin), Style->GetColor(TEXT("Color.Select.Hover"))))
-		 .SetDropIndicator_Below(BOX_BRUSH("Common/VerticalBoxDropZoneIndicator_Below",  FMargin(10.0f / 16.0f, 0, 0, 10.0f / 16.0f), Style->GetColor(TEXT("Color.Select.Hover"))))
+		.SetDropIndicator_Above(BORDER_BRUSH("Common/LayersDropZoneDashed_Above", FMargin(DropZoneMargin, DropZoneMargin, 0.f, 0.f), Style->GetColor(TEXT("Color.Select.Hover"))))
+		.SetDropIndicator_Below(BORDER_BRUSH("Common/LayersDropZoneDashed_Below", FMargin(DropZoneMargin, 0, 0, DropZoneMargin), Style->GetColor(TEXT("Color.Select.Hover"))))
 	);
 	Style->Set("LayerView.AddIcon", new IMAGE_BRUSH("Icons/EditorIcons/LayerAdd", Icon16x16));
 	Style->Set("LayerView.DuplicateIcon", new IMAGE_BRUSH("Icons/EditorIcons/Duplicate_40x", Icon40x40));
diff --git a/Engine/Source/Editor/MaterialEditor/Private/MaterialInstanceEditor.cpp b/Engine/Source/Editor/MaterialEditor/Private/MaterialInstanceEditor.cpp
--- a/Engine/Source/Editor/MaterialEditor/Private/MaterialInstanceEditor.cpp
+++ b/Engine/Source/Editor/MaterialEditor/Private/MaterialInstanceEditor.cpp
@@ -255,7 +255,12 @@ void FMaterialInstanceEditor::RegisterTabSpawners(const TSharedRef<class FTabMan
 	if (!bIsFunctionPreviewMaterial)
 	{
 		InTabManager->RegisterTabSpawner(LayerPropertiesTabId, FOnSpawnTab::CreateSP(this, &FMaterialInstanceEditor::SpawnTab_LayerProperties))
+#if ENABLE_MATERIAL_LAYER_PROTOTYPE		
+			.SetDisplayName(LOCTEXT("LayerPropertiesTab", "Material Layers"))
+#else
 			.SetDisplayName(LOCTEXT("LayerPropertiesTab", "Layer Parameters"))
+#endif
+		
 			.SetGroup(WorkspaceMenuCategoryRef)
 			.SetIcon(FSlateIcon(FAppStyle::GetAppStyleSetName(), "LevelEditor.Tabs.Layers"));
 	}
diff --git a/Engine/Source/Editor/MaterialEditor/Private/SMaterialLayersFunctionsTree.cpp b/Engine/Source/Editor/MaterialEditor/Private/SMaterialLayersFunctionsTree.cpp
--- a/Engine/Source/Editor/MaterialEditor/Private/SMaterialLayersFunctionsTree.cpp
+++ b/Engine/Source/Editor/MaterialEditor/Private/SMaterialLayersFunctionsTree.cpp
@@ -1786,6 +1786,16 @@ void SMaterialLayersFunctionsInstanceWrapper::Refresh()
 	FOnClicked 	OnChildButtonClicked = FOnClicked::CreateStatic(&FMaterialPropertyHelpers::OnClickedSaveNewMaterialInstance, ImplicitConv<UMaterialInterface*>(MaterialEditorInstance->SourceInstance), ImplicitConv<UObject*>(MaterialEditorInstance));
 	FOnClicked	OnSiblingButtonClicked = FOnClicked::CreateStatic(&FMaterialPropertyHelpers::OnClickedSaveNewMaterialInstance, MaterialEditorInstance->SourceInstance->Parent, ImplicitConv<UObject*>(MaterialEditorInstance));
 
+#if ENABLE_MATERIAL_LAYER_PROTOTYPE
+	const float ThumbnailSize = 64.0f;
+	TSharedPtr<SBox> ThumbnailBox;
+	UObject* ThumbnailObject = MaterialEditorInstance->SourceInstance.Get();
+	const TSharedPtr<FAssetThumbnail> AssetThumbnail = MakeShareable(new FAssetThumbnail(ThumbnailObject, ThumbnailSize, ThumbnailSize, NestedTree->GetTreeThumbnailPool()));
+	TSharedRef<SWidget> ThumbnailWidget = AssetThumbnail->MakeThumbnailWidget();
+	FText MaterialName = FText::FromName(MaterialEditorInstance->SourceInstance->GetFName());
+	
+#endif
+	
 	if (LayerParameter != nullptr)
 	{
 #if ENABLE_MATERIAL_LAYER_PROTOTYPE
@@ -1798,9 +1808,51 @@ void SMaterialLayersFunctionsInstanceWrapper::Refresh()
 			[
 				SNew(SVerticalBox)
 				+ SVerticalBox::Slot()
-				.Padding(0.0f)
+				.Padding(FMargin(0.0f, 0.0f, 0.0f, 10.0f))
 				.AutoHeight()
 				[
+#if ENABLE_MATERIAL_LAYER_PROTOTYPE
+					SAssignNew(HeaderBox, SHorizontalBox)
+					+ SHorizontalBox::Slot()
+					.AutoWidth()
+							.HAlign(HAlign_Center)
+							.VAlign(VAlign_Center)
+							.Padding(4.0f)
+							.MaxWidth(ThumbnailSize)
+							[
+								SAssignNew(ThumbnailBox, SBox)
+								.MaxDesiredWidth(ThumbnailSize)
+								.MaxDesiredWidth(ThumbnailSize)
+								.MaxDesiredHeight(ThumbnailSize)
+								.MinDesiredHeight(ThumbnailSize)
+								[
+									ThumbnailWidget
+								]
+							]
+
+					+ SHorizontalBox::Slot()
+					.VAlign(VAlign_Center)
+					.AutoWidth()
+					.Padding(5.0f)
+					[
+						SNew(SVerticalBox)
+						+ SVerticalBox::Slot()
+						.AutoHeight()
+						[
+							SNew(STextBlock)
+							.Text(MaterialName)
+							.TextStyle(FAppStyle::Get(), "LargeText")
+						]
+
+						+ SVerticalBox::Slot()
+						.AutoHeight()
+						[
+							SNew(STextBlock)
+							.Text(LOCTEXT("MaterialInstance", "Material Instance"))
+							.TextStyle(FAppStyle::Get(), "NormalText.Important")
+						]
+					]
+#else
 					SAssignNew(HeaderBox, SHorizontalBox)
 					+ SHorizontalBox::Slot()
 					.Padding(FMargin(4.0f, 0.0f))
@@ -1811,6 +1863,7 @@ void SMaterialLayersFunctionsInstanceWrapper::Refresh()
 						SNew(STextBlock)
 						.Text(LOCTEXT("MaterialLayers", "Material Layers"))
 					]
+#endif
 				]
 				+ SVerticalBox::Slot()
 				.Padding(FMargin(0.0f))
diff --git a/Engine/Source/Editor/MaterialEditor/Private/SMaterialSubstrateTree.cpp b/Engine/Source/Editor/MaterialEditor/Private/SMaterialSubstrateTree.cpp
--- a/Engine/Source/Editor/MaterialEditor/Private/SMaterialSubstrateTree.cpp
+++ b/Engine/Source/Editor/MaterialEditor/Private/SMaterialSubstrateTree.cpp
@@ -26,6 +26,7 @@
 #include "Widgets/Views/STreeView.h"
 #include "IDetailTreeNode.h"
 #include "AssetThumbnail.h"
+#include "ContentBrowserModule.h"
 #include "MaterialEditorInstanceDetailCustomization.h"
 #include "MaterialPropertyHelpers.h"
 #include "DetailWidgetRow.h"
@@ -34,22 +35,27 @@
 #include "Widgets/Images/SImage.h"
 #include "MaterialEditor/MaterialEditorPreviewParameters.h"
 #include "Widgets/Input/SButton.h"
+#include "Widgets/Layout/SSeparator.h"
 #include "Widgets/Text/SInlineEditableTextBlock.h"
 #include "Materials/MaterialFunctionInstance.h"
 #include "Materials/MaterialExpressionTextureSampleParameter.h"
 #include "Framework/Application/SlateApplication.h"
+#include "Framework/MultiBox/MultiBoxBuilder.h"
 #include "Widgets/Input/SEditableTextBox.h"
 #include "Curves/CurveLinearColor.h"
 #include "Curves/CurveLinearColorAtlas.h"
 #include "Subsystems/AssetEditorSubsystem.h"
 #include "Editor.h"
 #include "EditorFontGlyphs.h"
+#include "IContentBrowserSingleton.h"
 #include "SResetToDefaultPropertyEditor.h"
 #include "DragAndDrop/AssetDragDropOp.h"
 #include "ThumbnailRendering/ThumbnailManager.h"
 #include "Styling/AppStyle.h"
 #include "Styling/StyleColors.h"
 #include "MaterialEditorStyle.h"
+#include "DetailTreeNode.h"
+#include "Materials/MaterialFunctionMaterialLayerBlend.h"
 
 #define LOCTEXT_NAMESPACE "MaterialSubstrateTree"
 
@@ -177,6 +183,13 @@ FText SMaterialSubstrateTreeItem::GetLayerName() const
 	return Tree->FunctionInstance->GetLayerName(LayerFuncIndex);
 }
 
+FText SMaterialSubstrateTreeItem::GetLayerDesc() const
+{
+	const FString LayerDescText[] = { TEXT("Material Evaluation"), TEXT("Material Instance Attributes") };
+	
+	return FText::FromString(LayerDescText[GetIndentLevel()]);
+}
+
 void SMaterialSubstrateTreeItem::OnNameChanged(const FText& InText, ETextCommit::Type CommitInfo)
 {
 	const FScopedTransaction Transaction(LOCTEXT("RenamedSection", "Renamed layer and blend section"));
@@ -352,6 +365,21 @@ void  SMaterialSubstrateTreeItem::OnOverrideParameter(bool NewValue, TObjectPtr<
 
 void SMaterialSubstrateTreeItem::Construct(const FArguments& InArgs, const TSharedRef<STableViewBase>& InOwnerTableView)
 {
+	
+	FOnTableRowDragEnter LayerDragDelegate = FOnTableRowDragEnter::CreateSP(this, &SMaterialSubstrateTreeItem::OnLayerDragEnter);
+	FOnTableRowDragLeave LayerDragLeaveDelegate = FOnTableRowDragLeave::CreateSP(this, &SMaterialSubstrateTreeItem::OnLayerDragLeave);
+
+	STableRow< TSharedPtr<FSortedParamData> >::ConstructInternal(
+		STableRow< TSharedPtr<FSortedParamData> >::FArguments()
+		.Style(FSubstrateMaterialEditorStyle::Get(), "LayerView.Row")
+		.ShowSelection(true)
+		.OnCanAcceptDrop(this, &SMaterialSubstrateTreeItem::CanAcceptDrop)
+		.OnAcceptDrop(this, &SMaterialSubstrateTreeItem::OnLayerDrop)
+		.OnDragEnter(LayerDragDelegate)
+		.OnDragLeave(LayerDragLeaveDelegate),
+		InOwnerTableView
+	);
+	
 	StackParameterData = InArgs._StackParameterData;
 	MaterialEditorInstance = InArgs._MaterialEditorInstance;
 	Tree = InArgs._InTree;
@@ -360,23 +388,28 @@ void SMaterialSubstrateTreeItem::Construct(const FArguments& InArgs, const TShar
 	TSharedRef<SWidget> RightSideWidget = SNullWidget::NullWidget;
 	TSharedRef<SWidget> ResetWidget = SNullWidget::NullWidget;
 	FText NameOverride;
+	
+	TSharedPtr<SHorizontalBox> MainStack;
 	TSharedRef<SVerticalBox> WrapperWidget = SNew(SVerticalBox);
+
+	if (StackParameterData->StackDataType == EStackDataType::Stack)
+	{
+		WrapperWidget->AddSlot()
+		.Padding(10.0f)/*.AutoHeight()*/
+		[
+			SAssignNew(MainStack, SHorizontalBox)
+		];
+	}
+		
 	EHorizontalAlignment ValueAlignment = HAlign_Left;
 
 	bool bIsBackgroundItem = IsNodeKeyForBackgroundParameter(StackParameterData->NodeKey);
-	bool bCanAppendSubLayer = Tree->FunctionInstance->CanAppendLayerNode(StackParameterData->ParameterInfo.Index) && !bIsBackgroundItem;
-	bool bCanRemoveLayer = Tree->FunctionInstance->CanRemoveLayerNode(StackParameterData->ParameterInfo.Index) && !bIsBackgroundItem;
 	bool bCanReorderLayer = !(bIsBackgroundItem || StackParameterData->ParameterInfo.Index == 0);
-
+	bool bCanAppendSubLayer = Tree->FunctionInstance->CanAppendLayerNode(StackParameterData->ParameterInfo.Index) && !bIsBackgroundItem;
+		
 // STACK --------------------------------------------------
 	if (StackParameterData->StackDataType == EStackDataType::Stack)
 	{
-		// WrapperWidget->AddSlot()
-		// 	.Padding(5.0f)
-		// 	.AutoHeight()
-		// 	[
-		// 		SNullWidget::NullWidget
-		// 	];
 #if WITH_EDITOR
 		int32 LayerFuncIndex = Tree->FunctionInstance->GetLayerFuncIndex(StackParameterData->ParameterInfo.Index);
 		int32 BlendFuncIndex = Tree->FunctionInstance->GetBlendFuncIndex(StackParameterData->ParameterInfo.Index);
@@ -419,57 +452,113 @@ void SMaterialSubstrateTreeItem::Construct(const FArguments& InArgs, const TShar
 						PreviewAssociation = EMaterialParameterAssociation::LayerParameter;
 						Tree->UpdateThumbnailMaterial(PreviewAssociation, PreviewIndex);
 						ThumbnailIndex = PreviewIndex;
+					
+						HeaderRowWidget->AddSlot()
+							.AutoWidth()
+							.HAlign(HAlign_Center)
+							.VAlign(VAlign_Center)
+							.Padding(4.0f)
+							.MaxWidth(ThumbnailSize)
+							[
+								SAssignNew(ThumbnailBox, SBox)
+								.MaxDesiredWidth(ThumbnailSize)
+								.MaxDesiredWidth(ThumbnailSize)
+								.MaxDesiredHeight(ThumbnailSize)
+								.MinDesiredHeight(ThumbnailSize)
+									[
+										Tree->CreateThumbnailWidget(PreviewAssociation, ThumbnailIndex, ThumbnailSize)
+									]
+							];
 					}
+
+					// if blend asset, we set it up in the Wrapper Widget at the bottom of the VerticalBox
 					if (Cast<UMaterialFunctionInterface>(AssetObject)->GetMaterialFunctionUsage() == EMaterialFunctionUsage::MaterialLayerBlend)
 					{
-						PreviewIndex = BlendFuncIndex;
-						PreviewAssociation = EMaterialParameterAssociation::BlendParameter;
-						Tree->UpdateThumbnailMaterial(PreviewAssociation, PreviewIndex);
-						ThumbnailIndex = PreviewIndex;
+						
+						 WrapperWidget->AddSlot()
+						 .Padding(2.0f)
+						 .AutoHeight()
+						 [
+						 	SNew(SSeparator)
+						 	.Thickness(2.0f) // Set the thickness of the separator
+						 ];
+						IDetailTreeNode& Node = *AssetChild->ParameterNode;
+						TSharedPtr<IDetailPropertyRow> GeneratedRow = StaticCastSharedPtr<IDetailPropertyRow>(Node.GetRow());
+						IDetailPropertyRow& Row = *GeneratedRow.Get();
+						
+						TSharedRef<SWidget> AssetPickerWidget = SNew(SObjectPropertyEntryBox)
+							.ObjectPath_Lambda([=, this]()
+							{
+								UObject* AssetObject = nullptr;
+								AssetChild->ParameterHandle->GetValue(AssetObject);
+								return AssetObject->GetPathName();
+							})
+							.OnObjectChanged_Lambda([=, this](const FAssetData& InAssetData)
+							{
+								FSoftObjectPath ObjPath = InAssetData.GetSoftObjectPath();
+								AssetChild->ParameterHandle->SetValue(ObjPath.TryLoad());
+							})
+							.AllowedClass(UMaterialFunctionMaterialLayerBlend::StaticClass())
+							.AllowClear(true)
+							.DisplayUseSelected(false)
+							.DisplayBrowse(false);
+						
+						 WrapperWidget->AddSlot()
+						 	.Padding(5.0f)
+						 	.AutoHeight()
+						 [
+						 	SNew(SHorizontalBox)
+						 	// + SHorizontalBox::Slot()
+							// .AutoWidth()
+							// .VAlign(VAlign_Center)
+							// [
+							// 	PropertyCustomizationHelpers::MakeVisibilityButton(VisibilityClickedDelegate, FText(), IsEnabledAttribute)
+							// ]
+							
+						 	+ SHorizontalBox::Slot()
+						 	.Padding(5.0f)
+						 	.HAlign(HAlign_Left)
+						 	.VAlign(VAlign_Center)
+						 	[
+						 		SNew(STextBlock)
+						 		.Justification(ETextJustify::Center)
+						 		.Text(LOCTEXT("BlendLabel", "Blend"))
+						 	]
+						 	+ SHorizontalBox::Slot()
+						 	.Padding(5.0f)
+						 	[
+						 		AssetPickerWidget
+						 		// AssetChild->ParameterHandle->CreatePropertyValueWidget(false)	
+						 	]
+						 ];
 					}
 				}
-				HeaderRowWidget->AddSlot()
-					.AutoWidth()
-					.HAlign(HAlign_Center)
-					.VAlign(VAlign_Center)
-					.Padding(4.0f)
-					.MaxWidth(ThumbnailSize)
-					[
-						SAssignNew(ThumbnailBox, SBox)
-							[
-								Tree->CreateThumbnailWidget(PreviewAssociation, ThumbnailIndex, ThumbnailSize)
-							]
-					];
-				ThumbnailBox->SetMaxDesiredHeight(ThumbnailSize);
-				ThumbnailBox->SetMinDesiredHeight(ThumbnailSize);
-				ThumbnailBox->SetMinDesiredWidth(ThumbnailSize);
-				ThumbnailBox->SetMaxDesiredWidth(ThumbnailSize);
 			}
 		}
 
-		
-
-		if (StackParameterData->ParameterInfo.Index != 0)
-		{
-			HeaderRowWidget->AddSlot()
-				.VAlign(VAlign_Center)
-				.AutoWidth()
-				.Padding(5.0f)
-				[
-					SNew(SEditableTextBox)
-					.Text(TAttribute<FText>::Create(TAttribute<FText>::FGetter::CreateSP(this, &SMaterialSubstrateTreeItem::GetLayerName)))
-					.OnTextCommitted(FOnTextCommitted::CreateSP(this, &SMaterialSubstrateTreeItem::OnNameChanged))
-				];
-		}
-		else
 		{
 			HeaderRowWidget->AddSlot()
 				.VAlign(VAlign_Center)
 				.AutoWidth()
 				.Padding(5.0f)
 				[
-					SNew(STextBlock)
-					.Text(NameOverride)
+					SNew(SVerticalBox)
+					+ SVerticalBox::Slot()
+					.AutoHeight()
+					[
+						SNew(STextBlock)
+						.Text(NameOverride)
+						.TextStyle(FSubstrateMaterialEditorStyle::Get(), "LayerView.Row.HeaderText")
+					]
+
+					+ SVerticalBox::Slot()
+					.AutoHeight()
+					[
+						SNew(STextBlock)
+						.Text(TAttribute<FText>::Create(TAttribute<FText>::FGetter::CreateSP(this, &SMaterialSubstrateTreeItem::GetLayerDesc)))
+						.TextStyle(FSubstrateMaterialEditorStyle::Get(), "LayerView.Row.HeaderText.Small")
+					]
+				
 				];
 		}
 
@@ -494,29 +583,6 @@ void SMaterialSubstrateTreeItem::Construct(const FArguments& InArgs, const TShar
 				.Visibility(this, &SMaterialSubstrateTreeItem::GetUnlinkLayerVisibility)
 			];
 
-		// Can only add sub layer up to a certain level
-		if (bCanAppendSubLayer)
-		{
-			HeaderRowWidget->AddSlot()
-				.HAlign(HAlign_Left)
-				.AutoWidth()
-				.VAlign(VAlign_Center)
-				[
-					PropertyCustomizationHelpers::MakeAddButton(FSimpleDelegate::CreateSP(InArgs._InTree, &SMaterialSubstrateTree::AddNodeLayer, StackParameterData->ParameterInfo.Index))
-				];
-		}
-
-		// Can only remove layers that aren't the base layer.
-		if (bCanRemoveLayer)
-		{
-			HeaderRowWidget->AddSlot()
-				.AutoWidth()
-				.VAlign(VAlign_Center)
-				.Padding(0.0f, 0.0f, 5.0f, 0.0f)
-				[
-					PropertyCustomizationHelpers::MakeClearButton(FSimpleDelegate::CreateSP(InArgs._InTree, &SMaterialSubstrateTree::RemoveNodeLayer, StackParameterData->ParameterInfo.Index))
-				];
-		}
 		LeftSideWidget = HeaderRowWidget;
 	}
 // END STACK
@@ -525,16 +591,16 @@ void SMaterialSubstrateTreeItem::Construct(const FArguments& InArgs, const TShar
 // FINAL WRAPPER
 	if (StackParameterData->StackDataType == EStackDataType::Stack)
 	{
-		TSharedPtr<SHorizontalBox> FinalStack;
-		WrapperWidget->AddSlot()
-			.Padding(15.0f)/*.AutoHeight()*/
+		MainStack->AddSlot()
+			.Padding(FMargin(2.0f))
+			.VAlign(VAlign_Center)
 			[
-				SAssignNew(FinalStack, SHorizontalBox)
+				LeftSideWidget
 			];
-
+		
 		if (bCanReorderLayer)
 		{
-			FinalStack->AddSlot()
+			MainStack->AddSlot()
 				.HAlign(HAlign_Center)
 				.VAlign(VAlign_Center)
 				.Padding(2.5f, 0)
@@ -546,7 +612,7 @@ void SMaterialSubstrateTreeItem::Construct(const FArguments& InArgs, const TShar
 
 		if (bCanAppendSubLayer)
 		{
-			FinalStack->AddSlot()
+			MainStack->AddSlot()
 				.AutoWidth()
 				.VAlign(VAlign_Center)
 				.Padding(FMargin(2.0f))
@@ -554,54 +620,72 @@ void SMaterialSubstrateTreeItem::Construct(const FArguments& InArgs, const TShar
 					SNew(SExpanderArrow, SharedThis(this))
 				];
 		}
-
-		FinalStack->AddSlot()
-			.Padding(FMargin(2.0f))
-			.VAlign(VAlign_Center)
-			[
-				LeftSideWidget
-			];
 	}
 
 	this->ChildSlot
 		[
 			WrapperWidget
 		];
-
-	FOnTableRowDragEnter LayerDragDelegate = FOnTableRowDragEnter::CreateSP(this, &SMaterialSubstrateTreeItem::OnLayerDragEnter);
-	FOnTableRowDragLeave LayerDragLeaveDelegate = FOnTableRowDragLeave::CreateSP(this, &SMaterialSubstrateTreeItem::OnLayerDragLeave);
-
-	STableRow< TSharedPtr<FSortedParamData> >::ConstructInternal(
-		STableRow< TSharedPtr<FSortedParamData> >::FArguments()
-		.Padding(20.0f)
-		.OnPaintDropIndicator(this, &SMaterialSubstrateTreeItem::OnLayerItemPaintDropIndicator)
-		.Style(FSubstrateMaterialEditorStyle::Get(), "LayerView.Row")
-		.ShowSelection(true)
-		.OnCanAcceptDrop(this, &SMaterialSubstrateTreeItem::CanAcceptDrop)
-		.OnAcceptDrop(this, &SMaterialSubstrateTreeItem::OnLayerDrop)
-		.OnDragEnter(LayerDragDelegate)
-		.OnDragLeave(LayerDragLeaveDelegate),
-		InOwnerTableView
-	);
-
+	
+	this->SetDesiredSizeScale(FVector2d(1.0f, 1.2f));
 }
-int32 SMaterialSubstrateTreeItem::OnLayerItemPaintDropIndicator(EItemDropZone InItemDropZone, const FPaintArgs& Args, const FGeometry& AllottedGeometry, const FSlateRect& MyCullingRect, FSlateWindowElementList& OutDrawElements, int32 LayerId, const FWidgetStyle& InWidgetStyle, bool bParentEnabled) const
+
+int32 SMaterialSubstrateTreeItem::OnPaint(const FPaintArgs& Args, const FGeometry& AllottedGeometry, const FSlateRect& MyCullingRect, FSlateWindowElementList& OutDrawElements, int32 LayerId, const FWidgetStyle& InWidgetStyle, bool bParentEnabled) const
 {
-	const FSlateBrush* DropIndicatorBrush = GetDropIndicatorBrush(InItemDropZone);
-	static float OffsetX = 10.0f;
-	FVector2D Offset(OffsetX * GetIndentLevel(), 0.0f);
+	const FSlateBrush* BackgroundBrushResource = Tree->BackgroundBrush.Get();
+	const int32 IndentLevel = GetIndentLevel();
+	// calculate children + self size
+	int NumStackChildren = 0;
+	for(TSharedPtr<FSortedParamData> Child : this->StackParameterData->Children)
+	{
+		if (Child->StackDataType == Stack)
+			NumStackChildren++;
+	}
+	float HeightAdjFactor = 1;
 
+	LayerId += IndentLevel;
+	if (IsItemExpanded())
+	{
+		HeightAdjFactor = NumStackChildren + 1;
+	}
+	
+	FVector2d OuterBorderSize = AllottedGeometry.GetLocalSize() * FVector2d(1, HeightAdjFactor);
+	static float OffsetX = 30.0f;
+	FVector2D Offset = FVector2D(OffsetX * (IndentLevel+1), 0.0f );
 	FSlateDrawElement::MakeBox(
-		OutDrawElements,
-		LayerId++,
-		AllottedGeometry.ToPaintGeometry(FVector2D(AllottedGeometry.GetLocalSize() - Offset), FSlateLayoutTransform(Offset)),
-		DropIndicatorBrush,
-		ESlateDrawEffect::None,
-		DropIndicatorBrush->GetTint(InWidgetStyle) * InWidgetStyle.GetColorAndOpacityTint()
+			OutDrawElements,
+			++LayerId,
+			AllottedGeometry.ToPaintGeometry(OuterBorderSize - Offset, FSlateLayoutTransform(Offset/2)),
+			BackgroundBrushResource,
+			ESlateDrawEffect::None,
+			FStyleColors::Black.GetSpecifiedColor() * InWidgetStyle.GetColorAndOpacityTint() 
 	);
-
-	return LayerId;
+	
+	// only need to draw gray part for first level
+	if (IndentLevel <= 0)
+	{
+		Offset = FVector2d(OffsetX, -5.0f);
+		FSlateDrawElement::MakeBox(
+				OutDrawElements,
+				++LayerId,
+				AllottedGeometry.ToPaintGeometry(OuterBorderSize - Offset, FSlateLayoutTransform(Offset/2)),
+				BackgroundBrushResource,
+				ESlateDrawEffect::None,
+				FStyleColors::Header.GetSpecifiedColor() * InWidgetStyle.GetColorAndOpacityTint() 
+		);
+	}
+	
+	
+	OffsetX = 20.0f;
+	const float ReductionFactorX = 50.0f;
+	const float ReductionFactorY = 15.0f;
+	Offset = FVector2d(OffsetX * IndentLevel + ReductionFactorX, ReductionFactorY);
+	FGeometry BorderGeom = AllottedGeometry.MakeChild(AllottedGeometry.GetLocalSize() - Offset, FSlateLayoutTransform(Offset/2));
+	
+	return STableRow<TSharedPtr<FSortedParamData>>::OnPaint(Args, BorderGeom, MyCullingRect, OutDrawElements, LayerId+1, InWidgetStyle, bParentEnabled);
 }
+
+
 FString SMaterialSubstrateTreeItem::GetInstancePath(SMaterialSubstrateTree* InTree) const
 {
 	int32 LayerFuncIndex = Tree->FunctionInstance->GetLayerFuncIndex(StackParameterData->ParameterInfo.Index);
@@ -619,6 +703,36 @@ FString SMaterialSubstrateTreeItem::GetInstancePath(SMaterialSubstrateTree* InTr
 	return InstancePath;
 }
 
+TSharedPtr<SWidget> SMaterialSubstrateTree::CreateContextMenu()
+{
+	const bool bCloseAfterSelection = true;
+	
+	FMenuBuilder MenuBuilder(bCloseAfterSelection, TSharedPtr<FUICommandList>());
+
+	TArray<FSortedParamDataPtr> SelectedItemsArray = GetSelectedItems();
+	if (SelectedItemsArray.Num() > 0)
+	{
+		FSortedParamDataPtr StackParameterData = SelectedItemsArray[0];
+		
+		bool bIsBackgroundItem = IsNodeKeyForBackgroundParameter(StackParameterData->NodeKey);
+		bool bCanAppendSubLayer = FunctionInstance->CanAppendLayerNode(StackParameterData->ParameterInfo.Index) && !bIsBackgroundItem;
+		bool bCanRemoveLayer = FunctionInstance->CanRemoveLayerNode(StackParameterData->ParameterInfo.Index) && !bIsBackgroundItem;
+	
+		if (bCanAppendSubLayer)
+		{
+			const FSlateIcon PlusIcon(FAppStyle::GetAppStyleSetName(), "Plus");
+			MenuBuilder.AddMenuEntry(LOCTEXT("MaterialSubstrateTree", "Add New Layer"), FText(), PlusIcon, FUIAction(FExecuteAction::CreateSP(this, &SMaterialSubstrateTree::AddNodeLayer, StackParameterData->ParameterInfo.Index)));
+		}
+
+		if (bCanRemoveLayer)
+		{
+			const FSlateIcon MinusIcon(FAppStyle::GetAppStyleSetName(), "Icons.Minus");
+			MenuBuilder.AddMenuEntry(LOCTEXT("MaterialSubstrateTree", "Remove Layer"), FText(), MinusIcon, FUIAction(FExecuteAction::CreateSP(this, &SMaterialSubstrateTree::RemoveNodeLayer, StackParameterData->ParameterInfo.Index)));
+		}
+	}
+	return MenuBuilder.MakeWidget();
+}
+
 void SMaterialSubstrateTree::Construct(const FArguments& InArgs)
 {
 	ColumnSizeData.SetValueColumnWidth(0.5f);
@@ -654,6 +768,7 @@ void SMaterialSubstrateTree::Construct(const FArguments& InArgs)
 	STreeView<TSharedPtr<FSortedParamData>>::Construct(
 		STreeView::FArguments()
 		.TreeItemsSource(&LayerProperties)
+		.OnContextMenuOpening(this, &SMaterialSubstrateTree::CreateContextMenu)
 		.SelectionMode(ESelectionMode::Single)
 		.OnSelectionChanged(this, &SMaterialSubstrateTree::OnSelectionChangedMaterialSubstrateView)
 		.OnGenerateRow(this, &SMaterialSubstrateTree::OnGenerateRowMaterialLayersFunctionsTreeView)
diff --git a/Engine/Source/Editor/MaterialEditor/Private/SMaterialSubstrateTree.h b/Engine/Source/Editor/MaterialEditor/Private/SMaterialSubstrateTree.h
--- a/Engine/Source/Editor/MaterialEditor/Private/SMaterialSubstrateTree.h
+++ b/Engine/Source/Editor/MaterialEditor/Private/SMaterialSubstrateTree.h
@@ -47,6 +47,7 @@ public:
 	void OnSelectionChangedMaterialSubstrateView(TSharedPtr<FSortedParamData> InSelectedItem, ESelectInfo::Type SelectInfo);
 	void SetParentsExpansionState();
 
+	TSharedPtr<SWidget> CreateContextMenu();
 	void ShowHiddenValues(bool& bShowHiddenParameters) { bShowHiddenParameters = true; }
 	TWeakObjectPtr<class UDEditorParameterValue> FunctionParameter;
 	struct FMaterialLayersFunctions* FunctionInstance;
@@ -143,6 +144,7 @@ public:
 	void FilterClicked(const ECheckBoxState NewCheckedState, SMaterialSubstrateTree* InTree, TSharedPtr<FSortedParamData> InStackData);
 	ECheckBoxState GetFilterChecked(SMaterialSubstrateTree* InTree, TSharedPtr<FSortedParamData> InStackData) const;
 	FText GetLayerName() const;
+	FText GetLayerDesc() const;
 	void OnNameChanged(const FText& InText, ETextCommit::Type CommitInfo);
 	
 	FReply ToggleLayerVisibility();
@@ -181,9 +183,9 @@ public:
 	* @param InArgs   A declaration from which to construct the widget
 	*/
 	void Construct(const FArguments& InArgs, const TSharedRef<STableViewBase>& InOwnerTableView);
-	int32 OnLayerItemPaintDropIndicator(EItemDropZone InItemDropZone, const FPaintArgs& Args, const FGeometry& AllottedGeometry, const FSlateRect& MyCullingRect, FSlateWindowElementList& OutDrawElements, int32 LayerId, const FWidgetStyle& InWidgetStyle, bool bParentEnabled) const;
-
+	int32 OnPaint( const FPaintArgs& Args, const FGeometry& AllottedGeometry, const FSlateRect& MyCullingRect, FSlateWindowElementList& OutDrawElements, int32 LayerId, const FWidgetStyle& InWidgetStyle, bool bParentEnabled ) const override;
 	/** The node info to build the tree view row from. */
+	
 	FSortedParamDataPtr StackParameterData;
 
 	SMaterialSubstrateTree* Tree;
