Update tooltip for hiding unconnected pins in materials

- The name of the command and the tooltip now match the blueprint editors similar functionality

diff --git a/Engine/Source/Editor/MaterialEditor/Private/MaterialEditorActions.cpp b/Engine/Source/Editor/MaterialEditor/Private/MaterialEditorActions.cpp
--- a/Engine/Source/Editor/MaterialEditor/Private/MaterialEditorActions.cpp
+++ b/Engine/Source/Editor/MaterialEditor/Private/MaterialEditorActions.cpp
@@ -41,7 +41,7 @@ void FMaterialEditorCommands::RegisterCommands()
 
 	UI_COMMAND( CameraHome, "Home", "Goes home on the canvas.", EUserInterfaceActionType::Button, FInputChord() );
 	UI_COMMAND( CleanUnusedExpressions, "Clean Up", "Cleans up any unused Expressions.", EUserInterfaceActionType::Button, FInputChord() );
-	UI_COMMAND( ShowHideConnectors, "Hide Unused Connectors", "Hide unused connectors to clean up the graph view.", EUserInterfaceActionType::ToggleButton, FInputChord());
+	UI_COMMAND( ShowHideConnectors, "Hide Unconnected Pins", "Hides all pins with no connections", EUserInterfaceActionType::ToggleButton, FInputChord());
 	UI_COMMAND( ToggleLivePreview, "Preview Material", "Toggles real time update of the preview material.", EUserInterfaceActionType::ToggleButton, FInputChord());
 	UI_COMMAND( ToggleRealtimeExpressions, "Realtime Nodes", "Nodes impacted by time-based functions such as panners, etc. will update realtime.", EUserInterfaceActionType::ToggleButton, FInputChord() );
 	UI_COMMAND( AlwaysRefreshAllPreviews, "All Node Previews", "All node previews are updated upon any change to the graph.", EUserInterfaceActionType::ToggleButton, FInputChord() );
