[Material Editor Viewport] Add check before dereferencing Material editor weak ptr

The addition of a new Viewport Toolbar to the Material Editor (CL 34834818) caused crashes to happen when FMaterialEditorViewportClient::Draw dereferences MaterialEditorPtr without checking its validity.
While this CL does not solve the issue leading to the crash, it should prevent the crash from happening while looking for the proper solution.

diff --git a/Engine/Source/Editor/MaterialEditor/Private/SMaterialEditorViewport.cpp b/Engine/Source/Editor/MaterialEditor/Private/SMaterialEditorViewport.cpp
--- a/Engine/Source/Editor/MaterialEditor/Private/SMaterialEditorViewport.cpp
+++ b/Engine/Source/Editor/MaterialEditor/Private/SMaterialEditorViewport.cpp
@@ -126,7 +126,11 @@ void FMaterialEditorViewportClient::Tick(float DeltaSeconds)
 void FMaterialEditorViewportClient::Draw(FViewport* InViewport,FCanvas* Canvas)
 {
 	FEditorViewportClient::Draw(InViewport, Canvas);
-	MaterialEditorPtr.Pin()->DrawMessages(InViewport, Canvas);
+
+	if (MaterialEditorPtr.IsValid())
+	{
+		MaterialEditorPtr.Pin()->DrawMessages(InViewport, Canvas);
+	}
 }
 
 bool FMaterialEditorViewportClient::ShouldOrbitCamera() const
@@ -266,10 +270,13 @@ void SMaterialEditor3DPreviewViewport::Construct(const FArguments& InArgs)
 	PreviewMeshComponent = nullptr;
 	PostProcessVolumeActor = nullptr;
 
-	UMaterialInterface* Material = MaterialEditorPtr.Pin()->GetMaterialInterface();
-	if (Material)
+	if (MaterialEditorPtr.IsValid())
 	{
-		SetPreviewMaterial(Material);
+		UMaterialInterface* Material = MaterialEditorPtr.Pin()->GetMaterialInterface();
+		if (Material)
+		{
+			SetPreviewMaterial(Material);
+		}
 	}
 
 	SetPreviewAsset( GUnrealEd->GetThumbnailManager()->EditorSphere );
