Fix for reversed logic pointer dereference crash when quitting the material editor.

diff --git a/Engine/Source/Editor/MaterialEditor/Private/SMaterialEditorViewport.cpp b/Engine/Source/Editor/MaterialEditor/Private/SMaterialEditorViewport.cpp
--- a/Engine/Source/Editor/MaterialEditor/Private/SMaterialEditorViewport.cpp
+++ b/Engine/Source/Editor/MaterialEditor/Private/SMaterialEditorViewport.cpp
@@ -127,7 +127,7 @@ void FMaterialEditorViewportClient::Draw(FViewport* InViewport,FCanvas* Canvas)
 {
 	FEditorViewportClient::Draw(InViewport, Canvas);
 
-	if (!MaterialEditorPtr.IsValid())
+	if (MaterialEditorPtr.IsValid())
 	{
 		MaterialEditorPtr.Pin()->DrawMessages(InViewport, Canvas);
 	}
