Blend dropdown always visible now.
Also fixed material layer tree refresh when changing the dropdown values.

diff --git a/Engine/Source/Editor/MaterialEditor/Private/SMaterialSubstrateTree.cpp b/Engine/Source/Editor/MaterialEditor/Private/SMaterialSubstrateTree.cpp
--- a/Engine/Source/Editor/MaterialEditor/Private/SMaterialSubstrateTree.cpp
+++ b/Engine/Source/Editor/MaterialEditor/Private/SMaterialSubstrateTree.cpp
@@ -532,95 +532,96 @@ void SMaterialSubstrateTreeItem::Construct(const FArguments& InArgs, const TShar
 				int32 PreviewIndex = INDEX_NONE;
 				int32 ThumbnailIndex = INDEX_NONE;
 				EMaterialParameterAssociation PreviewAssociation = EMaterialParameterAssociation::GlobalParameter;
-				if (AssetObject)
+				
+				if (AssetChild->ParameterInfo.Association == LayerParameter)
 				{
-					if (Cast<UMaterialFunctionInterface>(AssetObject)->GetMaterialFunctionUsage() == EMaterialFunctionUsage::MaterialLayer)
-					{
-						PreviewIndex = LayerFuncIndex;
-						PreviewAssociation = EMaterialParameterAssociation::LayerParameter;
-						Tree->UpdateThumbnailMaterial(PreviewAssociation, PreviewIndex);
-						ThumbnailIndex = PreviewIndex;
-					
-						HeaderRowWidget->AddSlot()
-							.AutoWidth()
-							.HAlign(HAlign_Center)
-							.VAlign(VAlign_Center)
-							.Padding(4.0f)
-							.MaxWidth(ThumbnailSize)
-							[
-								SAssignNew(ThumbnailBox, SBox)
-								.MaxDesiredWidth(ThumbnailSize)
-								.MaxDesiredWidth(ThumbnailSize)
-								.MaxDesiredHeight(ThumbnailSize)
-								.MinDesiredHeight(ThumbnailSize)
-									[
-										Tree->CreateThumbnailWidget(PreviewAssociation, ThumbnailIndex, ThumbnailSize)
-									]
-							];
-					}
-
-					// if blend asset, we set it up in the Wrapper Widget at the bottom of the VerticalBox
-					if (Cast<UMaterialFunctionInterface>(AssetObject)->GetMaterialFunctionUsage() == EMaterialFunctionUsage::MaterialLayerBlend)
-					{
-						
-						 WrapperWidget->AddSlot()
-						 .Padding(2.0f)
-						 .AutoHeight()
-						 [
-						 	SNew(SSeparator)
-						 	.Thickness(2.0f) // Set the thickness of the separator
-						 ];
-						IDetailTreeNode& Node = *AssetChild->ParameterNode;
-						TSharedPtr<IDetailPropertyRow> GeneratedRow = StaticCastSharedPtr<IDetailPropertyRow>(Node.GetRow());
-						IDetailPropertyRow& Row = *GeneratedRow.Get();
-						
-						TSharedRef<SWidget> AssetPickerWidget = SNew(SObjectPropertyEntryBox)
-							.ObjectPath_Lambda([=, this]()
-							{
-								UObject* AssetObject = nullptr;
-								AssetChild->ParameterHandle->GetValue(AssetObject);
-								return AssetObject->GetPathName();
-							})
-							.OnObjectChanged_Lambda([=, this](const FAssetData& InAssetData)
-							{
-								FSoftObjectPath ObjPath = InAssetData.GetSoftObjectPath();
-								AssetChild->ParameterHandle->SetValue(ObjPath.TryLoad());
-							})
-							.AllowedClass(UMaterialFunctionMaterialLayerBlend::StaticClass())
-							.AllowClear(true)
-							.DisplayUseSelected(false)
-							.DisplayBrowse(false);
-						
-						 WrapperWidget->AddSlot()
-						 	.Padding(5.0f)
-						 	.AutoHeight()
-						 [
-						 	SNew(SHorizontalBox)
-						 	// + SHorizontalBox::Slot()
-							// .AutoWidth()
-							// .VAlign(VAlign_Center)
-							// [
-							// 	PropertyCustomizationHelpers::MakeVisibilityButton(VisibilityClickedDelegate, FText(), IsEnabledAttribute)
-							// ]
-							
-						 	+ SHorizontalBox::Slot()
-						 	.Padding(5.0f)
-						 	.HAlign(HAlign_Left)
-						 	.VAlign(VAlign_Center)
-						 	[
-						 		SNew(STextBlock)
-						 		.Justification(ETextJustify::Center)
-						 		.Text(LOCTEXT("BlendLabel", "Blend"))
-						 	]
-						 	+ SHorizontalBox::Slot()
-						 	.Padding(5.0f)
-						 	[
-						 		AssetPickerWidget
-						 		// AssetChild->ParameterHandle->CreatePropertyValueWidget(false)	
-						 	]
-						 ];
-					}
+					PreviewIndex = LayerFuncIndex;
+					PreviewAssociation = EMaterialParameterAssociation::LayerParameter;
+					Tree->UpdateThumbnailMaterial(PreviewAssociation, PreviewIndex);
+					ThumbnailIndex = PreviewIndex;
+				
+					HeaderRowWidget->AddSlot()
+						.AutoWidth()
+						.HAlign(HAlign_Center)
+						.VAlign(VAlign_Center)
+						.Padding(4.0f)
+						.MaxWidth(ThumbnailSize)
+						[
+							SAssignNew(ThumbnailBox, SBox)
+							.MaxDesiredWidth(ThumbnailSize)
+							.MinDesiredWidth(ThumbnailSize)
+							.MaxDesiredHeight(ThumbnailSize)
+							.MinDesiredHeight(ThumbnailSize)
+								[
+									Tree->CreateThumbnailWidget(PreviewAssociation, ThumbnailIndex, ThumbnailSize)
+								]
+						];
 				}
+				// if blend asset, we set it up in the Wrapper Widget at the bottom of the VerticalBox
+                else if (AssetChild->ParameterInfo.Association == BlendParameter)
+                {
+                		
+                	 WrapperWidget->AddSlot()
+                	 .Padding(2.0f)
+                	 .AutoHeight()
+                	 [
+                		SNew(SSeparator)
+                		.Thickness(2.0f) // Set the thickness of the separator
+                	 ];
+                	IDetailTreeNode& Node = *AssetChild->ParameterNode;
+                	TSharedPtr<IDetailPropertyRow> GeneratedRow = StaticCastSharedPtr<IDetailPropertyRow>(Node.GetRow());
+                	IDetailPropertyRow& Row = *GeneratedRow.Get();
+                	
+                	TSharedRef<SWidget> AssetPickerWidget = SNew(SObjectPropertyEntryBox)
+                		.ObjectPath_Lambda([=, this]()
+                		{
+                			UObject* AssetObject = nullptr;
+                			AssetChild->ParameterHandle->GetValue(AssetObject);
+                			return AssetObject->GetPathName();
+                		})
+                		.OnObjectChanged_Lambda([=, this](const FAssetData& InAssetData)
+                		{
+                			FSoftObjectPath ObjPath = InAssetData.GetSoftObjectPath();
+                			AssetChild->ParameterHandle->SetValue(ObjPath.TryLoad());
+
+                			Tree->FunctionInstanceHandle->NotifyPostChange(EPropertyChangeType::ValueSet);
+							Tree->CreateGroupsWidget();
+							Tree->RequestTreeRefresh();
+                		})
+                		.AllowedClass(UMaterialFunctionMaterialLayerBlend::StaticClass())
+                		.AllowClear(true)
+                		.DisplayUseSelected(false)
+                		.DisplayBrowse(false);
+                	
+                	 WrapperWidget->AddSlot()
+                		.Padding(5.0f)
+                		.AutoHeight()
+                	 [
+                		SNew(SHorizontalBox)
+                		// + SHorizontalBox::Slot()
+                		// .AutoWidth()
+                		// .VAlign(VAlign_Center)
+                		// [
+                		// 	PropertyCustomizationHelpers::MakeVisibilityButton(VisibilityClickedDelegate, FText(), IsEnabledAttribute)
+                		// ]
+                		
+                		+ SHorizontalBox::Slot()
+                		.Padding(5.0f)
+                		.HAlign(HAlign_Left)
+                		.VAlign(VAlign_Center)
+                		[
+                			SNew(STextBlock)
+                			.Justification(ETextJustify::Center)
+                			.Text(LOCTEXT("BlendLabel", "Blend"))
+                		]
+                		+ SHorizontalBox::Slot()
+                		.Padding(5.0f)
+                		[
+                			AssetPickerWidget
+                			// AssetChild->ParameterHandle->CreatePropertyValueWidget(false)	
+                		]
+                	 ];
+                }
 			}
 		}
 
