Warnings in Build Engine Localization - Text conflict from LOCTEXT macro for namespace in SMaterialSubstrateTree.cpp and MaterialInstanceEditor.cpp

diff --git a/Engine/Source/Editor/MaterialEditor/Private/MaterialInstanceEditor.cpp b/Engine/Source/Editor/MaterialEditor/Private/MaterialInstanceEditor.cpp
--- a/Engine/Source/Editor/MaterialEditor/Private/MaterialInstanceEditor.cpp
+++ b/Engine/Source/Editor/MaterialEditor/Private/MaterialInstanceEditor.cpp
@@ -256,7 +256,7 @@ void FMaterialInstanceEditor::RegisterTabSpawners(const TSharedRef<class FTabMan
 	{
 		InTabManager->RegisterTabSpawner(LayerPropertiesTabId, FOnSpawnTab::CreateSP(this, &FMaterialInstanceEditor::SpawnTab_LayerProperties))
 #if ENABLE_MATERIAL_LAYER_PROTOTYPE		
-			.SetDisplayName(LOCTEXT("LayerPropertiesTab", "Material Layers"))
+			.SetDisplayName(LOCTEXT("MaterialLayersTab", "Material Layers"))
 #else
 			.SetDisplayName(LOCTEXT("LayerPropertiesTab", "Layer Parameters"))
 #endif
diff --git a/Engine/Source/Editor/MaterialEditor/Private/SMaterialSubstrateTree.cpp b/Engine/Source/Editor/MaterialEditor/Private/SMaterialSubstrateTree.cpp
--- a/Engine/Source/Editor/MaterialEditor/Private/SMaterialSubstrateTree.cpp
+++ b/Engine/Source/Editor/MaterialEditor/Private/SMaterialSubstrateTree.cpp
@@ -809,13 +809,13 @@ TSharedPtr<SWidget> SMaterialSubstrateTree::CreateContextMenu()
 		if (bCanAppendSubLayer)
 		{
 			const FSlateIcon PlusIcon(FAppStyle::GetAppStyleSetName(), "Plus");
-			MenuBuilder.AddMenuEntry(LOCTEXT("MaterialSubstrateTree", "Add New Layer"), FText(), PlusIcon, FUIAction(FExecuteAction::CreateSP(this, &SMaterialSubstrateTree::AddNodeLayer, StackParameterData->ParameterInfo.Index)));
+			MenuBuilder.AddMenuEntry(LOCTEXT("AddNewLayer", "Add New Layer"), FText(), PlusIcon, FUIAction(FExecuteAction::CreateSP(this, &SMaterialSubstrateTree::AddNodeLayer, StackParameterData->ParameterInfo.Index)));
 		}
 
 		if (bCanRemoveLayer)
 		{
 			const FSlateIcon MinusIcon(FAppStyle::GetAppStyleSetName(), "Icons.Minus");
-			MenuBuilder.AddMenuEntry(LOCTEXT("MaterialSubstrateTree", "Remove Layer"), FText(), MinusIcon, FUIAction(FExecuteAction::CreateSP(this, &SMaterialSubstrateTree::RemoveNodeLayer, StackParameterData->ParameterInfo.Index)));
+			MenuBuilder.AddMenuEntry(LOCTEXT("RemoveLayer", "Remove Layer"), FText(), MinusIcon, FUIAction(FExecuteAction::CreateSP(this, &SMaterialSubstrateTree::RemoveNodeLayer, StackParameterData->ParameterInfo.Index)));
 		}
 	}
 	return MenuBuilder.MakeWidget();
