The grid in the material editor has been hardcoded disabled for ~12 years, but we left in the ability to toggle the grid via the UI.  This UI element didn't do anything and was useless.

- Remove the toggle in the UI for the grid (since it doesn't do anything).
- The Material Instance Editor still had the ability to toggle the grid, so this change effectively removes this behavior to keep the editors consistent.

diff --git a/Engine/Source/Editor/MaterialEditor/Private/MaterialEditor.cpp b/Engine/Source/Editor/MaterialEditor/Private/MaterialEditor.cpp
--- a/Engine/Source/Editor/MaterialEditor/Private/MaterialEditor.cpp
+++ b/Engine/Source/Editor/MaterialEditor/Private/MaterialEditor.cpp
@@ -2659,11 +2659,6 @@ void FMaterialEditor::LoadEditorSettings()
 
 	if ( PreviewViewport.IsValid() )
 	{
-		if (EditorOptions->bShowGrid)
-		{
-			PreviewViewport->TogglePreviewGrid();
-		}
-
 		if (EditorOptions->bRealtimeMaterialViewport && PreviewViewport->GetViewportClient())
 		{
 			PreviewViewport->GetViewportClient()->SetRealtime(true);
@@ -2690,7 +2685,6 @@ void FMaterialEditor::SaveEditorSettings()
 
 	if ( EditorOptions )
 	{
-		EditorOptions->bShowGrid					= PreviewViewport->IsTogglePreviewGridChecked();
 		EditorOptions->bRealtimeMaterialViewport	= PreviewViewport->IsRealtime();
 		EditorOptions->bHideUnusedConnectorsSetting	= IsOnHideConnectorsChecked();
 		EditorOptions->bAlwaysRefreshAllPreviews	= IsOnAlwaysRefreshAllPreviews();
diff --git a/Engine/Source/Editor/MaterialEditor/Private/MaterialEditorActions.cpp b/Engine/Source/Editor/MaterialEditor/Private/MaterialEditorActions.cpp
--- a/Engine/Source/Editor/MaterialEditor/Private/MaterialEditorActions.cpp
+++ b/Engine/Source/Editor/MaterialEditor/Private/MaterialEditorActions.cpp
@@ -36,7 +36,6 @@ void FMaterialEditorCommands::RegisterCommands()
 	UI_COMMAND( SetPlanePreview, "Plane", "Sets the preview mesh to a plane primitive.", EUserInterfaceActionType::ToggleButton, FInputChord() );
 	UI_COMMAND( SetCubePreview, "Cube", "Sets the preview mesh to a cube primitive.", EUserInterfaceActionType::ToggleButton, FInputChord() );
 	UI_COMMAND( SetPreviewMeshFromSelection, "Mesh", "Sets the preview mesh based on the current content browser selection.", EUserInterfaceActionType::ToggleButton, FInputChord() );
-	UI_COMMAND( TogglePreviewGrid, "Grid", "Toggles the preview pane's grid.", EUserInterfaceActionType::ToggleButton, FInputChord() );
 	UI_COMMAND( TogglePreviewBackground, "Background", "Toggles the preview pane's background.", EUserInterfaceActionType::ToggleButton, FInputChord() );
 
 	UI_COMMAND( CameraHome, "Home", "Goes home on the canvas.", EUserInterfaceActionType::Button, FInputChord() );
diff --git a/Engine/Source/Editor/MaterialEditor/Private/MaterialEditorViewportToolbarSections.cpp b/Engine/Source/Editor/MaterialEditor/Private/MaterialEditorViewportToolbarSections.cpp
--- a/Engine/Source/Editor/MaterialEditor/Private/MaterialEditorViewportToolbarSections.cpp
+++ b/Engine/Source/Editor/MaterialEditor/Private/MaterialEditorViewportToolbarSections.cpp
@@ -1,4 +1,4 @@
-﻿// Copyright Epic Games, Inc. All Rights Reserved.
+// Copyright Epic Games, Inc. All Rights Reserved.
 
 #include "MaterialEditorViewportToolbarSections.h"
 #include "EditorViewportCommands.h"
@@ -32,7 +32,6 @@ TSharedRef<SWidget> UE::MaterialEditor::CreateShowMenuWidget(const TSharedRef<SM
 
 		ShowMenuBuilder.AddMenuSeparator();
 
-		ShowMenuBuilder.AddMenuEntry(Commands.TogglePreviewGrid);
 		ShowMenuBuilder.AddMenuEntry(Commands.TogglePreviewBackground);
 	}
 
diff --git a/Engine/Source/Editor/MaterialEditor/Private/MaterialInstanceEditor.cpp b/Engine/Source/Editor/MaterialEditor/Private/MaterialInstanceEditor.cpp
--- a/Engine/Source/Editor/MaterialEditor/Private/MaterialInstanceEditor.cpp
+++ b/Engine/Source/Editor/MaterialEditor/Private/MaterialInstanceEditor.cpp
@@ -1626,7 +1626,6 @@ void FMaterialInstanceEditor::SaveAssetAs_Execute()
 
 void FMaterialInstanceEditor::SaveSettings()
 {
-	GConfig->SetBool(TEXT("MaterialInstanceEditor"), TEXT("bShowGrid"), PreviewVC->IsTogglePreviewGridChecked(), GEditorPerProjectIni);
 	GConfig->SetBool(TEXT("MaterialInstanceEditor"), TEXT("bDrawGrid"), PreviewVC->IsRealtime(), GEditorPerProjectIni);
 	GConfig->SetInt(TEXT("MaterialInstanceEditor"), TEXT("PrimType"), PreviewVC->PreviewPrimType, GEditorPerProjectIni);
 }
@@ -1634,18 +1633,12 @@ void FMaterialInstanceEditor::SaveSettings()
 void FMaterialInstanceEditor::LoadSettings()
 {
 	bool bRealtime=false;
-	bool bShowGrid=false;
 	int32 PrimType=static_cast<EThumbnailPrimType>( TPT_Sphere );
-	GConfig->GetBool(TEXT("MaterialInstanceEditor"), TEXT("bShowGrid"), bShowGrid, GEditorPerProjectIni);
 	GConfig->GetBool(TEXT("MaterialInstanceEditor"), TEXT("bDrawGrid"), bRealtime, GEditorPerProjectIni);
 	GConfig->GetInt(TEXT("MaterialInstanceEditor"), TEXT("PrimType"), PrimType, GEditorPerProjectIni);
 
 	if(PreviewVC.IsValid())
 	{
-		if ( bShowGrid )
-		{
-			PreviewVC->TogglePreviewGrid();
-		}
 		if ( bRealtime )
 		{
 			PreviewVC->OnToggleRealtime();
diff --git a/Engine/Source/Editor/MaterialEditor/Private/SMaterialEditorViewport.cpp b/Engine/Source/Editor/MaterialEditor/Private/SMaterialEditorViewport.cpp
--- a/Engine/Source/Editor/MaterialEditor/Private/SMaterialEditorViewport.cpp
+++ b/Engine/Source/Editor/MaterialEditor/Private/SMaterialEditorViewport.cpp
@@ -756,12 +756,6 @@ void SMaterialEditor3DPreviewViewport::BindCommands()
 		FCanExecuteAction(),
 		FIsActionChecked::CreateSP( this, &SMaterialEditor3DPreviewViewport::IsPreviewMeshFromSelectionChecked ) );
 
-	CommandList->MapAction(
-		Commands.TogglePreviewGrid,
-		FExecuteAction::CreateSP( this, &SMaterialEditor3DPreviewViewport::TogglePreviewGrid ),
-		FCanExecuteAction(),
-		FIsActionChecked::CreateSP( this, &SMaterialEditor3DPreviewViewport::IsTogglePreviewGridChecked ) );
-
 	CommandList->MapAction(
 		Commands.TogglePreviewBackground,
 		FExecuteAction::CreateSP( this, &SMaterialEditor3DPreviewViewport::TogglePreviewBackground ),
@@ -873,17 +867,6 @@ bool SMaterialEditor3DPreviewViewport::IsPreviewMeshFromSelectionChecked() const
 	return (PreviewPrimType == TPT_None && PreviewMeshComponent != nullptr);
 }
 
-void SMaterialEditor3DPreviewViewport::TogglePreviewGrid()
-{
-	EditorViewportClient->SetShowGrid();
-	RefreshViewport();
-}
-
-bool SMaterialEditor3DPreviewViewport::IsTogglePreviewGridChecked() const
-{
-	return EditorViewportClient->IsSetShowGridChecked();
-}
-
 void SMaterialEditor3DPreviewViewport::TogglePreviewBackground()
 {
 	UAssetViewerSettings* Settings = UAssetViewerSettings::Get();
diff --git a/Engine/Source/Editor/MaterialEditor/Private/SMaterialEditorViewport.h b/Engine/Source/Editor/MaterialEditor/Private/SMaterialEditorViewport.h
--- a/Engine/Source/Editor/MaterialEditor/Private/SMaterialEditorViewport.h
+++ b/Engine/Source/Editor/MaterialEditor/Private/SMaterialEditorViewport.h
@@ -81,8 +81,6 @@ public:
 	bool IsPreviewPrimitiveChecked(EThumbnailPrimType PrimType) const;
 	void OnSetPreviewMeshFromSelection();
 	bool IsPreviewMeshFromSelectionChecked() const;
-	void TogglePreviewGrid();
-	bool IsTogglePreviewGridChecked() const;
 	void TogglePreviewBackground();
 	bool IsTogglePreviewBackgroundChecked() const;
 	/** Call back for when the user changes preview scene settings in the UI */
diff --git a/Engine/Source/Editor/MaterialEditor/Public/MaterialEditorActions.h b/Engine/Source/Editor/MaterialEditor/Public/MaterialEditorActions.h
--- a/Engine/Source/Editor/MaterialEditor/Public/MaterialEditorActions.h
+++ b/Engine/Source/Editor/MaterialEditor/Public/MaterialEditorActions.h
@@ -74,9 +74,6 @@ public:
 	/** Sets the preview mesh to the current selection in the level editor */
 	TSharedPtr< FUICommandInfo > SetPreviewMeshFromSelection;
 	
-	/** Toggles the preview pane's grid */
-	TSharedPtr< FUICommandInfo > TogglePreviewGrid;
-	
 	/** Toggles the preview pane's background */
 	TSharedPtr< FUICommandInfo > TogglePreviewBackground;
 
diff --git a/Engine/Source/Editor/UnrealEd/Classes/Preferences/MaterialEditorOptions.h b/Engine/Source/Editor/UnrealEd/Classes/Preferences/MaterialEditorOptions.h
--- a/Engine/Source/Editor/UnrealEd/Classes/Preferences/MaterialEditorOptions.h
+++ b/Engine/Source/Editor/UnrealEd/Classes/Preferences/MaterialEditorOptions.h
@@ -18,10 +18,6 @@ class UMaterialEditorOptions : public UObject
 {
 	GENERATED_UCLASS_BODY()
 
-	/** If true, render grid the preview scene. */
-	UPROPERTY(EditAnywhere, config, Category=Options)
-	uint32 bShowGrid:1;
-
 	/** If true, don't render connectors that are not connected to anything. */
 	UPROPERTY(EditAnywhere, config, Category=Options)
 	uint32 bHideUnusedConnectorsSetting:1;
