Material Layer visual updates

diff --git a/Engine/Source/Editor/MaterialEditor/Private/MaterialInstanceEditor.cpp b/Engine/Source/Editor/MaterialEditor/Private/MaterialInstanceEditor.cpp
--- a/Engine/Source/Editor/MaterialEditor/Private/MaterialInstanceEditor.cpp
+++ b/Engine/Source/Editor/MaterialEditor/Private/MaterialInstanceEditor.cpp
@@ -1158,7 +1158,9 @@ TSharedRef<SDockTab> FMaterialInstanceEditor::SpawnTab_AssetBrowser(const FSpawn
 	Config.ThumbnailScale = 0.4f;
 	Config.InitialThumbnailSize = EThumbnailSize::Small;
 	Config.InitialAssetViewType = EAssetViewType::Tile;
-	
+	Config.OnAssetDoubleClicked = FOnAssetDoubleClicked::CreateSP(this, &FMaterialInstanceEditor::OnAssetDoubleClicked);
+	Config.OnGetAssetContextMenu = FOnGetAssetContextMenu::CreateSP(this, &FMaterialInstanceEditor::OnGetAssetContextMenu);
+	Config.bForceShowEngineContent = true;
 	TSharedRef<SDockTab> SpawnedTab = SNew(SDockTab)
 		.Label(LOCTEXT("AssetBrowserTab", "Asset Browser"))
 		[
@@ -1171,6 +1173,53 @@ TSharedRef<SDockTab> FMaterialInstanceEditor::SpawnTab_AssetBrowser(const FSpawn
 	return SpawnedTab;
 }
 
+TSharedPtr<SWidget> FMaterialInstanceEditor::OnGetAssetContextMenu(const TArray<FAssetData>& SelectedAssets) const
+{
+	if (SelectedAssets.Num() <= 0)
+	{
+		return nullptr;
+	}
+
+	UObject* SelectedAsset = SelectedAssets[0].GetAsset();
+	if (SelectedAsset == nullptr)
+	{
+		return nullptr;
+	}
+	
+	FMenuBuilder MenuBuilder(true, MakeShared<FUICommandList>());
+
+	MenuBuilder.BeginSection(TEXT("Asset"), LOCTEXT("AssetSectionLabel", "Asset"));
+	{
+		MenuBuilder.AddMenuEntry(
+			LOCTEXT("Browse", "Browse to Asset"),
+			LOCTEXT("BrowseTooltip", "Browses to the associated asset and selects it in the most recently used Content Browser (summoning one if necessary)"),
+			FSlateIcon(FAppStyle::GetAppStyleSetName(), "SystemWideCommands.FindInContentBrowser.Small"),
+			FUIAction(
+				FExecuteAction::CreateLambda([SelectedAsset] ()
+				{
+					if (SelectedAsset)
+					{
+						const TArray<FAssetData>& Assets = { SelectedAsset };
+						const FContentBrowserModule& ContentBrowserModule = FModuleManager::Get().LoadModuleChecked<FContentBrowserModule>("ContentBrowser");
+						ContentBrowserModule.Get().SyncBrowserToAssets(Assets);
+					}
+				}),
+				FCanExecuteAction::CreateLambda([] () { return true; })
+			)
+		);
+	}
+	MenuBuilder.EndSection();
+
+	return MenuBuilder.MakeWidget();
+}
+
+void FMaterialInstanceEditor::OnAssetDoubleClicked(const FAssetData& AssetData)
+{
+	if (UAssetEditorSubsystem* EditorSubsystem = GEditor->GetEditorSubsystem<UAssetEditorSubsystem>())
+	{
+		EditorSubsystem->OpenEditorForAsset(AssetData.ToSoftObjectPath());
+	}
+}
 void FMaterialInstanceEditor::AddToSpawnedToolPanels(const FName& TabIdentifier, const TSharedRef<SDockTab>& SpawnedTab)
 {
 	TWeakPtr<SDockTab>* TabSpot = SpawnedToolPanels.Find(TabIdentifier);
diff --git a/Engine/Source/Editor/MaterialEditor/Private/MaterialInstanceEditor.h b/Engine/Source/Editor/MaterialEditor/Private/MaterialInstanceEditor.h
--- a/Engine/Source/Editor/MaterialEditor/Private/MaterialInstanceEditor.h
+++ b/Engine/Source/Editor/MaterialEditor/Private/MaterialInstanceEditor.h
@@ -197,7 +197,9 @@ private:
 	
 	/** Spawns the asset browser tab */
 	TSharedRef<SDockTab> SpawnTab_AssetBrowser(const FSpawnTabArgs& Args);
-
+	void OnAssetDoubleClicked(const FAssetData& AssetData);
+	TSharedPtr<SWidget> OnGetAssetContextMenu(const TArray<FAssetData>& SelectedAssets) const;
+	
 	/**	Caches the specified tab for later retrieval */
 	void AddToSpawnedToolPanels( const FName& TabIdentifier, const TSharedRef<SDockTab>& SpawnedTab );
 
