Updated the substrate material tab to match the simplification mesage.

#rnx

diff --git a/Engine/Source/Editor/MaterialEditor/Private/SMaterialEditorSubstrateWidget.cpp b/Engine/Source/Editor/MaterialEditor/Private/SMaterialEditorSubstrateWidget.cpp
--- a/Engine/Source/Editor/MaterialEditor/Private/SMaterialEditorSubstrateWidget.cpp
+++ b/Engine/Source/Editor/MaterialEditor/Private/SMaterialEditorSubstrateWidget.cpp
@@ -35,6 +35,13 @@ void SMaterialEditorSubstrateWidget::Construct(const FArguments& InArgs, TWeakPt
 		.Padding(5.0f)
 		.ToolTipText(LOCTEXT("CheckBoxForceFullSimplificationToolTip", "This will force full simplification of the material."));
 
+	MaterialBudgetTextBlock = SNew(STextBlock)
+		.TextStyle(FAppStyle::Get(), "Log.Normal")
+		.ColorAndOpacity(FLinearColor::White)
+		.ShadowColorAndOpacity(FLinearColor::Black)
+		.ShadowOffset(FVector2D::UnitVector)
+		.Text(LOCTEXT("DescriptionTextBlock_Default", "Shader is compiling"));
+
 	DescriptionTextBlock = SNew(STextBlock)
 		.TextStyle(FAppStyle::Get(), "Log.Normal")
 		.ColorAndOpacity(FLinearColor::White)
@@ -112,7 +119,7 @@ void SMaterialEditorSubstrateWidget::Construct(const FArguments& InArgs, TWeakPt
 							.ColorAndOpacity(FLinearColor::White)
 							.ShadowColorAndOpacity(FLinearColor::Black)
 							.ShadowOffset(FVector2D::UnitVector)
-							.Text(LOCTEXT("MaterialSimplificationPreview", "Material simplification preview"))
+							.Text(LOCTEXT("MaterialSimplificationPreview", "Material Simplification Preview"))
 						]
 					]
 					+SVerticalBox::Slot()
@@ -273,7 +280,7 @@ void SMaterialEditorSubstrateWidget::Construct(const FArguments& InArgs, TWeakPt
 							.ColorAndOpacity(FLinearColor::White)
 							.ShadowColorAndOpacity(FLinearColor::Black)
 							.ShadowOffset(FVector2D::UnitVector)
-							.Text(LOCTEXT("MaterialTopologyPreview", "Material topology preview"))
+							.Text(LOCTEXT("MaterialTopologyPreview", "Material Topology Preview"))
 						]
 					]
 					+SVerticalBox::Slot()
@@ -314,7 +321,46 @@ void SMaterialEditorSubstrateWidget::Construct(const FArguments& InArgs, TWeakPt
 							.ColorAndOpacity(FLinearColor::White)
 							.ShadowColorAndOpacity(FLinearColor::Black)
 							.ShadowOffset(FVector2D::UnitVector)
-							.Text(LOCTEXT("MaterialAdvancedDetails", "Material advanced details"))
+							.Text(LOCTEXT("MaterialBudgetTextBlock", "Material Budget"))
+						]
+					]
+					+SVerticalBox::Slot()
+					.AutoHeight()
+					.Padding(0.0f, 0.0f, 0.0f, 0.0f)
+					[
+						SNew(SBorder)
+						.Padding(FMargin(5.0f, 5.0f, 5.0f, 5.0f))
+						[
+							SNew(SWrapBox)
+							.UseAllottedSize(true)
+							+SWrapBox::Slot()
+							.Padding(5.0f)
+							.HAlign(HAlign_Center)
+							.VAlign(VAlign_Center)
+							[
+								MaterialBudgetTextBlock->AsShared()
+							]
+						]
+					]
+
+
+					
+					+SVerticalBox::Slot()
+					.AutoHeight()
+					.Padding(0.0f, 15.0f, 0.0f, 0.0f)
+					[
+						SNew(SWrapBox)
+						.UseAllottedSize(true)
+						+ SWrapBox::Slot()
+						.Padding(0.0f)
+						.HAlign(HAlign_Left)
+						.VAlign(VAlign_Center)
+						[
+							SNew(STextBlock)
+							.ColorAndOpacity(FLinearColor::White)
+							.ShadowColorAndOpacity(FLinearColor::Black)
+							.ShadowOffset(FVector2D::UnitVector)
+							.Text(LOCTEXT("MaterialAdvancedDetails", "Material Advanced Details"))
 						]
 					]
 					+SVerticalBox::Slot()
@@ -410,27 +456,33 @@ void SMaterialEditorSubstrateWidget::Tick(const FGeometry& AllottedGeometry, con
 		if (MaterialResource)
 		{
 			FString MaterialDescription;
+			FString MaterialBudget;
 
+			bool bMaterialOutOfBudgetHasBeenSimplified = false;
 			FMaterialShaderMap* ShaderMap = MaterialResource->GetGameThreadShaderMap();
 			if (ShaderMap)
 			{
 				const FSubstrateMaterialCompilationOutput& CompilationOutput = ShaderMap->GetSubstrateMaterialCompilationOutput();
 				const uint32 FinalPixelByteCount = CompilationOutput.SubstrateUintPerPixel * sizeof(uint32);
+				const uint32 FinalPixelClosureCount = CompilationOutput.SubstrateClosureCount;
+				bMaterialOutOfBudgetHasBeenSimplified = CompilationOutput.bMaterialOutOfBudgetHasBeenSimplified > 0;
 
-				if (CompilationOutput.bMaterialOutOfBudgetHasBeenSimplified)
+				if (bMaterialOutOfBudgetHasBeenSimplified)
 				{
-					MaterialDescription += FString::Printf(TEXT("The material was OUT-OF-BUDGET so it has been fully simplified: Request bytes = %i / budget = %i  -  Request Closures = %i / budget = %i\r\n"),
+					MaterialBudget += FString::Printf(TEXT("The material was OUT-OF-BUDGET so it has been simplified: Requested bytes = %i / budget = %i  -  Requested Closures = %i / budget = %i\r\n"),
 						CompilationOutput.RequestedBytePerPixel, CompilationOutput.PlatformBytePerPixel,
 						CompilationOutput.RequestedClosurePerPixel, CompilationOutput.PlatformClosurePixel);
-					MaterialDescription += FString::Printf(TEXT("Final per pixel byte count   = %i\r\n"),
-						FinalPixelByteCount);
+					MaterialBudget += FString::Printf(TEXT("Final per pixel byte count   = %i\r\n"), FinalPixelByteCount);
+					MaterialBudget += FString::Printf(TEXT("Final per pixel closure count   = %i\r\n"), FinalPixelClosureCount);
 				}
 				else
 				{
-					MaterialDescription += FString::Printf(TEXT("Material per pixel byte count= %i / budget = %i\r\n"),
+					MaterialBudget += FString::Printf(TEXT("Material per pixel byte count    = %i / budget = %i\r\n"),
 						FinalPixelByteCount, CompilationOutput.PlatformBytePerPixel);
+					MaterialBudget += FString::Printf(TEXT("Material per pixel closure count = %i / budget = %i\r\n"),
+						FinalPixelClosureCount, CompilationOutput.PlatformClosurePixel);
 				}
-				MaterialDescription += FString::Printf(TEXT("BSDF Count	                  = %i\r\n"), CompilationOutput.SubstrateClosureCount);
+				MaterialDescription += FString::Printf(TEXT("Closures Count               = %i\r\n"), CompilationOutput.SubstrateClosureCount);
 				MaterialDescription += FString::Printf(TEXT("Local bases Count            = %i\r\n"), CompilationOutput.SharedLocalBasesCount);
 				MaterialDescription += FString::Printf(TEXT("Material complexity          = %s\r\n"), *SubstrateMaterialTypeToString(CompilationOutput.SubstrateMaterialType));
 				MaterialDescription += FString::Printf(TEXT("Root Node Is Thin            = %i\r\n"), CompilationOutput.bIsThin);
@@ -523,6 +575,16 @@ void SMaterialEditorSubstrateWidget::Tick(const FGeometry& AllottedGeometry, con
 				MaterialBox->SetContent(SNullWidget::NullWidget);
 			}
 
+			if (bMaterialOutOfBudgetHasBeenSimplified)
+			{
+				MaterialBudgetTextBlock->SetColorAndOpacity(FLinearColor(0.8, 0.5f, 0.07f, 1.0f)); // Same color as Graph.WarningText
+			}
+			else
+			{
+				MaterialBudgetTextBlock->SetColorAndOpacity(FLinearColor::White);
+			}
+			MaterialBudgetTextBlock->SetText(FText::FromString(MaterialBudget));
+
 			DescriptionTextBlock->SetText(FText::FromString(MaterialDescription));
 		}
 	}
diff --git a/Engine/Source/Editor/MaterialEditor/Private/SMaterialEditorSubstrateWidget.h b/Engine/Source/Editor/MaterialEditor/Private/SMaterialEditorSubstrateWidget.h
--- a/Engine/Source/Editor/MaterialEditor/Private/SMaterialEditorSubstrateWidget.h
+++ b/Engine/Source/Editor/MaterialEditor/Private/SMaterialEditorSubstrateWidget.h
@@ -47,6 +47,7 @@ private:
 
 	TSharedPtr<class SButton> ButtonApplyToPreview;
 
+	TSharedPtr<class STextBlock> MaterialBudgetTextBlock;
 	TSharedPtr<class STextBlock> DescriptionTextBlock;
 
 	TSharedPtr<class SBox> MaterialBox;
diff --git a/Engine/Source/Runtime/Engine/Private/Materials/HLSLMaterialTranslator.cpp b/Engine/Source/Runtime/Engine/Private/Materials/HLSLMaterialTranslator.cpp
--- a/Engine/Source/Runtime/Engine/Private/Materials/HLSLMaterialTranslator.cpp
+++ b/Engine/Source/Runtime/Engine/Private/Materials/HLSLMaterialTranslator.cpp
@@ -12633,6 +12633,7 @@ bool FHLSLMaterialTranslator::FSubstrateCompilationContext::SubstrateGenerateDer
 	const uint32 SubstrateClosurePerPixel = Compiler->SubstrateCompilationConfig.ClosuresPerPixelOverride > 0 ? Compiler->SubstrateCompilationConfig.ClosuresPerPixelOverride : Substrate::GetClosurePerPixel(ShaderPlatform);
 
 	bool bFirstLoop = true;
+	bool bRequestMaterialDetailsSet = false;
 	do 
 	{
 		if (!bFirstLoop && !SubstrateSimplificationStatus.bFullSimplificationStepHasBeenRun && !SubstrateSimplificationStatus.bMaterialFitsInMemoryBudget && SubstrateSimplificationStatus.OperatorSimplificationOrder.Num() > 0)
@@ -13192,11 +13193,12 @@ bool FHLSLMaterialTranslator::FSubstrateCompilationContext::SubstrateGenerateDer
 				Compiler->Errorf(TEXT("Material %s could not be simplified to fit in Substrate per pixel (asset: %s).\r\n"), *CompilerMaterial->GetDebugName(), *CompilerMaterial->GetAssetPath().ToString());
 				return false;
 			}
-			if (!SubstrateSimplificationStatus.bRunFullSimplification)
+			if (!bRequestMaterialDetailsSet)
 			{
 				// Record the original requested byte size before simplification, only for the first pass.
 				SubstrateSimplificationStatus.OriginalRequestedByteSize = SubstrateMaterialRequestedSizeByte;
 				SubstrateSimplificationStatus.OriginalRequestedClosureCount = SubstrateMaterialClosureCount;
+				bRequestMaterialDetailsSet = true;
 			}
 			SubstrateSimplificationStatus.bFullSimplificationStepHasBeenRun |= SubstrateSimplificationStatus.bRunFullSimplification;
 
@@ -13212,11 +13214,11 @@ bool FHLSLMaterialTranslator::FSubstrateCompilationContext::SubstrateGenerateDer
 
 #if WITH_EDITOR
 				Compiler->MaterialCompilationOutput.SubstrateMaterialCompilationOutput.SharedLocalBasesCount = 0; // FinalUsedSharedLocalBasesCount is not valid yet
-				Compiler->MaterialCompilationOutput.SubstrateMaterialCompilationOutput.RequestedBytePerPixel = SubstrateMaterialRequestedSizeByte;
+				Compiler->MaterialCompilationOutput.SubstrateMaterialCompilationOutput.RequestedBytePerPixel = SubstrateSimplificationStatus.OriginalRequestedByteSize;
 				Compiler->MaterialCompilationOutput.SubstrateMaterialCompilationOutput.PlatformBytePerPixel = SubstrateBytePerPixel;
 
 				Compiler->MaterialCompilationOutput.SubstrateMaterialCompilationOutput.RequestedClosurePerPixel = SubstrateSimplificationStatus.OriginalRequestedClosureCount;
-				Compiler->MaterialCompilationOutput.SubstrateMaterialCompilationOutput.PlatformClosurePixel = SubstrateMaterialClosureCount;
+				Compiler->MaterialCompilationOutput.SubstrateMaterialCompilationOutput.PlatformClosurePixel = SubstrateClosurePerPixel;
 
 				Compiler->MaterialCompilationOutput.SubstrateMaterialCompilationOutput.bIsThin = CompilerMaterial->IsThinSurface() ? 1 : 0;
 
@@ -15549,7 +15551,7 @@ void FHLSLMaterialTranslator::PrepareEnvironmentDefines()
 					return TEXT("ERROR");
 				};
 				FString SubstrateCompilationContextName = GetSubstrateCompilationContextName(SubstrateCompilationContextIndex);
-				SubstrateMaterialContextDescription += FString::Printf(TEXT("----- SUBSTRATE - %s -----\n"), *SubstrateCompilationContextName);
+				SubstrateMaterialContextDescription += FString::Printf(TEXT("\n\n----- SUBSTRATE - %s -----\n"), *SubstrateCompilationContextName);
 				SubstrateMaterialContextDescription += FString::Printf(TEXT("SubstrateCompilationInfo -\n"));
 				SubstrateMaterialContextDescription += FString::Printf(TEXT(" - Byte Per Pixel Budget                           %u\n"), SubstrateBytePerPixel_Platform);
 				SubstrateMaterialContextDescription += FString::Printf(TEXT(" - Closure Per Pixel Budget                        %u\n"), SubstrateClosurePerPixel_Platform);
