Make SubstrateTreeView available in MaterialEditor

diff --git a/Engine/Source/Editor/MaterialEditor/Private/MaterialEditor.cpp b/Engine/Source/Editor/MaterialEditor/Private/MaterialEditor.cpp
--- a/Engine/Source/Editor/MaterialEditor/Private/MaterialEditor.cpp
+++ b/Engine/Source/Editor/MaterialEditor/Private/MaterialEditor.cpp
@@ -1023,6 +1023,12 @@ void FMaterialEditor::UpdateGenerator()
 		TArray<UObject*> Objects;
 		Objects.Add(MaterialEditorInstance);
 		Generator->SetObjects(Objects);
+#if ENABLE_MATERIAL_LAYER_PROTOTYPE
+		if (MaterialLayersFunctionsInstance.IsValid())
+		{
+		 	MaterialLayersFunctionsInstance->SetEditorInstance(MaterialEditorInstance);
+		}
+#endif
 	}
 }
 
@@ -1756,10 +1762,15 @@ void FMaterialEditor::CreateInternalWidgets()
 	Generator->OnRowsRefreshed().AddSP(this, &FMaterialEditor::GeneratorRowsRefreshed);
 	MaterialCustomPrimitiveDataWidget = SNew(SMaterialCustomPrimitiveDataPanel, MaterialEditorInstance);
 
-	MaterialLayersFunctionsInstance = SNew(SMaterialLayersFunctionsMaterialWrapper)
+#if ENABLE_MATERIAL_LAYER_PROTOTYPE
+	MaterialLayersFunctionsInstance = SNew(SMaterialLayersFunctionsInstanceWrapper)
 		.InMaterialEditorInstance(MaterialEditorInstance)
 		.InGenerator(Generator);
-
+#else
+	MaterialLayersFunctionsInstance = SNew(SMaterialLayersFunctionsMaterialWrapper)
+			.InMaterialEditorInstance(MaterialEditorInstance)
+			.InGenerator(Generator);
+#endif
 	Palette = SNew(SMaterialPalette, SharedThis(this));
 
 	FMessageLogModule& MessageLogModule = FModuleManager::LoadModuleChecked<FMessageLogModule>("MessageLog");
@@ -1871,10 +1882,12 @@ void FMaterialEditor::OnChangeBreadCrumbGraph(UEdGraph* InGraph)
 void FMaterialEditor::GeneratorRowsRefreshed()
 {
 	MaterialParametersOverviewWidget->Refresh();
+#ifndef ENABLE_MATERIAL_LAYER_PROTOTYPE
 	if (MaterialLayersFunctionsInstance)
 	{
 		MaterialLayersFunctionsInstance->Refresh();
 	}
+#endif
 }
 
 FName FMaterialEditor::GetToolkitFName() const
diff --git a/Engine/Source/Editor/MaterialEditor/Private/MaterialEditor.h b/Engine/Source/Editor/MaterialEditor/Private/MaterialEditor.h
--- a/Engine/Source/Editor/MaterialEditor/Private/MaterialEditor.h
+++ b/Engine/Source/Editor/MaterialEditor/Private/MaterialEditor.h
@@ -998,8 +998,11 @@ private:
 	TSharedPtr<class SMaterialCustomPrimitiveDataPanel> MaterialCustomPrimitiveDataWidget;
 
 	/** Layer Properties View */
+#if ENABLE_MATERIAL_LAYER_PROTOTYPE
+	TSharedPtr<class SMaterialLayersFunctionsInstanceWrapper> MaterialLayersFunctionsInstance;
+#else
 	TSharedPtr<class SMaterialLayersFunctionsMaterialWrapper> MaterialLayersFunctionsInstance;
-
+#endif
 	/** The current transaction. */
 	FScopedTransaction* ScopedTransaction;
 
diff --git a/Engine/Source/Editor/MaterialEditor/Private/MaterialEditorDetailCustomization.cpp b/Engine/Source/Editor/MaterialEditor/Private/MaterialEditorDetailCustomization.cpp
--- a/Engine/Source/Editor/MaterialEditor/Private/MaterialEditorDetailCustomization.cpp
+++ b/Engine/Source/Editor/MaterialEditor/Private/MaterialEditorDetailCustomization.cpp
@@ -982,14 +982,18 @@ void FMaterialExpressionLayersParameterDetails::CustomizeDetails(IDetailLayoutBu
 {
 	FMaterialExpressionParameterDetails::CustomizeDetails(DetailLayout);
 	// for expression parameters all their properties are in one category based on their class name.
+
 	FName LayerCategory = FName("Layers");
+#if ENABLE_MATERIAL_LAYER_PROTOTYPE
+	DetailLayout.HideCategory(LayerCategory);
+#endif
 	IDetailCategoryBuilder& Category = DetailLayout.EditCategory(LayerCategory);
 	// Get a handle to the property we are about to edit
 	GroupPropertyHandle = DetailLayout.GetProperty("DefaultLayers");
 	GroupPropertyHandle->MarkHiddenByCustomization();
 	TSharedRef<FMaterialLayersFunctionsCustomization> MaterialLayersFunctionsCustomization = MakeShareable(new FMaterialLayersFunctionsCustomization(GroupPropertyHandle, &DetailLayout));
 	Category.AddCustomBuilder(MaterialLayersFunctionsCustomization);
-	
+
 }
 
 TSharedRef<class IDetailCustomization> FMaterialExpressionCompositeDetails::MakeInstance()
diff --git a/Engine/Source/Editor/MaterialEditor/Private/MaterialEditorInstanceDetailCustomization.cpp b/Engine/Source/Editor/MaterialEditor/Private/MaterialEditorInstanceDetailCustomization.cpp
--- a/Engine/Source/Editor/MaterialEditor/Private/MaterialEditorInstanceDetailCustomization.cpp
+++ b/Engine/Source/Editor/MaterialEditor/Private/MaterialEditorInstanceDetailCustomization.cpp
@@ -112,6 +112,11 @@ FString FMaterialInstanceParameterDetails::GetFunctionParentPath() const
 
 void FMaterialInstanceParameterDetails::CustomizeDetails(IDetailLayoutBuilder& DetailLayout)
 {
+#if ENABLE_MATERIAL_LAYER_PROTOTYPE
+	// early out if the new prototype is enabled as we only want to edit FMaterialLayersFunctions through the new Substrate tree UI
+	return;
+#endif
+	
 	PropertyUtilities = DetailLayout.GetPropertyUtilities();
 
 	// Create a new category for a custom layout for the MIC parameters at the very top
diff --git a/Engine/Source/Editor/MaterialEditor/Private/MaterialLayersFunctionsCustomization.cpp b/Engine/Source/Editor/MaterialEditor/Private/MaterialLayersFunctionsCustomization.cpp
--- a/Engine/Source/Editor/MaterialEditor/Private/MaterialLayersFunctionsCustomization.cpp
+++ b/Engine/Source/Editor/MaterialEditor/Private/MaterialLayersFunctionsCustomization.cpp
@@ -109,6 +109,11 @@ void FMaterialLayersFunctionsCustomization::RebuildChildren()
 void FMaterialLayersFunctionsCustomization::GenerateHeaderRowContent(FDetailWidgetRow& NodeRow)
 {
 
+#if ENABLE_MATERIAL_LAYER_PROTOTYPE
+	// early out if the new prototype is enabled as we only want to edit FMaterialLayersFunctions through the new Substrate tree UI
+	return;
+#endif
+	
 	TAttribute<bool>::FGetter IsResetButtonEnabledDelegate = TAttribute<bool>::FGetter::CreateSP(this, &FMaterialLayersFunctionsCustomization::IsResetToDefaultVisible);
 	TAttribute<bool> IsEnabledAttribute = TAttribute<bool>::Create(IsResetButtonEnabledDelegate);
 
@@ -145,6 +150,11 @@ void FMaterialLayersFunctionsCustomization::GenerateHeaderRowContent(FDetailWidg
 
 void FMaterialLayersFunctionsCustomization::GenerateChildContent(IDetailChildrenBuilder& ChildrenBuilder)
 {
+#if ENABLE_MATERIAL_LAYER_PROTOTYPE
+	// early out if the new prototype is enabled as we only want to edit FMaterialLayersFunctions through the new Substrate tree UI
+	return;
+#endif
+	
 	uint32 LayerChildren;
 	LayerHandle->GetNumChildren(LayerChildren);
 	uint32 BlendChildren;
diff --git a/Engine/Source/Editor/MaterialEditor/Private/SMaterialLayersFunctionsTree.cpp b/Engine/Source/Editor/MaterialEditor/Private/SMaterialLayersFunctionsTree.cpp
--- a/Engine/Source/Editor/MaterialEditor/Private/SMaterialLayersFunctionsTree.cpp
+++ b/Engine/Source/Editor/MaterialEditor/Private/SMaterialLayersFunctionsTree.cpp
@@ -1783,16 +1783,16 @@ void SMaterialLayersFunctionsInstanceWrapper::Refresh()
 	TSharedPtr<SHorizontalBox> HeaderBox;
 	NestedTree->CreateGroupsWidget();
 	LayerParameter = NestedTree->FunctionParameter;
-	FOnClicked 	OnChildButtonClicked = FOnClicked::CreateStatic(&FMaterialPropertyHelpers::OnClickedSaveNewMaterialInstance, ImplicitConv<UMaterialInterface*>(MaterialEditorInstance->SourceInstance), ImplicitConv<UObject*>(MaterialEditorInstance));
-	FOnClicked	OnSiblingButtonClicked = FOnClicked::CreateStatic(&FMaterialPropertyHelpers::OnClickedSaveNewMaterialInstance, MaterialEditorInstance->SourceInstance->Parent, ImplicitConv<UObject*>(MaterialEditorInstance));
+	FOnClicked 	OnChildButtonClicked = FOnClicked::CreateStatic(&FMaterialPropertyHelpers::OnClickedSaveNewMaterialInstance, MaterialEditorInstance->GetMaterialInterface(), ImplicitConv<UObject*>(MaterialEditorInstance));
+	FOnClicked	OnSiblingButtonClicked = FOnClicked::CreateStatic(&FMaterialPropertyHelpers::OnClickedSaveNewMaterialInstance, MaterialEditorInstance->GetParentMaterialInterface(), ImplicitConv<UObject*>(MaterialEditorInstance));
 
 #if ENABLE_MATERIAL_LAYER_PROTOTYPE
 	const float ThumbnailSize = 64.0f;
 	TSharedPtr<SBox> ThumbnailBox;
-	UObject* ThumbnailObject = MaterialEditorInstance->SourceInstance.Get();
+	UObject* ThumbnailObject = MaterialEditorInstance->GetMaterialInterface().Get();
 	const TSharedPtr<FAssetThumbnail> AssetThumbnail = MakeShareable(new FAssetThumbnail(ThumbnailObject, ThumbnailSize, ThumbnailSize, NestedTree->GetTreeThumbnailPool()));
 	TSharedRef<SWidget> ThumbnailWidget = AssetThumbnail->MakeThumbnailWidget();
-	FText MaterialName = FText::FromName(MaterialEditorInstance->SourceInstance->GetFName());
+	FText MaterialName = FText::FromName(MaterialEditorInstance->GetMaterialInterface()->GetFName());
 	
 #endif
 	
@@ -1946,11 +1946,12 @@ void SMaterialLayersFunctionsInstanceWrapper::Construct(const FArguments& InArgs
 #if ENABLE_MATERIAL_LAYER_PROTOTYPE
 	NestedTree = SNew(SMaterialSubstrateTree)
 		.InMaterialEditorInstance(InArgs._InMaterialEditorInstance)
+		.InGenerator(InArgs._InGenerator)
 		.InWrapper(this)
 		.InShowHiddenDelegate(InArgs._InShowHiddenDelegate);
 #else
 	NestedTree = SNew(SMaterialLayersFunctionsInstanceTree)
-		.InMaterialEditorInstance(InArgs._InMaterialEditorInstance)
+		.InMaterialEditorInstance(Cast<UMaterialEditorInstanceConstant>(InArgs._InMaterialEditorInstance))
 		.InWrapper(this)
 		.InShowHiddenDelegate(InArgs._InShowHiddenDelegate);
 #endif
@@ -1963,12 +1964,22 @@ void SMaterialLayersFunctionsInstanceWrapper::Construct(const FArguments& InArgs
 
 }
 
-void SMaterialLayersFunctionsInstanceWrapper::SetEditorInstance(UMaterialEditorInstanceConstant* InMaterialEditorInstance)
+void SMaterialLayersFunctionsInstanceWrapper::SetEditorInstance(UMaterialEditorParameters* InMaterialEditorInstance)
 {
+#if ENABLE_MATERIAL_LAYER_PROTOTYPE
 	NestedTree->MaterialEditorInstance = InMaterialEditorInstance;
+#else
+	NestedTree->MaterialEditorInstance = Cast<UMaterialEditorInstanceConstant>(InMaterialEditorInstance);
+#endif
+	
 	Refresh();
 }
 
+TSharedPtr<IPropertyRowGenerator> SMaterialLayersFunctionsInstanceWrapper::GetGenerator()
+{
+	return Generator.Pin();
+}
+
 
 #undef LOCTEXT_NAMESPACE
 
@@ -3232,10 +3243,12 @@ void SMaterialLayersFunctionsMaterialWrapper::Construct(const FArguments& InArgs
 	TSharedPtr<IPropertyRowGenerator> InGenerator = InArgs._InGenerator;
 	Generator = InGenerator;
 
+
 	NestedTree = SNew(SMaterialLayersFunctionsMaterialTree)
 		.InMaterialEditorInstance(InArgs._InMaterialEditorInstance)
 		.InWrapper(this);
 
+
 	LayerParameter = NestedTree->FunctionParameter;
 
 	MaterialEditorInstance = InArgs._InMaterialEditorInstance;
diff --git a/Engine/Source/Editor/MaterialEditor/Private/SMaterialLayersFunctionsTree.h b/Engine/Source/Editor/MaterialEditor/Private/SMaterialLayersFunctionsTree.h
--- a/Engine/Source/Editor/MaterialEditor/Private/SMaterialLayersFunctionsTree.h
+++ b/Engine/Source/Editor/MaterialEditor/Private/SMaterialLayersFunctionsTree.h
@@ -102,26 +102,33 @@ class SMaterialLayersFunctionsInstanceWrapper : public SCompoundWidget
 public:
 	SLATE_BEGIN_ARGS(SMaterialLayersFunctionsInstanceWrapper)
 		: _InMaterialEditorInstance(nullptr),
+		_InGenerator(nullptr),
 		_InShowHiddenDelegate()
 	{}
 
-	SLATE_ARGUMENT(UMaterialEditorInstanceConstant*, InMaterialEditorInstance)
+	SLATE_ARGUMENT(UMaterialEditorParameters*, InMaterialEditorInstance)
+	SLATE_ARGUMENT(TSharedPtr<class IPropertyRowGenerator>, InGenerator)
 	SLATE_ARGUMENT(FGetShowHiddenParameters, InShowHiddenDelegate)
 
 	SLATE_END_ARGS()
 	void Refresh();
 	void Construct(const FArguments& InArgs);
-	void SetEditorInstance(UMaterialEditorInstanceConstant* InMaterialEditorInstance);
-
+	void SetEditorInstance(UMaterialEditorParameters* InMaterialEditorInstance);
+	
+	TSharedPtr<class IPropertyRowGenerator> GetGenerator();
 	TAttribute<ECheckBoxState> IsParamChecked;
 	TWeakObjectPtr<class UDEditorParameterValue> LayerParameter;
-	class UMaterialEditorInstanceConstant* MaterialEditorInstance;
+	class UMaterialEditorParameters* MaterialEditorInstance;
 #if ENABLE_MATERIAL_LAYER_PROTOTYPE
 	TSharedPtr<class SMaterialSubstrateTree> NestedTree;
 #else
 	TSharedPtr<class SMaterialLayersFunctionsInstanceTree> NestedTree;
 #endif
 	FSimpleDelegate OnLayerPropertyChanged;
+	
+	
+private:
+	TWeakPtr<class IPropertyRowGenerator> Generator;
 };
 
 class SMaterialLayersFunctionsInstanceTree : public STreeView<TSharedPtr<FSortedParamData>>
diff --git a/Engine/Source/Editor/MaterialEditor/Private/SMaterialSubstrateTree.cpp b/Engine/Source/Editor/MaterialEditor/Private/SMaterialSubstrateTree.cpp
--- a/Engine/Source/Editor/MaterialEditor/Private/SMaterialSubstrateTree.cpp
+++ b/Engine/Source/Editor/MaterialEditor/Private/SMaterialSubstrateTree.cpp
@@ -401,7 +401,12 @@ FReply SMaterialSubstrateTreeItem::OnLayerDrop(const FDragDropEvent& DragDropEve
 
 			if (OriginalIndex != NewIndex)
 			{
-				Tree->MaterialEditorInstance->SourceInstance->SwapLayerParameterIndices(OriginalIndex, NewIndex);
+				if (Tree->MaterialEditorInstance->GetMaterialInterface()->IsA<UMaterialInstanceConstant>())
+				{
+					UMaterialInstanceConstant* MaterialConstant = Cast<UMaterialInstanceConstant>(Tree->MaterialEditorInstance->GetMaterialInterface());
+					MaterialConstant->SwapLayerParameterIndices(OriginalIndex, NewIndex);	
+				}
+				
 
 				// Need to save the moving and target expansion states before swapping
 				const bool bOriginalSwappableExpansion = IsItemExpanded();
@@ -443,7 +448,7 @@ FGetShowHiddenParameters SMaterialSubstrateTree::GetShowHiddenDelegate() const
 
 void  SMaterialSubstrateTreeItem::OnOverrideParameter(bool NewValue, class UDEditorParameterValue* Parameter)
 {
-	FMaterialPropertyHelpers::OnOverrideParameter(NewValue, Parameter, MaterialEditorInstance);
+	FMaterialPropertyHelpers::OnOverrideParameter(NewValue, Parameter, Cast<UMaterialEditorInstanceConstant>(MaterialEditorInstance));
 }
 
 void  SMaterialSubstrateTreeItem::OnOverrideParameter(bool NewValue, TObjectPtr<UDEditorParameterValue> Parameter)
@@ -826,6 +831,7 @@ void SMaterialSubstrateTree::Construct(const FArguments& InArgs)
 	ColumnSizeData.SetValueColumnWidth(0.5f);
 
 	MaterialEditorInstance = InArgs._InMaterialEditorInstance;
+	Generator = InArgs._InGenerator;
 	Wrapper = InArgs._InWrapper;
 	ShowHiddenDelegate = InArgs._InShowHiddenDelegate;
 	CreateGroupsWidget();
@@ -868,7 +874,14 @@ void SMaterialSubstrateTree::Construct(const FArguments& InArgs)
 }
 void SMaterialSubstrateTree::OnSelectionChangedMaterialSubstrateView(TSharedPtr<FSortedParamData> InSelectedItem, ESelectInfo::Type SelectInfo)
 {
-	MaterialEditorInstance->DetailsView.Pin()->ForceRefresh();
+	if (MaterialEditorInstance->IsA<UMaterialEditorInstanceConstant>())
+	{
+		UMaterialEditorInstanceConstant* MaterialEditorInstanceConstant = Cast<UMaterialEditorInstanceConstant>(MaterialEditorInstance);
+		if (TSharedPtr<IDetailsView> DetailsViewPinned = MaterialEditorInstanceConstant->DetailsView.Pin())
+		{
+			DetailsViewPinned->ForceRefresh();
+		}
+	}
 }
 TSharedRef< ITableRow > SMaterialSubstrateTree::OnGenerateRowMaterialLayersFunctionsTreeView(TSharedPtr<FSortedParamData> Item, const TSharedRef< STableViewBase >& OwnerTable)
 {
@@ -887,19 +900,20 @@ void SMaterialSubstrateTree::OnGetChildrenMaterialLayersFunctionsTreeView(TShare
 
 void SMaterialSubstrateTree::OnExpansionChanged(TSharedPtr<FSortedParamData> Item, bool bIsExpanded)
 {
-	bool* ExpansionValue = MaterialEditorInstance->SourceInstance->LayerParameterExpansion.Find(Item->NodeKey);
+	UMaterialInterface* MaterialInterface = MaterialEditorInstance->GetMaterialInterface();
+	bool* ExpansionValue = MaterialInterface->LayerParameterExpansion.Find(Item->NodeKey);
 	if (ExpansionValue == nullptr)
 	{
-		MaterialEditorInstance->SourceInstance->LayerParameterExpansion.Add(Item->NodeKey, bIsExpanded);
+		MaterialInterface->LayerParameterExpansion.Add(Item->NodeKey, bIsExpanded);
 	}
 	else if (*ExpansionValue != bIsExpanded)
 	{
-		MaterialEditorInstance->SourceInstance->LayerParameterExpansion.Emplace(Item->NodeKey, bIsExpanded);
+		MaterialInterface->LayerParameterExpansion.Emplace(Item->NodeKey, bIsExpanded);
 	}
 	// Expand any children that are also expanded
 	for (auto Child : Item->Children)
 	{
-		bool* ChildExpansionValue = MaterialEditorInstance->SourceInstance->LayerParameterExpansion.Find(Child->NodeKey);
+		bool* ChildExpansionValue = MaterialInterface->LayerParameterExpansion.Find(Child->NodeKey);
 		if (ChildExpansionValue != nullptr && *ChildExpansionValue == true)
 		{
 			SetItemExpansion(Child, true);
@@ -909,11 +923,13 @@ void SMaterialSubstrateTree::OnExpansionChanged(TSharedPtr<FSortedParamData> Ite
 
 void SMaterialSubstrateTree::SetParentsExpansionState()
 {
+	UMaterialInterface* MaterialInterface = MaterialEditorInstance->GetMaterialInterface();
+	
 	for (const auto& Pair : LayerProperties)
 	{
 		if (Pair->Children.Num())
 		{
-			bool* bIsExpanded = MaterialEditorInstance->SourceInstance->LayerParameterExpansion.Find(Pair->NodeKey);
+			bool* bIsExpanded = MaterialInterface->LayerParameterExpansion.Find(Pair->NodeKey);
 			if (bIsExpanded)
 			{
 				SetItemExpansion(Pair, *bIsExpanded);
@@ -941,10 +957,13 @@ void SMaterialSubstrateTree::RefreshOnAssetChange(const struct FAssetData& InAss
 
 void SMaterialSubstrateTree::ResetAssetToDefault(TSharedPtr<FSortedParamData> InData)
 {
-	FMaterialPropertyHelpers::ResetLayerAssetToDefault(InData->Parameter, InData->ParameterInfo.Association, InData->ParameterInfo.Index, MaterialEditorInstance);
-	UpdateThumbnailMaterial(InData->ParameterInfo.Association, InData->ParameterInfo.Index, false);
-	CreateGroupsWidget();
-	RequestTreeRefresh();
+	if (MaterialEditorInstance->IsA<UMaterialEditorInstanceConstant>())
+	{
+		FMaterialPropertyHelpers::ResetLayerAssetToDefault(InData->Parameter, InData->ParameterInfo.Association, InData->ParameterInfo.Index, Cast<UMaterialEditorInstanceConstant>(MaterialEditorInstance));
+		UpdateThumbnailMaterial(InData->ParameterInfo.Association, InData->ParameterInfo.Index, false);
+		CreateGroupsWidget();
+		RequestTreeRefresh();
+	}
 }
 
 void SMaterialSubstrateTree::AddNodeLayer(int32 InParent)
@@ -977,8 +996,14 @@ void SMaterialSubstrateTree::RemoveNodeLayer(int32 InNodeId)
 	auto NodePayload = FunctionInstance->GetNodePayload(InNodeId);
 
 	FunctionInstance->RemoveLayerNodeAt(InNodeId);
-	
-	MaterialEditorInstance->SourceInstance->RemoveLayerParameterIndex(NodePayload.Layer);
+	if (MaterialEditorInstance->IsA<UMaterialEditorInstanceConstant>())
+	{
+		UMaterialEditorInstanceConstant* MaterialEditorInstanceConstant = Cast<UMaterialEditorInstanceConstant>(MaterialEditorInstance);
+		if (MaterialEditorInstanceConstant && MaterialEditorInstanceConstant->SourceInstance)
+		{
+			MaterialEditorInstanceConstant->SourceInstance->RemoveLayerParameterIndex(NodePayload.Layer);
+		}
+	}
 	FunctionInstanceHandle->NotifyPostChange(EPropertyChangeType::ArrayRemove);
 	CreateGroupsWidget();
 	RequestTreeRefresh();
@@ -1007,7 +1032,7 @@ FReply SMaterialSubstrateTree::RelinkLayersToParent()
 
 EVisibility SMaterialSubstrateTree::GetUnlinkLayerVisibility(int32 Index) const
 {
-	if (FunctionInstance->IsLayerLinkedToParent(Index))
+	if (FunctionInstance->IsLayerLinkedToParent(Index) && MaterialEditorInstance->IsA<UMaterialEditorInstanceConstant>())
 	{
 		return EVisibility::Visible;
 	}
@@ -1186,13 +1211,17 @@ void SMaterialSubstrateTree::RecursiveCreateWidgets(FRecursiveCreateWidgetsConte
 void SMaterialSubstrateTree::CreateGroupsWidget()
 {
 	check(MaterialEditorInstance);
-	MaterialEditorInstance->RegenerateArrays();
+	if (MaterialEditorInstance->IsA<UMaterialEditorInstanceConstant>())
+	{
+		MaterialEditorInstance->RegenerateArrays();
+	}
 	NonLayerProperties.Empty();
 	LayerProperties.Empty();
 	FunctionParameter = nullptr;
 	TSharedPtr<IPropertyHandle> FunctionParameterHandle;
 
 	FPropertyEditorModule& Module = FModuleManager::LoadModuleChecked<FPropertyEditorModule>("PropertyEditor");
+
 	if (!Generator.IsValid())
 	{
 		FPropertyRowGeneratorArgs Args;
@@ -1214,7 +1243,7 @@ void SMaterialSubstrateTree::CreateGroupsWidget()
 		Objects.Add(MaterialEditorInstance);
 		Generator->SetObjects(Objects);
 	}
-
+	
 
 	TSharedPtr<IDetailTreeNode> ParameterGroups = FindParameterGroupsNode(Generator);
 	if (ParameterGroups.IsValid())
diff --git a/Engine/Source/Editor/MaterialEditor/Private/SMaterialSubstrateTree.h b/Engine/Source/Editor/MaterialEditor/Private/SMaterialSubstrateTree.h
--- a/Engine/Source/Editor/MaterialEditor/Private/SMaterialSubstrateTree.h
+++ b/Engine/Source/Editor/MaterialEditor/Private/SMaterialSubstrateTree.h
@@ -32,10 +32,12 @@ class SMaterialSubstrateTree : public STreeView<FSortedParamDataPtr>
 public:
 	SLATE_BEGIN_ARGS(SMaterialSubstrateTree)
 		: _InMaterialEditorInstance(nullptr)
+		, _InGenerator(nullptr)
 	{}
 
-	SLATE_ARGUMENT(UMaterialEditorInstanceConstant*, InMaterialEditorInstance)
+	SLATE_ARGUMENT(UMaterialEditorParameters*, InMaterialEditorInstance)
 	SLATE_ARGUMENT(SMaterialLayersFunctionsInstanceWrapper*, InWrapper)
+	SLATE_ARGUMENT(TSharedPtr<class IPropertyRowGenerator>, InGenerator)
 	SLATE_ARGUMENT(FGetShowHiddenParameters, InShowHiddenDelegate)
 	SLATE_END_ARGS()
 
@@ -69,7 +71,7 @@ public:
 	TSharedPtr<class FAssetThumbnailPool> GetTreeThumbnailPool();
 
 	/** Object that stores all of the possible parameters we can edit */
-	UMaterialEditorInstanceConstant* MaterialEditorInstance;
+	UMaterialEditorParameters* MaterialEditorInstance;
 
 	/** Builds the custom parameter groups category */
 	void CreateGroupsWidget();
@@ -122,7 +124,7 @@ public:
 
 	/** The item content. */
 	SLATE_ARGUMENT(FSortedParamDataPtr, StackParameterData)
-	SLATE_ARGUMENT(UMaterialEditorInstanceConstant*, MaterialEditorInstance)
+	SLATE_ARGUMENT(UMaterialEditorParameters*, MaterialEditorInstance)
 	SLATE_ARGUMENT(SMaterialSubstrateTree*, InTree)
 	SLATE_ATTRIBUTE( FMargin, Padding )
 	SLATE_END_ARGS()
@@ -190,7 +192,7 @@ public:
 
 	SMaterialSubstrateTree* Tree;
 
-	UMaterialEditorInstanceConstant* MaterialEditorInstance;
+	UMaterialEditorParameters* MaterialEditorInstance;
 
 	FString GetInstancePath(SMaterialSubstrateTree* InTree) const;
 };
diff --git a/Engine/Source/Editor/UnrealEd/Classes/MaterialEditor/MaterialEditorInstanceConstant.h b/Engine/Source/Editor/UnrealEd/Classes/MaterialEditor/MaterialEditorInstanceConstant.h
--- a/Engine/Source/Editor/UnrealEd/Classes/MaterialEditor/MaterialEditorInstanceConstant.h
+++ b/Engine/Source/Editor/UnrealEd/Classes/MaterialEditor/MaterialEditorInstanceConstant.h
@@ -260,11 +260,51 @@ struct FMaterialEditorPostProcessOverrides
 	FName UserSceneTextureOutput;
 };
 
+/** Common Interface for material parameter containers */
+UCLASS( abstract )
+class UMaterialEditorParameters: public UObject
+{
+	GENERATED_BODY()
+public:
+	/** 
+	 * Get the source/preview material interface for the parameters
+	 * @return source/preview material interface
+	 */
+	virtual TObjectPtr<UMaterialInterface> GetMaterialInterface() { return nullptr;};
+	virtual TObjectPtr<UMaterialInterface> GetParentMaterialInterface() { return nullptr;};
+
+#if WITH_EDITORONLY_DATA
+	UPROPERTY()
+	TArray<TObjectPtr<class UMaterialInstanceConstant>> StoredLayerPreviews;
+
+	UPROPERTY()
+	TArray<TObjectPtr<class UMaterialInstanceConstant>> StoredBlendPreviews;
+
+#endif
+	
+	/** 
+	 * Regenerates the parameter arrays. 
+	 */
+	virtual	UNREALED_API void RegenerateArrays() {}
+
+#if WITH_EDITOR
+	/** Sets back to zero the overrides for any parameters copied out of the layer stack */
+	UNREALED_API virtual void CleanParameterStack(int32 Index, EMaterialParameterAssociation MaterialType) { }
+
+	/** Copies the overrides for any parameters copied out of the layer stack from the layer or blend */
+	UNREALED_API virtual void ResetOverrides(int32 Index, EMaterialParameterAssociation MaterialType) {}
+	
+	/** Copies the parameter array values back to the source instance. */
+	UNREALED_API virtual void CopyToSourceInstance(const bool bForceStaticPermutationUpdate = false) {}
+	
+#endif
+};
+
 UCLASS(hidecategories=Object, collapsecategories, MinimalAPI)
-class UMaterialEditorInstanceConstant : public UObject
+class UMaterialEditorInstanceConstant : public UMaterialEditorParameters
 {
+private:
 	GENERATED_UCLASS_BODY()
-
 	/** Physical material to use for this graphics material. Used for sounds, effects etc.*/
 	UPROPERTY(EditAnywhere, Category=MaterialEditorInstanceConstant)
 	TObjectPtr<class UPhysicalMaterial> PhysMaterial;
@@ -337,21 +377,21 @@ class UMaterialEditorInstanceConstant : public UObject
 	//~ End UObject Interface.
 
 	/** Regenerates the parameter arrays. */
-	UNREALED_API void RegenerateArrays();
+	UNREALED_API void RegenerateArrays() override;
 
 #if WITH_EDITOR
 	/** Sets back to zero the overrides for any parameters copied out of the layer stack */
-	UNREALED_API void CleanParameterStack(int32 Index, EMaterialParameterAssociation MaterialType);
+	UNREALED_API void CleanParameterStack(int32 Index, EMaterialParameterAssociation MaterialType) override;
 
 	/** Copies the overrides for any parameters copied out of the layer stack from the layer or blend */
-	UNREALED_API void ResetOverrides(int32 Index, EMaterialParameterAssociation MaterialType);
+	UNREALED_API void ResetOverrides(int32 Index, EMaterialParameterAssociation MaterialType) override;
 
 	/** Arrays and clears parameters no longer valid (e.g. curve atlases). It has the potential effect of regenerating parameter arrays. */
 	UNREALED_API void ClearInvalidParameterOverrides();
 #endif
 
 	/** Copies the parameter array values back to the source instance. */
-	UNREALED_API void CopyToSourceInstance(const bool bForceStaticPermutationUpdate = false);
+	UNREALED_API void CopyToSourceInstance(const bool bForceStaticPermutationUpdate = false) override;
 
 	UNREALED_API void ApplySourceFunctionChanges();
 
@@ -385,17 +425,13 @@ class UMaterialEditorInstanceConstant : public UObject
 	 */
 	UNREALED_API void AssignParameterToGroup(UDEditorParameterValue* ParameterValue, const FName& GroupName);
 
+	TObjectPtr<UMaterialInterface> GetMaterialInterface() override;
+	TObjectPtr<UMaterialInterface> GetParentMaterialInterface() override;
+
 	static UNREALED_API FName GlobalGroupPrefix;
 
 	TWeakPtr<class IDetailsView> DetailsView;
 
-#if WITH_EDITORONLY_DATA
-	UPROPERTY()
-	TArray<TObjectPtr<class UMaterialInstanceConstant>> StoredLayerPreviews;
-
-	UPROPERTY()
-	TArray<TObjectPtr<class UMaterialInstanceConstant>> StoredBlendPreviews;
-#endif
 
 	/** Whether or not we should show only overridden properties*/
 	bool bShowOnlyOverrides;
diff --git a/Engine/Source/Editor/UnrealEd/Classes/MaterialEditor/MaterialEditorPreviewParameters.h b/Engine/Source/Editor/UnrealEd/Classes/MaterialEditor/MaterialEditorPreviewParameters.h
--- a/Engine/Source/Editor/UnrealEd/Classes/MaterialEditor/MaterialEditorPreviewParameters.h
+++ b/Engine/Source/Editor/UnrealEd/Classes/MaterialEditor/MaterialEditorPreviewParameters.h
@@ -18,11 +18,12 @@
 
 class UDEditorParameterValue;
 class UMaterial;
+class UMaterialInterface;
 class UMaterialInstanceConstant;
 struct FPropertyChangedEvent;
 
 UCLASS(hidecategories = Object, collapsecategories, MinimalAPI)
-class UMaterialEditorPreviewParameters : public UObject
+class UMaterialEditorPreviewParameters : public UMaterialEditorParameters
 {
 	GENERATED_BODY()
 
@@ -39,13 +40,6 @@ public:
 	UPROPERTY()
 	TObjectPtr<class UMaterial> OriginalMaterial;
 
-#if WITH_EDITORONLY_DATA
-	UPROPERTY()
-	TArray<TObjectPtr<class UMaterialInstanceConstant>> StoredLayerPreviews;
-
-	UPROPERTY()
-	TArray<TObjectPtr<class UMaterialInstanceConstant>> StoredBlendPreviews;
-#endif
 
 	//~ Begin UObject Interface.
 	UNREALED_API virtual void PostEditChangeProperty(FPropertyChangedEvent& PropertyChangedEvent) override;
@@ -55,11 +49,13 @@ public:
 	//~ End UObject Interface.
 
 	/** Regenerates the parameter arrays. */
-	UNREALED_API void RegenerateArrays();
+	UNREALED_API void RegenerateArrays() override;
+	
+	TObjectPtr<UMaterialInterface> GetMaterialInterface() override;
 
 protected:
 	/** Copies the parameter array values back to the source instance. */
-	UNREALED_API void CopyToSourceInstance();
+	UNREALED_API void CopyToSourceInstance(const bool bForceStaticPermutationUpdate = false) override;
 
 	UNREALED_API void ApplySourceFunctionChanges();
 
diff --git a/Engine/Source/Editor/UnrealEd/Private/PreviewMaterial.cpp b/Engine/Source/Editor/UnrealEd/Private/PreviewMaterial.cpp
--- a/Engine/Source/Editor/UnrealEd/Private/PreviewMaterial.cpp
+++ b/Engine/Source/Editor/UnrealEd/Private/PreviewMaterial.cpp
@@ -257,6 +257,35 @@ void UMaterialEditorPreviewParameters::PostEditChangeProperty(FPropertyChangedEv
 		FProperty* PropertyThatChanged = PropertyChangedEvent.Property;
 		if (OriginalFunction == nullptr)
 		{
+			bool bLayersParameterChanged = false;
+			// If a material layers parameter changed we need to update it on the source instance
+			// immediately so parameters contained within the new functions can be collected
+			for (FEditorParameterGroup& Group : ParameterGroups)
+			{
+				for (UDEditorParameterValue* Parameter : Group.Parameters)
+				{
+					if (UDEditorMaterialLayersParameterValue* LayersParam = Cast<UDEditorMaterialLayersParameterValue>(Parameter))
+					{
+						UMaterialExpressionMaterialAttributeLayers* LayersNode = nullptr;
+						for(UMaterialExpression* Expression : PreviewMaterial->GetExpressions())
+						{
+							if(Expression->IsA<UMaterialExpressionMaterialAttributeLayers>())
+							{
+								LayersNode = Cast<UMaterialExpressionMaterialAttributeLayers>(Expression);
+								LayersNode->DefaultLayers = LayersParam->ParameterValue;
+								bLayersParameterChanged = true;
+								break;
+							}
+						}
+					}
+				}
+			}
+
+			if (bLayersParameterChanged)
+			{
+				RegenerateArrays();
+			}
+			
 			CopyToSourceInstance();
 			PreviewMaterial->PostEditChangeProperty(PropertyChangedEvent);
 		}
@@ -396,6 +425,7 @@ void UMaterialEditorPreviewParameters::RegenerateArrays()
 			break;
 		}
 	}
+
 	if (ParameterDefaultGroups.Num() > 0)
 	{
 		ParameterGroups.Append(ParameterDefaultGroups);
@@ -403,11 +433,17 @@ void UMaterialEditorPreviewParameters::RegenerateArrays()
 
 }
 
-void UMaterialEditorPreviewParameters::CopyToSourceInstance()
+TObjectPtr<UMaterialInterface> UMaterialEditorPreviewParameters::GetMaterialInterface()
+{
+	return Cast<UMaterialInterface>(PreviewMaterial);
+}
+
+void UMaterialEditorPreviewParameters::CopyToSourceInstance(const bool bForceStaticPermutationUpdate/* = false*/)
 {
 	if (PreviewMaterial->IsTemplate(RF_ClassDefaultObject) == false && OriginalMaterial != nullptr)
 	{
 		OriginalMaterial->MarkPackageDirty();
+
 		// Scalar Parameters
 		for (int32 GroupIdx = 0; GroupIdx < ParameterGroups.Num(); GroupIdx++)
 		{
@@ -422,6 +458,19 @@ void UMaterialEditorPreviewParameters::CopyToSourceInstance()
 					{
 						PreviewMaterial->SetParameterValueEditorOnly(Parameter->ParameterInfo.Name, EditorValue);
 					}
+					else if (UDEditorMaterialLayersParameterValue* LayersParameter = Cast<UDEditorMaterialLayersParameterValue>(Parameter))
+					{
+						// find material expression material attribute layer and update it's default/param layers
+						for(UMaterialExpression* Expression : PreviewMaterial->GetExpressions())
+						{
+							if(Expression->IsA<UMaterialExpressionMaterialAttributeLayers>())
+							{
+								UMaterialExpressionMaterialAttributeLayers* LayersNode = Cast<UMaterialExpressionMaterialAttributeLayers>(Expression);
+								LayersNode->DefaultLayers = LayersParameter->ParameterValue;
+								break;
+							}
+						}
+					}
 				}
 			}
 		}
@@ -586,6 +635,16 @@ void  UMaterialEditorInstanceConstant::AssignParameterToGroup(UDEditorParameterV
 	CurrentGroup.Parameters.Add(ParameterValue);
 }
 
+TObjectPtr<UMaterialInterface> UMaterialEditorInstanceConstant::GetMaterialInterface()
+{
+	return Cast<UMaterialInterface>(SourceInstance);
+}
+
+TObjectPtr<UMaterialInterface> UMaterialEditorInstanceConstant::GetParentMaterialInterface()
+{
+	return SourceInstance ? SourceInstance->Parent : nullptr;
+}
+
 void UMaterialEditorInstanceConstant::RegenerateArrays()
 {
 	VisibleExpressions.Empty();
diff --git a/Engine/Source/Runtime/Engine/Public/Materials/MaterialLayersFunctions.h b/Engine/Source/Runtime/Engine/Public/Materials/MaterialLayersFunctions.h
--- a/Engine/Source/Runtime/Engine/Public/Materials/MaterialLayersFunctions.h
+++ b/Engine/Source/Runtime/Engine/Public/Materials/MaterialLayersFunctions.h
@@ -352,7 +352,7 @@ struct FMaterialLayersFunctionsRuntimeData
 	TArray<TObjectPtr<class UMaterialFunctionInterface>> Blends;
 
 	// ENABLE_MATERIAL_LAYER_PROTOTYPE : UProperty describing the tree data structure
-	UPROPERTY(EditAnywhere, Category = MaterialLayers)
+	UPROPERTY()
 	FMaterialLayersFunctionsTree Tree;
 
 	FMaterialLayersFunctionsRuntimeData() = default;
