removed check for showing Material Instance details correctly

diff --git a/Engine/Source/Editor/MaterialEditor/Private/MaterialEditorInstanceDetailCustomization.cpp b/Engine/Source/Editor/MaterialEditor/Private/MaterialEditorInstanceDetailCustomization.cpp
--- a/Engine/Source/Editor/MaterialEditor/Private/MaterialEditorInstanceDetailCustomization.cpp
+++ b/Engine/Source/Editor/MaterialEditor/Private/MaterialEditorInstanceDetailCustomization.cpp
@@ -112,11 +112,6 @@ FString FMaterialInstanceParameterDetails::GetFunctionParentPath() const
 
 void FMaterialInstanceParameterDetails::CustomizeDetails(IDetailLayoutBuilder& DetailLayout)
 {
-#if ENABLE_MATERIAL_LAYER_PROTOTYPE
-	// early out if the new prototype is enabled as we only want to edit FMaterialLayersFunctions through the new Substrate tree UI
-	return;
-#endif
-	
 	PropertyUtilities = DetailLayout.GetPropertyUtilities();
 
 	// Create a new category for a custom layout for the MIC parameters at the very top
