Mobile: Removed Linear64 and Gamma32 permutations variants as we do not support both permutations at the same time. Cleaned up material stats and fixed missing one for a projects with a CSM branch in the shader

diff --git a/Engine/Source/Editor/MaterialEditor/Private/MaterialStatsCommon.cpp b/Engine/Source/Editor/MaterialEditor/Private/MaterialStatsCommon.cpp
--- a/Engine/Source/Editor/MaterialEditor/Private/MaterialStatsCommon.cpp
+++ b/Engine/Source/Editor/MaterialEditor/Private/MaterialStatsCommon.cpp
@@ -300,26 +300,8 @@ FSlateColor FMaterialStatsUtils::QualitySettingColor(const EMaterialQualityLevel
 	}
 }
 
-static void MobileBasePassShaderName(bool bVertexShader, const TCHAR* PolicyName, const TCHAR* LocalLightSetting, bool bHDR, bool bSkyLight, FString& OutName)
-{
-	OutName.Reset();
-	OutName.Append(bVertexShader ? TEXT("TMobileBasePassVS") : TEXT("TMobileBasePassPS"));
-	OutName.Append(PolicyName);
-	OutName.Append(bHDR ? TEXT("HDRLinear64") : TEXT("LDRGamma32"));
-	if (!bVertexShader && bSkyLight)
-	{
-		OutName.Append(TEXT("SkyLight"));
-	}
-	if (!bVertexShader)
-	{
-		OutName.Append(LocalLightSetting);
-	}
-}
-
 void FMaterialStatsUtils::GetRepresentativeShaderTypesAndDescriptions(TMap<FName, TArray<FRepresentativeShaderInfo>>& ShaderTypeNamesAndDescriptions, const FMaterial* TargetMaterial)
 {
-	bool bMobileHDR = IsMobileHDR();
-
 	static const FName FLocalVertexFactoryName = FLocalVertexFactory::StaticType.GetFName();
 	static const FName FGPUFactoryName = TEXT("TGPUSkinVertexFactoryDefault");
 	static const FName FClothVertexFactoryName = TEXT("TGPUSkinAPEXClothVertexFactoryDefault");
@@ -412,162 +394,73 @@ void FMaterialStatsUtils::GetRepresentativeShaderTypesAndDescriptions(TMap<FName
 	}
 	else
 	{
-		const TCHAR* DescSuffix = bMobileHDR ? TEXT(" (HDR)") : TEXT(" (LDR)");
-		FString ShaderNameStr;
-
+		static const FName TMobileBasePassVSFNoLightMapPolicyName = TEXT("TMobileBasePassVSFNoLightMapPolicy");
+		static const FName TMobileBasePassPSFNoLightMapPolicyName[2] = {
+			TEXT("TMobileBasePassPSFNoLightMapPolicyLOCAL_LIGHTS_DISABLED"),
+			TEXT("TMobileBasePassPSFNoLightMapPolicySkyLightLOCAL_LIGHTS_DISABLED")
+		};
+		
 		if (TargetMaterial->GetShadingModels().IsUnlit())
 		{
-			//unlit materials are never lightmapped
-			MobileBasePassShaderName(false, TEXT("FNoLightMapPolicy"), TEXT("LOCAL_LIGHTS_DISABLED"), bMobileHDR, false, ShaderNameStr);
-			const FString Description = FString::Printf(TEXT("Mobile base pass shader without light map%s"), DescSuffix);
 			ShaderTypeNamesAndDescriptions.Add(FLocalVertexFactoryName)
-				.Add(FRepresentativeShaderInfo(ERepresentativeShader::StationarySurface, FName(ShaderNameStr), Description));
-
-			MobileBasePassShaderName(true, TEXT("FNoLightMapPolicy"), TEXT("LOCAL_LIGHTS_DISABLED"), bMobileHDR, false, ShaderNameStr);
-			ShaderTypeNamesAndDescriptions.FindOrAdd(FLocalVertexFactoryName)
-				.Add(FRepresentativeShaderInfo(ERepresentativeShader::StaticMesh, FName(ShaderNameStr),
-					FString::Printf(TEXT("Mobile base pass vertex shader%s"), DescSuffix)));
-
-			if (TargetMaterial->IsUsedWithSkeletalMesh() || TargetMaterial->IsUsedWithMorphTargets())
-			{
-				ShaderTypeNamesAndDescriptions.FindOrAdd(FGPUFactoryName)
-					.Add(FRepresentativeShaderInfo(ERepresentativeShader::SkeletalMesh, FName(ShaderNameStr),
-						FString::Printf(TEXT("Mobile base pass vertex shader%s"), DescSuffix)));
-			}
-			if (TargetMaterial->IsUsedWithAPEXCloth())
-			{
-				ShaderTypeNamesAndDescriptions.FindOrAdd(FClothVertexFactoryName)
-					.Add(FRepresentativeShaderInfo(ERepresentativeShader::SkinnedCloth, FName(ShaderNameStr),
-						FString::Printf(TEXT("Mobile base pass vertex shader%s"), DescSuffix)));
-			}
+				.Add(FRepresentativeShaderInfo(ERepresentativeShader::StationarySurface, TMobileBasePassPSFNoLightMapPolicyName[0],
+					TEXT("Mobile base pass shader without light map")));
 		}
 		else
 		{			
+			// Use SkyLight permutation of projects settings force it
 			static auto* CVarMobileSkyLightPermutation = IConsoleManager::Get().FindTConsoleVariableDataInt(TEXT("r.Mobile.SkyLightPermutation"));
-			const bool bOnlySkyPermutation = CVarMobileSkyLightPermutation->GetValueOnAnyThread() == 2;
+			const bool SkyLightId = (CVarMobileSkyLightPermutation->GetValueOnAnyThread() == 2) ? 1 : 0;
 						
 			if (IsStaticLightingAllowed() && TargetMaterial->IsUsedWithStaticLighting())
 			{
 				static auto* CVarAllowDistanceFieldShadows = IConsoleManager::Get().FindTConsoleVariableDataInt(TEXT("r.Mobile.AllowDistanceFieldShadows"));
 				const bool bAllowDistanceFieldShadows = CVarAllowDistanceFieldShadows->GetValueOnAnyThread() != 0;
-
-				if (bAllowDistanceFieldShadows)// distance field shadows
+				if (bAllowDistanceFieldShadows)
 				{
 					// distance field shadows only shaders
-					{
-						MobileBasePassShaderName(false, TEXT("FMobileDistanceFieldShadowsAndLQLightMapPolicy"), TEXT("LOCAL_LIGHTS_DISABLED"), bMobileHDR, bOnlySkyPermutation, ShaderNameStr);
-						const FString Description = FString::Printf(TEXT("Mobile base pass shader with distance field shadows%s"), DescSuffix);
-						ShaderTypeNamesAndDescriptions.FindOrAdd(FLocalVertexFactoryName)
-							.Add(FRepresentativeShaderInfo(ERepresentativeShader::StationarySurface, FName(ShaderNameStr), Description));
-					}
-
-					static auto* CVarAllowDistanceFieldShadowsAndCSM = IConsoleManager::Get().FindTConsoleVariableDataInt(TEXT("r.Mobile.EnableStaticAndCSMShadowReceivers"));
-					const bool bAllowDistanceFieldShadowsAndCSM = CVarAllowDistanceFieldShadowsAndCSM->GetValueOnAnyThread() != 0;
-					if (bAllowDistanceFieldShadowsAndCSM)
-					{
-						// distance field shadows & CSM shaders
-						{
-							MobileBasePassShaderName(false, TEXT("FMobileDistanceFieldShadowsLightMapAndCSMLightingPolicy"), TEXT("LOCAL_LIGHTS_DISABLED"), bMobileHDR, bOnlySkyPermutation, ShaderNameStr);
-							const FString Description = FString::Printf(TEXT("Mobile base pass shader with distance field shadows and CSM%s"), DescSuffix);
-							ShaderTypeNamesAndDescriptions.FindOrAdd(FLocalVertexFactoryName)
-								.Add(FRepresentativeShaderInfo(ERepresentativeShader::StationarySurfaceCSM, FName(ShaderNameStr), Description));
-						}
-
-						{
-							MobileBasePassShaderName(false, TEXT("FMobileDistanceFieldShadowsAndLQLightMapPolicy"), TEXT("LOCAL_LIGHTS_ENABLED"), bMobileHDR, bOnlySkyPermutation, ShaderNameStr);
-							const FString Description = FString::Printf(TEXT("Mobile base pass shader with distance field shadows, CSM and local light(s) %s"), DescSuffix);
-
-							FRepresentativeShaderInfo ShaderInfo = FRepresentativeShaderInfo(ERepresentativeShader::StationarySurfaceNPointLights, FName(ShaderNameStr), Description);
-
-							ShaderTypeNamesAndDescriptions.FindOrAdd(FLocalVertexFactoryName).Add(ShaderInfo);
-								
-						}
-
-						{
-							MobileBasePassShaderName(false, TEXT("FMobileDistanceFieldShadowsAndLQLightMapPolicy"), TEXT("LOCAL_LIGHTS_BUFFER"), bMobileHDR, bOnlySkyPermutation, ShaderNameStr);
-							const FString Description = FString::Printf(TEXT("Mobile base pass shader with distance field shadows, CSM and local light(s) %s"), DescSuffix);
-
-							FRepresentativeShaderInfo ShaderInfo = FRepresentativeShaderInfo(ERepresentativeShader::StationarySurfaceNPointLights, FName(ShaderNameStr), Description);
-
-							ShaderTypeNamesAndDescriptions.FindOrAdd(FLocalVertexFactoryName).Add(ShaderInfo);
-
-						}
-					}
+					static const FName TMobileBasePassPSFMobileDistanceFieldShadowsAndLQLightMapPolicyName[2] = {
+						TEXT("TMobileBasePassPSFMobileDistanceFieldShadowsAndLQLightMapPolicyLOCAL_LIGHTS_DISABLED"),
+						TEXT("TMobileBasePassPSFMobileDistanceFieldShadowsAndLQLightMapPolicySkyLightLOCAL_LIGHTS_DISABLED"),
+					};
+					
+					ShaderTypeNamesAndDescriptions.FindOrAdd(FLocalVertexFactoryName)
+							.Add(FRepresentativeShaderInfo(ERepresentativeShader::StationarySurface, TMobileBasePassPSFMobileDistanceFieldShadowsAndLQLightMapPolicyName[SkyLightId],
+								TEXT("Mobile base pass shader with distance field shadows")));
 				}
 				else //no shadows & lightmapped
 				{
-					{
-						MobileBasePassShaderName(false, TEXT("TLightMapPolicyLQ"), TEXT("LOCAL_LIGHTS_DISABLED"), bMobileHDR, bOnlySkyPermutation, ShaderNameStr);
-						ShaderTypeNamesAndDescriptions.FindOrAdd(FLocalVertexFactoryName)
-							.Add(FRepresentativeShaderInfo(ERepresentativeShader::StationarySurface, FName(ShaderNameStr),
-								FString::Printf(TEXT("Mobile base pass shader with static lighting%s"), DescSuffix)));
-					}
-
-					{
-						MobileBasePassShaderName(false, TEXT("TLightMapPolicyLQ"), TEXT("LOCAL_LIGHTS_ENABLED"), bMobileHDR, bOnlySkyPermutation, ShaderNameStr);
-						const FString Description = FString::Printf(TEXT("Mobile base pass shader with static lighting and local light(s) %s"), DescSuffix);
-						FRepresentativeShaderInfo ShaderInfo = FRepresentativeShaderInfo(ERepresentativeShader::StationarySurfaceNPointLights, FName(ShaderNameStr), Description);
-						ShaderTypeNamesAndDescriptions.FindOrAdd(FLocalVertexFactoryName).Add(ShaderInfo);
-					}
-
-					{
-						MobileBasePassShaderName(false, TEXT("TLightMapPolicyLQ"), TEXT("LOCAL_LIGHTS_BUFFER"), bMobileHDR, bOnlySkyPermutation, ShaderNameStr);
-						const FString Description = FString::Printf(TEXT("Mobile base pass shader with static lighting and local light(s) %s"), DescSuffix);
-						FRepresentativeShaderInfo ShaderInfo = FRepresentativeShaderInfo(ERepresentativeShader::StationarySurfaceNPointLights, FName(ShaderNameStr), Description);
-						ShaderTypeNamesAndDescriptions.FindOrAdd(FLocalVertexFactoryName).Add(ShaderInfo);
-					}
+					static const FName TMobileBasePassPSTLightMapPolicyLQName[2] = {
+						TEXT("TMobileBasePassPSTLightMapPolicyLQLOCAL_LIGHTS_DISABLED"),
+						TEXT("TMobileBasePassPSTLightMapPolicyLQSkyLightLOCAL_LIGHTS_DISABLED"),
+					};
+										
+					ShaderTypeNamesAndDescriptions.FindOrAdd(FLocalVertexFactoryName)
+						.Add(FRepresentativeShaderInfo(ERepresentativeShader::StationarySurface, TMobileBasePassPSTLightMapPolicyLQName[SkyLightId],
+							TEXT("Mobile base pass shader with static lighting")));
 				}
 			}
-
-			// only one of these 2 shader types will be displayed
-			
 			// dynamically lit shader NoLightmapPolicy
-			MobileBasePassShaderName(false, TEXT("FNoLightmapPolicy"), TEXT("LOCAL_LIGHTS_DISABLED"), bMobileHDR, bOnlySkyPermutation, ShaderNameStr);
-			ShaderTypeNamesAndDescriptions.FindOrAdd(FLocalVertexFactoryName)
-				.Add(FRepresentativeShaderInfo(ERepresentativeShader::DynamicallyLitObject, FName(ShaderNameStr),
-				FString::Printf(TEXT("Mobile base pass shader with only dynamic lighting%s"), DescSuffix)));
-
-			MobileBasePassShaderName(true, TEXT("FNoLightMapPolicy"), TEXT("LOCAL_LIGHTS_DISABLED"), bMobileHDR, bOnlySkyPermutation, ShaderNameStr);
 			ShaderTypeNamesAndDescriptions.FindOrAdd(FLocalVertexFactoryName)
-				.Add(FRepresentativeShaderInfo(ERepresentativeShader::StaticMesh, FName(ShaderNameStr),
-				FString::Printf(TEXT("Mobile base pass vertex shader%s"), DescSuffix)));
-
-			if (TargetMaterial->IsUsedWithSkeletalMesh() || TargetMaterial->IsUsedWithMorphTargets())
-			{
-				ShaderTypeNamesAndDescriptions.FindOrAdd(FGPUFactoryName)
-					.Add(FRepresentativeShaderInfo(ERepresentativeShader::SkeletalMesh, FName(ShaderNameStr),
-						FString::Printf(TEXT("Mobile base pass vertex shader%s"), DescSuffix)));
-			}
-			if (TargetMaterial->IsUsedWithAPEXCloth())
-			{
-				ShaderTypeNamesAndDescriptions.FindOrAdd(FClothVertexFactoryName)
-					.Add(FRepresentativeShaderInfo(ERepresentativeShader::SkinnedCloth, FName(ShaderNameStr),
-						FString::Printf(TEXT("Mobile base pass vertex shader%s"), DescSuffix)));
-			}
-
-			// dynamically lit shader FMobileDirectionalLightAndCSMPolicy
-			MobileBasePassShaderName(false, TEXT("FMobileDirectionalLightAndCSMPolicy"), TEXT("LOCAL_LIGHTS_DISABLED"), bMobileHDR, bOnlySkyPermutation, ShaderNameStr);
-			ShaderTypeNamesAndDescriptions.FindOrAdd(FLocalVertexFactoryName)
-				.Add(FRepresentativeShaderInfo(ERepresentativeShader::DynamicallyLitObject, FName(ShaderNameStr),
-				FString::Printf(TEXT("Mobile base pass shader with only dynamic lighting%s"), DescSuffix)));
+				.Add(FRepresentativeShaderInfo(ERepresentativeShader::DynamicallyLitObject, TMobileBasePassPSFNoLightMapPolicyName[SkyLightId],
+					TEXT("Mobile base pass shader with only dynamic lighting")));
+		}
 
-			MobileBasePassShaderName(true, TEXT("FMobileDirectionalLightAndCSMPolicy"), TEXT("LOCAL_LIGHTS_DISABLED"), bMobileHDR, bOnlySkyPermutation, ShaderNameStr);
-			ShaderTypeNamesAndDescriptions.FindOrAdd(FLocalVertexFactoryName)
-				.Add(FRepresentativeShaderInfo(ERepresentativeShader::StaticMesh, FName(ShaderNameStr),
-				FString::Printf(TEXT("Mobile base pass vertex shader%s"), DescSuffix)));
+		ShaderTypeNamesAndDescriptions.FindOrAdd(FLocalVertexFactoryName)
+			.Add(FRepresentativeShaderInfo(ERepresentativeShader::StaticMesh, TMobileBasePassVSFNoLightMapPolicyName,
+				TEXT("Mobile base pass vertex shader")));
 
-			if (TargetMaterial->IsUsedWithSkeletalMesh() || TargetMaterial->IsUsedWithMorphTargets())
-			{
-				ShaderTypeNamesAndDescriptions.FindOrAdd(FGPUFactoryName)
-					.Add(FRepresentativeShaderInfo(ERepresentativeShader::SkeletalMesh, FName(ShaderNameStr),
-						FString::Printf(TEXT("Mobile base pass vertex shader%s"), DescSuffix)));
-			}
-			if (TargetMaterial->IsUsedWithAPEXCloth())
-			{
-				ShaderTypeNamesAndDescriptions.FindOrAdd(FClothVertexFactoryName)
-					.Add(FRepresentativeShaderInfo(ERepresentativeShader::SkinnedCloth, FName(ShaderNameStr),
-						FString::Printf(TEXT("Mobile base pass vertex shader%s"), DescSuffix)));
-			}
+		if (TargetMaterial->IsUsedWithSkeletalMesh() || TargetMaterial->IsUsedWithMorphTargets())
+		{
+			ShaderTypeNamesAndDescriptions.FindOrAdd(FGPUFactoryName)
+				.Add(FRepresentativeShaderInfo(ERepresentativeShader::SkeletalMesh, TMobileBasePassVSFNoLightMapPolicyName,
+					TEXT("Mobile base pass vertex shader")));
+		}
+		if (TargetMaterial->IsUsedWithAPEXCloth())
+		{
+			ShaderTypeNamesAndDescriptions.FindOrAdd(FClothVertexFactoryName)
+				.Add(FRepresentativeShaderInfo(ERepresentativeShader::SkinnedCloth, TMobileBasePassVSFNoLightMapPolicyName,
+					TEXT("Mobile base pass vertex shader")));
 		}
 	}
 }
diff --git a/Engine/Source/Runtime/Landscape/Private/LandscapeRender.cpp b/Engine/Source/Runtime/Landscape/Private/LandscapeRender.cpp
--- a/Engine/Source/Runtime/Landscape/Private/LandscapeRender.cpp
+++ b/Engine/Source/Runtime/Landscape/Private/LandscapeRender.cpp
@@ -3966,26 +3966,15 @@ public:
 			FName(TEXT("TLightMapDensityPSFNoLightMapPolicy")),
 
 			// Mobile
-			FName(TEXT("TMobileBasePassPSFMobileDirectionalLightCSMAndSHIndirectPolicyINT32_MAXHDRLinear64Skylight")),
-			FName(TEXT("TMobileBasePassPSFMobileDirectionalLightCSMAndSHIndirectPolicyINT32_MAXHDRLinear64")),
-			FName(TEXT("TMobileBasePassPSFMobileDirectionalLightCSMAndSHIndirectPolicy0HDRLinear64Skylight")),
-			FName(TEXT("TMobileBasePassPSFMobileDirectionalLightCSMAndSHIndirectPolicy0HDRLinear64")),
-			FName(TEXT("TMobileBasePassVSFMobileDirectionalLightCSMAndSHIndirectPolicyHDRLinear64")),
-			FName(TEXT("TMobileBasePassPSFMobileDirectionalLightAndSHIndirectPolicyINT32_MAXHDRLinear64Skylight")),
-			FName(TEXT("TMobileBasePassPSFMobileDirectionalLightAndSHIndirectPolicyINT32_MAXHDRLinear64")),
-			FName(TEXT("TMobileBasePassPSFMobileDirectionalLightAndSHIndirectPolicy0HDRLinear64Skylight")),
-			FName(TEXT("TMobileBasePassPSFMobileDirectionalLightAndSHIndirectPolicy0HDRLinear64")),
-			FName(TEXT("TMobileBasePassVSFMobileDirectionalLightAndSHIndirectPolicyHDRLinear64")),
-			FName(TEXT("TMobileBasePassPSFMobileDirectionalLightAndCSMPolicyINT32_MAXHDRLinear64Skylight")),
-			FName(TEXT("TMobileBasePassPSFMobileDirectionalLightAndCSMPolicyINT32_MAXHDRLinear64")),
-			FName(TEXT("TMobileBasePassPSFMobileDirectionalLightAndCSMPolicy0HDRLinear64Skylight")),
-			FName(TEXT("TMobileBasePassPSFMobileDirectionalLightAndCSMPolicy0HDRLinear64")),
-			FName(TEXT("TMobileBasePassVSFMobileDirectionalLightAndCSMPolicyHDRLinear64")),
-			FName(TEXT("TMobileBasePassPSFNoLightMapPolicyINT32_MAXHDRLinear64Skylight")),
-			FName(TEXT("TMobileBasePassPSFNoLightMapPolicyINT32_MAXHDRLinear64")),
-			FName(TEXT("TMobileBasePassPSFNoLightMapPolicy0HDRLinear64Skylight")),
-			FName(TEXT("TMobileBasePassPSFNoLightMapPolicy0HDRLinear64")),
-			FName(TEXT("TMobileBasePassVSFNoLightMapPolicyHDRLinear64")),
+			FName(TEXT("TMobileBasePassPSFMobileDirectionalLightAndSHIndirectPolicySkylightLOCAL_LIGHTS_DISABLED")),
+			FName(TEXT("TMobileBasePassPSFMobileDirectionalLightAndSHIndirectPolicyLOCAL_LIGHTS_DISABLED")),
+			FName(TEXT("TMobileBasePassVSFMobileDirectionalLightAndSHIndirectPolicy")),
+			FName(TEXT("TMobileBasePassPSFMobileDirectionalLightAndCSMPolicySkylightLOCAL_LIGHTS_DISABLED")),
+			FName(TEXT("TMobileBasePassPSFMobileDirectionalLightAndCSMPolicyLOCAL_LIGHTS_DISABLED")),
+			FName(TEXT("TMobileBasePassVSFMobileDirectionalLightAndCSMPolicy")),
+			FName(TEXT("TMobileBasePassPSFNoLightMapPolicySkylightLOCAL_LIGHTS_DISABLED")),
+			FName(TEXT("TMobileBasePassPSFNoLightMapPolicyLOCAL_LIGHTS_DISABLED")),
+			FName(TEXT("TMobileBasePassVSFNoLightMapPolicy")),
 
 			// Forward shading required
 			FName(TEXT("TBasePassPSFCachedPointIndirectLightingPolicySkylight")),
@@ -4039,31 +4028,21 @@ public:
 			FName(TEXT("FDebugViewModePS")),
 
 			// Mobile
-			FName(TEXT("TMobileBasePassPSFMobileMovableDirectionalLightCSMWithLightmapPolicyINT32_MAXHDRLinear64Skylight")),
-			FName(TEXT("TMobileBasePassPSFMobileMovableDirectionalLightCSMWithLightmapPolicyINT32_MAXHDRLinear64")),
-			FName(TEXT("TMobileBasePassPSFMobileMovableDirectionalLightCSMWithLightmapPolicy0HDRLinear64Skylight")),
-			FName(TEXT("TMobileBasePassPSFMobileMovableDirectionalLightCSMWithLightmapPolicy0HDRLinear64")),
-			FName(TEXT("TMobileBasePassVSFMobileMovableDirectionalLightCSMWithLightmapPolicyHDRLinear64")),
-			FName(TEXT("TMobileBasePassPSFMobileMovableDirectionalLightWithLightmapPolicyINT32_MAXHDRLinear64Skylight")),
-			FName(TEXT("TMobileBasePassPSFMobileMovableDirectionalLightWithLightmapPolicyINT32_MAXHDRLinear64")),
-			FName(TEXT("TMobileBasePassPSFMobileMovableDirectionalLightWithLightmapPolicy0HDRLinear64Skylight")),
-			FName(TEXT("TMobileBasePassPSFMobileMovableDirectionalLightWithLightmapPolicy0HDRLinear64")),
-			FName(TEXT("TMobileBasePassVSFMobileMovableDirectionalLightWithLightmapPolicyHDRLinear64")),
-			FName(TEXT("TMobileBasePassPSFMobileDistanceFieldShadowsLightMapAndCSMLightingPolicyINT32_MAXHDRLinear64Skylight")),
-			FName(TEXT("TMobileBasePassPSFMobileDistanceFieldShadowsLightMapAndCSMLightingPolicyINT32_MAXHDRLinear64")),
-			FName(TEXT("TMobileBasePassPSFMobileDistanceFieldShadowsLightMapAndCSMLightingPolicy0HDRLinear64Skylight")),
-			FName(TEXT("TMobileBasePassPSFMobileDistanceFieldShadowsLightMapAndCSMLightingPolicy0HDRLinear64")),
-			FName(TEXT("TMobileBasePassVSFMobileDistanceFieldShadowsLightMapAndCSMLightingPolicyHDRLinear64")),
-			FName(TEXT("TMobileBasePassPSFMobileDistanceFieldShadowsAndLQLightMapPolicyINT32_MAXHDRLinear64Skylight")),
-			FName(TEXT("TMobileBasePassPSFMobileDistanceFieldShadowsAndLQLightMapPolicyINT32_MAXHDRLinear64")),
-			FName(TEXT("TMobileBasePassPSFMobileDistanceFieldShadowsAndLQLightMapPolicy0HDRLinear64Skylight")),
-			FName(TEXT("TMobileBasePassPSFMobileDistanceFieldShadowsAndLQLightMapPolicy0HDRLinear64")),
-			FName(TEXT("TMobileBasePassVSFMobileDistanceFieldShadowsAndLQLightMapPolicyHDRLinear64")),
-			FName(TEXT("TMobileBasePassPSTLightMapPolicyLQINT32_MAXHDRLinear64Skylight")),
-			FName(TEXT("TMobileBasePassPSTLightMapPolicyLQINT32_MAXHDRLinear64")),
-			FName(TEXT("TMobileBasePassPSTLightMapPolicyLQ0HDRLinear64Skylight")),
-			FName(TEXT("TMobileBasePassPSTLightMapPolicyLQ0HDRLinear64")),
-			FName(TEXT("TMobileBasePassVSTLightMapPolicyLQHDRLinear64")),
+			FName(TEXT("TMobileBasePassPSFMobileMovableDirectionalLightCSMWithLightmapPolicySkylightLOCAL_LIGHTS_DISABLED")),
+			FName(TEXT("TMobileBasePassPSFMobileMovableDirectionalLightCSMWithLightmapPolicyLOCAL_LIGHTS_DISABLED")),
+			FName(TEXT("TMobileBasePassVSFMobileMovableDirectionalLightCSMWithLightmapPolicy")),
+			FName(TEXT("TMobileBasePassPSFMobileMovableDirectionalLightWithLightmapPolicySkylightLOCAL_LIGHTS_DISABLED")),
+			FName(TEXT("TMobileBasePassPSFMobileMovableDirectionalLightWithLightmapPolicyLOCAL_LIGHTS_DISABLED")),
+			FName(TEXT("TMobileBasePassVSFMobileMovableDirectionalLightWithLightmapPolicy")),
+			FName(TEXT("TMobileBasePassPSFMobileDistanceFieldShadowsLightMapAndCSMLightingPolicySkylightLOCAL_LIGHTS_DISABLED")),
+			FName(TEXT("TMobileBasePassPSFMobileDistanceFieldShadowsLightMapAndCSMLightingPolicyLOCAL_LIGHTS_DISABLED")),
+			FName(TEXT("TMobileBasePassVSFMobileDistanceFieldShadowsLightMapAndCSMLightingPolicy")),
+			FName(TEXT("TMobileBasePassPSFMobileDistanceFieldShadowsAndLQLightMapPolicySkylightLOCAL_LIGHTS_DISABLED")),
+			FName(TEXT("TMobileBasePassPSFMobileDistanceFieldShadowsAndLQLightMapPolicyLOCAL_LIGHTS_DISABLED")),
+			FName(TEXT("TMobileBasePassVSFMobileDistanceFieldShadowsAndLQLightMapPolicy")),
+			FName(TEXT("TMobileBasePassPSTLightMapPolicyLQSkylightLOCAL_LIGHTS_DISABLED")),
+			FName(TEXT("TMobileBasePassPSTLightMapPolicyLQLOCAL_LIGHTS_DISABLED")),
+			FName(TEXT("TMobileBasePassVSTLightMapPolicyLQ")),
 
 			FName(TEXT("TBasePassVSFCachedVolumeIndirectLightingPolicy")),
 			FName(TEXT("TBasePassPSFCachedVolumeIndirectLightingPolicy")),
diff --git a/Engine/Source/Runtime/RenderCore/Private/Shader.cpp b/Engine/Source/Runtime/RenderCore/Private/Shader.cpp
--- a/Engine/Source/Runtime/RenderCore/Private/Shader.cpp
+++ b/Engine/Source/Runtime/RenderCore/Private/Shader.cpp
@@ -1868,6 +1868,11 @@ void ShaderMapAppendKey(EShaderPlatform Platform, FShaderKeyGenerator& KeyGen)
 
 	if (IsMobilePlatform(Platform))
 	{
+		{
+			KeyGen.AppendSeparator();
+			KeyGen.Append(IsMobileHDR() ? TEXT("HDR") : TEXT("LDR"));
+		}
+		
 		{
 			static IConsoleVariable* CVar = IConsoleManager::Get().FindConsoleVariable(TEXT("r.Mobile.DisableVertexFog"));
 			if (CVar && CVar->GetInt() != 0)
diff --git a/Engine/Source/Runtime/Renderer/Private/LightMapRendering.cpp b/Engine/Source/Runtime/Renderer/Private/LightMapRendering.cpp
--- a/Engine/Source/Runtime/Renderer/Private/LightMapRendering.cpp
+++ b/Engine/Source/Runtime/Renderer/Private/LightMapRendering.cpp
@@ -155,7 +155,6 @@ void FCachedPointIndirectLightingPolicy::ModifyCompilationEnvironment(const FMat
 void FMobileDirectionalLightAndCSMPolicy::ModifyCompilationEnvironment(const FMaterialShaderPermutationParameters& Parameters, FShaderCompilerEnvironment& OutEnvironment)
 {
 	OutEnvironment.SetDefine(TEXT("DIRECTIONAL_LIGHT_CSM"), TEXT("1"));
-	OutEnvironment.SetDefine(TEXT("MOBILE_USE_CSM_BRANCH"), MobileUseCSMShaderBranch());
 }
 
 bool FMobileDirectionalLightAndCSMPolicy::ShouldCompilePermutation(const FMeshMaterialShaderPermutationParameters& Parameters)
@@ -164,6 +163,12 @@ bool FMobileDirectionalLightAndCSMPolicy::ShouldCompilePermutation(const FMeshMa
 	{
 		return false;
 	}
+
+	if (!IsStaticLightingAllowed() && MobileUseCSMShaderBranch())
+	{
+		// CSM will be handled by FNoLightMapPolicy with a branch inside shader
+		return false;
+	}
 	
 	return Parameters.MaterialParameters.ShadingModels.IsLit() && !IsTranslucentBlendMode(Parameters.MaterialParameters) 
 		&& (!IsStaticLightingAllowed() || FReadOnlyCVARCache::MobileEnableStaticAndCSMShadowReceivers());
@@ -792,13 +797,14 @@ TGlobalResource< FEmptyIndirectLightingCacheUniformBuffer > GEmptyIndirectLighti
 
 bool FNoLightMapPolicy::ShouldCompilePermutation(const FMeshMaterialShaderPermutationParameters& Parameters)
 {
-	if (IsMobilePlatform(Parameters.Platform))
-	{
-		return MobileUsesNoLightMapPermutation(Parameters);
-	}
 	return true;
 }
 
 void FNoLightMapPolicy::ModifyCompilationEnvironment(const FMaterialShaderPermutationParameters& Parameters, FShaderCompilerEnvironment& OutEnvironment)
 {
+	if (IsMobilePlatform(Parameters.Platform) && !IsStaticLightingAllowed() && MobileUseCSMShaderBranch())
+	{
+		OutEnvironment.SetDefine(TEXT("DIRECTIONAL_LIGHT_CSM"), TEXT("1"));
+		OutEnvironment.SetDefine(TEXT("MOBILE_USE_CSM_BRANCH"), TEXT("1"));
+	}
 }
diff --git a/Engine/Source/Runtime/Renderer/Private/MobileBasePass.cpp b/Engine/Source/Runtime/Renderer/Private/MobileBasePass.cpp
--- a/Engine/Source/Runtime/Renderer/Private/MobileBasePass.cpp
+++ b/Engine/Source/Runtime/Renderer/Private/MobileBasePass.cpp
@@ -114,47 +114,16 @@ void SetMobileBasePassDepthState(FMeshPassProcessorRenderState& DrawRenderState,
 	DrawRenderState.SetStencilRef(StencilValue); 
 }
 
-bool MobileUsesNoLightMapPermutation(const FMeshMaterialShaderPermutationParameters& Parameters)
-{
-	const bool bAllowStaticLighting = IsStaticLightingAllowed();
-	const bool bIsLitMaterial = Parameters.MaterialParameters.ShadingModels.IsLit();
-	const bool bDeferredShading = IsMobileDeferredShadingEnabled(Parameters.Platform);
-
-	if (!bDeferredShading && !bAllowStaticLighting && bIsLitMaterial && 
-		!IsTranslucentBlendMode(Parameters.MaterialParameters) &&
-		!Parameters.MaterialParameters.ShadingModels.HasShadingModel(MSM_SingleLayerWater))
-	{
-		// We don't need NoLightMap permutation if CSM shader can handle no-CSM case with a branch inside shader
-		return !MobileUseCSMShaderBranch();
-	}
-		
-	return true;
-}
-
 template <ELightMapPolicyType Policy, EMobileLocalLightSetting LocalLightSetting, EMobileTranslucentColorTransmittanceMode ThinTranslucencyFallback>
-static void AddMobileBasePassPixelShaderTypes(FMaterialShaderTypes& ShaderTypes, const bool bIsMobileHDR, const bool bEnableSkyLight)
+static void AddMobileBasePassPixelShaderTypes(FMaterialShaderTypes& ShaderTypes, const bool bEnableSkyLight)
 {
-	if (bIsMobileHDR)
+	if (bEnableSkyLight)
 	{
-		if (bEnableSkyLight)
-		{
-			ShaderTypes.AddShaderType<TMobileBasePassPS<TUniformLightMapPolicy<Policy>, HDR_LINEAR_64, true, LocalLightSetting, ThinTranslucencyFallback>>();
-		}
-		else
-		{
-			ShaderTypes.AddShaderType<TMobileBasePassPS<TUniformLightMapPolicy<Policy>, HDR_LINEAR_64, false, LocalLightSetting, ThinTranslucencyFallback>>();
-		}
+		ShaderTypes.AddShaderType<TMobileBasePassPS<TUniformLightMapPolicy<Policy>, true, LocalLightSetting, ThinTranslucencyFallback>>();
 	}
 	else
 	{
-		if (bEnableSkyLight)
-		{
-			ShaderTypes.AddShaderType<TMobileBasePassPS<TUniformLightMapPolicy<Policy>, LDR_GAMMA_32, true, LocalLightSetting, ThinTranslucencyFallback>>();
-		}
-		else
-		{
-			ShaderTypes.AddShaderType<TMobileBasePassPS<TUniformLightMapPolicy<Policy>, LDR_GAMMA_32, false, LocalLightSetting, ThinTranslucencyFallback>>();
-		}
+		ShaderTypes.AddShaderType<TMobileBasePassPS<TUniformLightMapPolicy<Policy>, false, LocalLightSetting, ThinTranslucencyFallback>>();
 	}
 }
 
@@ -171,24 +140,16 @@ bool GetUniformMobileBasePassShaders(
 	using FVertexShaderType = TMobileBasePassVSPolicyParamType<FUniformLightMapPolicy>;
 	using FPixelShaderType = TMobileBasePassPSPolicyParamType<FUniformLightMapPolicy>;
 
-	const bool bIsMobileHDR = IsMobileHDR();
 	FMaterialShaderTypes ShaderTypes;
-	if (bIsMobileHDR)
-	{
-		ShaderTypes.AddShaderType<TMobileBasePassVS<TUniformLightMapPolicy<Policy>, HDR_LINEAR_64>>();	
-	}
-	else
-	{
-		ShaderTypes.AddShaderType<TMobileBasePassVS<TUniformLightMapPolicy<Policy>, LDR_GAMMA_32>>();			
-	}
+	ShaderTypes.AddShaderType<TMobileBasePassVS<TUniformLightMapPolicy<Policy>>>();	
 
 	switch (ColoredTransmittanceFallback)
 	{
 	default:
 	case EMobileTranslucentColorTransmittanceMode::DEFAULT:
-		AddMobileBasePassPixelShaderTypes<Policy, LocalLightSetting, EMobileTranslucentColorTransmittanceMode::DEFAULT>(ShaderTypes, bIsMobileHDR, bEnableSkyLight); break;
+		AddMobileBasePassPixelShaderTypes<Policy, LocalLightSetting, EMobileTranslucentColorTransmittanceMode::DEFAULT>(ShaderTypes, bEnableSkyLight); break;
 	case EMobileTranslucentColorTransmittanceMode::SINGLE_SRC_BLENDING:
-		AddMobileBasePassPixelShaderTypes<Policy, LocalLightSetting, EMobileTranslucentColorTransmittanceMode::SINGLE_SRC_BLENDING>(ShaderTypes, bIsMobileHDR, bEnableSkyLight); break;
+		AddMobileBasePassPixelShaderTypes<Policy, LocalLightSetting, EMobileTranslucentColorTransmittanceMode::SINGLE_SRC_BLENDING>(ShaderTypes, bEnableSkyLight); break;
 	}
 
 	FMaterialShaders Shaders;
@@ -357,20 +318,21 @@ ELightMapPolicyType MobileBasePass::SelectMeshLightmapPolicy(
 		
 		if (!IsStaticLightingAllowed())
 		{
- 			if (!bIsTranslucent)
-			{
-				// Whether to use a single CSM permutation with a branch in the shader
-				bPrimReceivesCSM |= MobileUseCSMShaderBranch();
-			}
-			
 			// no precomputed lighting
-			if (!bPrimReceivesCSM || bUsesDeferredShading)
+			if (bUsesDeferredShading)
 			{
 				SelectedLightmapPolicy = LMP_NO_LIGHTMAP;
 			}
 			else
 			{
-				SelectedLightmapPolicy = LMP_MOBILE_DIRECTIONAL_LIGHT_CSM;				
+				if (!bPrimReceivesCSM || MobileUseCSMShaderBranch())
+				{
+					SelectedLightmapPolicy = LMP_NO_LIGHTMAP;
+				}
+				else
+				{
+					SelectedLightmapPolicy = LMP_MOBILE_DIRECTIONAL_LIGHT_CSM;
+				}
 			}
 		}
 		else
@@ -448,12 +410,9 @@ static FMobileLightMapPolicyTypeList GetUniformLightMapPolicyTypeForPSOCollectio
 	{
 		if (!IsStaticLightingAllowed())
 		{
-			if (bUsesDeferredShading || !MobileUseCSMShaderBranch())
-			{
-				Result.Add(LMP_NO_LIGHTMAP);
-			}
+			Result.Add(LMP_NO_LIGHTMAP);
 						
-			if (!bUsesDeferredShading)
+			if (!bUsesDeferredShading && !MobileUseCSMShaderBranch())
 			{
 				// permutation that can receive CSM
 				Result.Add(LMP_MOBILE_DIRECTIONAL_LIGHT_CSM);
diff --git a/Engine/Source/Runtime/Renderer/Private/MobileBasePassRendering.cpp b/Engine/Source/Runtime/Renderer/Private/MobileBasePassRendering.cpp
--- a/Engine/Source/Runtime/Renderer/Private/MobileBasePassRendering.cpp
+++ b/Engine/Source/Runtime/Renderer/Private/MobileBasePassRendering.cpp
@@ -128,31 +128,15 @@ EMobileTranslucentColorTransmittanceMode MobileActiveTranslucentColorTransmittan
 	return bExplicitDefaultMode ? DefaultMode : EMobileTranslucentColorTransmittanceMode::DEFAULT;
 }
 
-bool ShouldCacheShaderByPlatformAndOutputFormat(EShaderPlatform Platform, EOutputFormat OutputFormat)
-{
-	bool bSupportsMobileHDR = IsMobileHDR();
-	bool bShaderUsesLDR = (OutputFormat == LDR_GAMMA_32);
-
-	// only cache this shader if the LDR/HDR output matches what we currently support.  IsMobileHDR can't change, so we don't need
-	// the LDR shaders if we are doing HDR, and vice-versa.
-	return (bShaderUsesLDR && !bSupportsMobileHDR) || (!bShaderUsesLDR && bSupportsMobileHDR);
-}
-
 #define IMPLEMENT_MOBILE_SHADING_BASEPASS_LIGHTMAPPED_VERTEX_SHADER_TYPE(LightMapPolicyType,LightMapPolicyName) \
-	typedef TMobileBasePassVS< LightMapPolicyType, LDR_GAMMA_32 > TMobileBasePassVS##LightMapPolicyName##LDRGamma32; \
-	typedef TMobileBasePassVS< LightMapPolicyType, HDR_LINEAR_64 > TMobileBasePassVS##LightMapPolicyName##HDRLinear64; \
-	IMPLEMENT_MATERIAL_SHADER_TYPE(template<>, TMobileBasePassVS##LightMapPolicyName##LDRGamma32, TEXT("/Engine/Private/MobileBasePassVertexShader.usf"), TEXT("Main"), SF_Vertex); \
-	IMPLEMENT_MATERIAL_SHADER_TYPE(template<>, TMobileBasePassVS##LightMapPolicyName##HDRLinear64, TEXT("/Engine/Private/MobileBasePassVertexShader.usf"), TEXT("Main"), SF_Vertex);
+	typedef TMobileBasePassVS< LightMapPolicyType > TMobileBasePassVS##LightMapPolicyName; \
+	IMPLEMENT_MATERIAL_SHADER_TYPE(template<>, TMobileBasePassVS##LightMapPolicyName, TEXT("/Engine/Private/MobileBasePassVertexShader.usf"), TEXT("Main"), SF_Vertex); \
 
 #define IMPLEMENT_MOBILE_SHADING_BASEPASS_LIGHTMAPPED_PIXEL_SHADER_TYPE2(LightMapPolicyType, LightMapPolicyName, LocalLightSetting, ThinTranslucencyEnum, ThinTranslucencyName) \
-	typedef TMobileBasePassPS< LightMapPolicyType, LDR_GAMMA_32, false, LocalLightSetting, EMobileTranslucentColorTransmittanceMode::ThinTranslucencyEnum > TMobileBasePassPS##LightMapPolicyName##LDRGamma32##LocalLightSetting##ThinTranslucencyName; \
-	typedef TMobileBasePassPS< LightMapPolicyType, HDR_LINEAR_64, false, LocalLightSetting, EMobileTranslucentColorTransmittanceMode::ThinTranslucencyEnum > TMobileBasePassPS##LightMapPolicyName##HDRLinear64##LocalLightSetting##ThinTranslucencyName; \
-	typedef TMobileBasePassPS< LightMapPolicyType, LDR_GAMMA_32, true, LocalLightSetting, EMobileTranslucentColorTransmittanceMode::ThinTranslucencyEnum > TMobileBasePassPS##LightMapPolicyName##LDRGamma32##Skylight##LocalLightSetting##ThinTranslucencyName; \
-	typedef TMobileBasePassPS< LightMapPolicyType, HDR_LINEAR_64, true, LocalLightSetting, EMobileTranslucentColorTransmittanceMode::ThinTranslucencyEnum > TMobileBasePassPS##LightMapPolicyName##HDRLinear64##Skylight##LocalLightSetting##ThinTranslucencyName; \
-	IMPLEMENT_MATERIAL_SHADER_TYPE(template<>, TMobileBasePassPS##LightMapPolicyName##LDRGamma32##LocalLightSetting##ThinTranslucencyName, TEXT("/Engine/Private/MobileBasePassPixelShader.usf"), TEXT("Main"), SF_Pixel); \
-	IMPLEMENT_MATERIAL_SHADER_TYPE(template<>, TMobileBasePassPS##LightMapPolicyName##HDRLinear64##LocalLightSetting##ThinTranslucencyName, TEXT("/Engine/Private/MobileBasePassPixelShader.usf"), TEXT("Main"), SF_Pixel); \
-	IMPLEMENT_MATERIAL_SHADER_TYPE(template<>, TMobileBasePassPS##LightMapPolicyName##LDRGamma32##Skylight##LocalLightSetting##ThinTranslucencyName, TEXT("/Engine/Private/MobileBasePassPixelShader.usf"), TEXT("Main"), SF_Pixel); \
-	IMPLEMENT_MATERIAL_SHADER_TYPE(template<>, TMobileBasePassPS##LightMapPolicyName##HDRLinear64##Skylight##LocalLightSetting##ThinTranslucencyName, TEXT("/Engine/Private/MobileBasePassPixelShader.usf"), TEXT("Main"), SF_Pixel);
+	typedef TMobileBasePassPS< LightMapPolicyType, false, LocalLightSetting, EMobileTranslucentColorTransmittanceMode::ThinTranslucencyEnum > TMobileBasePassPS##LightMapPolicyName##LocalLightSetting##ThinTranslucencyName; \
+	typedef TMobileBasePassPS< LightMapPolicyType, true, LocalLightSetting, EMobileTranslucentColorTransmittanceMode::ThinTranslucencyEnum > TMobileBasePassPS##LightMapPolicyName##Skylight##LocalLightSetting##ThinTranslucencyName; \
+	IMPLEMENT_MATERIAL_SHADER_TYPE(template<>, TMobileBasePassPS##LightMapPolicyName##LocalLightSetting##ThinTranslucencyName, TEXT("/Engine/Private/MobileBasePassPixelShader.usf"), TEXT("Main"), SF_Pixel); \
+	IMPLEMENT_MATERIAL_SHADER_TYPE(template<>, TMobileBasePassPS##LightMapPolicyName##Skylight##LocalLightSetting##ThinTranslucencyName, TEXT("/Engine/Private/MobileBasePassPixelShader.usf"), TEXT("Main"), SF_Pixel);
 
 #define IMPLEMENT_MOBILE_SHADING_BASEPASS_LIGHTMAPPED_PIXEL_SHADER_TYPE(LightMapPolicyType, LightMapPolicyName, LocalLightSetting) \
 	IMPLEMENT_MOBILE_SHADING_BASEPASS_LIGHTMAPPED_PIXEL_SHADER_TYPE2(LightMapPolicyType, LightMapPolicyName, LocalLightSetting, DEFAULT,) \
@@ -200,12 +184,13 @@ bool ShouldCacheShaderForColorTransmittanceFallback(const FMaterialShaderPermuta
 }
 
 // shared defines for mobile base pass VS and PS
-void MobileBasePassModifyCompilationEnvironment(const FMaterialShaderPermutationParameters& Parameters, FShaderCompilerEnvironment& OutEnvironment, EOutputFormat OutputFormat)
+void MobileBasePassModifyCompilationEnvironment(const FMaterialShaderPermutationParameters& Parameters, FShaderCompilerEnvironment& OutEnvironment)
 {
 	static auto* MobileUseHWsRGBEncodingCVAR = IConsoleManager::Get().FindTConsoleVariableDataInt(TEXT("r.Mobile.UseHWsRGBEncoding"));
 	const bool bMobileUseHWsRGBEncoding = (MobileUseHWsRGBEncodingCVAR && MobileUseHWsRGBEncodingCVAR->GetValueOnAnyThread() == 1);
-	OutEnvironment.SetDefine( TEXT("OUTPUT_GAMMA_SPACE"), OutputFormat == LDR_GAMMA_32 && !bMobileUseHWsRGBEncoding);
-	OutEnvironment.SetDefine( TEXT("OUTPUT_MOBILE_HDR"), OutputFormat == HDR_LINEAR_64 ? 1u : 0u);
+	const bool bMobileHDR = IsMobileHDR();
+	OutEnvironment.SetDefine( TEXT("OUTPUT_GAMMA_SPACE"), !bMobileHDR && !bMobileUseHWsRGBEncoding);
+	OutEnvironment.SetDefine( TEXT("OUTPUT_MOBILE_HDR"), bMobileHDR ? 1u : 0u);
 	
 	const bool bTranslucentMaterial = 
 		IsTranslucentBlendMode(Parameters.MaterialParameters) ||
diff --git a/Engine/Source/Runtime/Renderer/Private/MobileBasePassRendering.h b/Engine/Source/Runtime/Renderer/Private/MobileBasePassRendering.h
--- a/Engine/Source/Runtime/Renderer/Private/MobileBasePassRendering.h
+++ b/Engine/Source/Runtime/Renderer/Private/MobileBasePassRendering.h
@@ -105,12 +105,6 @@ extern void SetupMobileSkyReflectionUniformParameters(
 class FPlanarReflectionSceneProxy;
 class FScene;
 
-enum EOutputFormat
-{
-	LDR_GAMMA_32,
-	HDR_LINEAR_64,
-};
-
 EMobileLocalLightSetting GetMobileForwardLocalLightSetting(EShaderPlatform ShaderPlatform);
 
 enum class EMobileTranslucentColorTransmittanceMode
@@ -127,9 +121,8 @@ bool MaterialRequiresColorTransmittanceBlending(const FMaterial& MaterialResourc
 bool MaterialRequiresColorTransmittanceBlending(const FMaterialShaderParameters& MaterialParameters);
 bool ShouldCacheShaderForColorTransmittanceFallback(const FMaterialShaderPermutationParameters& Parameters, EMobileTranslucentColorTransmittanceMode TranslucentColorTransmittanceFallback);
 
-bool ShouldCacheShaderByPlatformAndOutputFormat(EShaderPlatform Platform, EOutputFormat OutputFormat);
 // shared defines for mobile base pass VS and PS
-void MobileBasePassModifyCompilationEnvironment(const FMaterialShaderPermutationParameters& Parameters, FShaderCompilerEnvironment& OutEnvironment, EOutputFormat OutputFormat);
+void MobileBasePassModifyCompilationEnvironment(const FMaterialShaderPermutationParameters& Parameters, FShaderCompilerEnvironment& OutEnvironment);
 
 template<typename LightMapPolicyType>
 class TMobileBasePassShaderElementData : public FMeshMaterialShaderElementData
@@ -215,7 +208,7 @@ public:
 	}
 };
 
-template< typename LightMapPolicyType, EOutputFormat OutputFormat >
+template< typename LightMapPolicyType >
 class TMobileBasePassVS : public TMobileBasePassVSBaseType<LightMapPolicyType>
 {
 	DECLARE_SHADER_TYPE(TMobileBasePassVS,MeshMaterial);
@@ -223,12 +216,12 @@ public:
 	
 	static bool ShouldCompilePermutation(const FMeshMaterialShaderPermutationParameters& Parameters)
 	{		
-		return TMobileBasePassVSBaseType<LightMapPolicyType>::ShouldCompilePermutation(Parameters) && ShouldCacheShaderByPlatformAndOutputFormat(Parameters.Platform, OutputFormat);
+		return TMobileBasePassVSBaseType<LightMapPolicyType>::ShouldCompilePermutation(Parameters);
 	}
 
 	static void ModifyCompilationEnvironment(const FMaterialShaderPermutationParameters& Parameters, FShaderCompilerEnvironment& OutEnvironment)
 	{
-		MobileBasePassModifyCompilationEnvironment(Parameters, OutEnvironment, OutputFormat);
+		MobileBasePassModifyCompilationEnvironment(Parameters, OutEnvironment);
 		TMobileBasePassVSBaseType<LightMapPolicyType>::ModifyCompilationEnvironment(Parameters, OutEnvironment);
 	}
 	
@@ -361,7 +354,7 @@ namespace MobileBasePass
 	}
 };
 
-template< typename LightMapPolicyType, EOutputFormat OutputFormat, bool bEnableSkyLight, EMobileLocalLightSetting LocalLightSetting, EMobileTranslucentColorTransmittanceMode TranslucentColorTransmittanceFallback = EMobileTranslucentColorTransmittanceMode::DEFAULT>
+template< typename LightMapPolicyType, bool bEnableSkyLight, EMobileLocalLightSetting LocalLightSetting, EMobileTranslucentColorTransmittanceMode TranslucentColorTransmittanceFallback = EMobileTranslucentColorTransmittanceMode::DEFAULT>
 class TMobileBasePassPS : public TMobileBasePassPSBaseType<LightMapPolicyType>
 {
 	DECLARE_SHADER_TYPE(TMobileBasePassPS,MeshMaterial);
@@ -398,7 +391,6 @@ public:
 		const bool bShouldCacheByLocalLights = !bEnableLocalLights || (bIsLit && (SupportedLocalLightsType == LocalLightSetting));
 
 		return TMobileBasePassPSBaseType<LightMapPolicyType>::ShouldCompilePermutation(Parameters) && 
-				ShouldCacheShaderByPlatformAndOutputFormat(Parameters.Platform, OutputFormat) && 
 				bShouldCacheBySkylight && 
 				bShouldCacheByLocalLights && 
 				ShouldCacheShaderForColorTransmittanceFallback(Parameters, TranslucentColorTransmittanceFallback);
@@ -406,7 +398,7 @@ public:
 	
 	static void ModifyCompilationEnvironment(const FMaterialShaderPermutationParameters& Parameters, FShaderCompilerEnvironment& OutEnvironment)
 	{		
-		MobileBasePassModifyCompilationEnvironment(Parameters, OutEnvironment, OutputFormat);
+		MobileBasePassModifyCompilationEnvironment(Parameters, OutEnvironment);
 
 		const bool bMobileUsesShadowMaskTexture = MobileUsesShadowMaskTexture(Parameters.Platform);
 		const bool bEnableClusteredReflections = MobileForwardEnableClusteredReflections(Parameters.Platform);
diff --git a/Engine/Source/Runtime/Renderer/Private/WaterInfoTextureRendering.cpp b/Engine/Source/Runtime/Renderer/Private/WaterInfoTextureRendering.cpp
--- a/Engine/Source/Runtime/Renderer/Private/WaterInfoTextureRendering.cpp
+++ b/Engine/Source/Runtime/Renderer/Private/WaterInfoTextureRendering.cpp
@@ -234,8 +234,8 @@ void FWaterInfoTexturePassMeshProcessor::CollectPSOInitializers(const FSceneText
 				TMobileBasePassPSPolicyParamType<FUniformLightMapPolicy>> PassShaders;
 
 			FMaterialShaderTypes ShaderTypes;
-			ShaderTypes.AddShaderType<TMobileBasePassVS<TUniformLightMapPolicy<LMP_NO_LIGHTMAP>, HDR_LINEAR_64>>();
-			ShaderTypes.AddShaderType<TMobileBasePassPS<TUniformLightMapPolicy<LMP_NO_LIGHTMAP>, HDR_LINEAR_64, false, LOCAL_LIGHTS_DISABLED>>();
+			ShaderTypes.AddShaderType<TMobileBasePassVS<TUniformLightMapPolicy<LMP_NO_LIGHTMAP>>>();
+			ShaderTypes.AddShaderType<TMobileBasePassPS<TUniformLightMapPolicy<LMP_NO_LIGHTMAP>, false, LOCAL_LIGHTS_DISABLED>>();
 			FMaterialShaders Shaders;
 			if (!Material.TryGetShaders(ShaderTypes, VertexFactoryData.VertexFactoryType, Shaders))
 			{
@@ -322,8 +322,8 @@ bool FWaterInfoTexturePassMeshProcessor::TryAddMeshBatch(const FMeshBatch& RESTR
 				TMobileBasePassPSPolicyParamType<FUniformLightMapPolicy>> PassShaders;
 
 			FMaterialShaderTypes ShaderTypes;
-			ShaderTypes.AddShaderType<TMobileBasePassVS<TUniformLightMapPolicy<LMP_NO_LIGHTMAP>, HDR_LINEAR_64>>();
-			ShaderTypes.AddShaderType<TMobileBasePassPS<TUniformLightMapPolicy<LMP_NO_LIGHTMAP>, HDR_LINEAR_64, false, LOCAL_LIGHTS_DISABLED>>();
+			ShaderTypes.AddShaderType<TMobileBasePassVS<TUniformLightMapPolicy<LMP_NO_LIGHTMAP>>>();
+			ShaderTypes.AddShaderType<TMobileBasePassPS<TUniformLightMapPolicy<LMP_NO_LIGHTMAP>, false, LOCAL_LIGHTS_DISABLED>>();
 			FMaterialShaders Shaders;
 			if (!Material.TryGetShaders(ShaderTypes, VertexFactory->GetType(), Shaders))
 			{
