Create the correct type of material expression when you promote the Shader Model output pin to a parameter.

[RN] minor, Rendering / Materials and Shaders

diff --git a/Engine/Source/Editor/MaterialEditor/Private/MaterialEditor.cpp b/Engine/Source/Editor/MaterialEditor/Private/MaterialEditor.cpp
--- a/Engine/Source/Editor/MaterialEditor/Private/MaterialEditor.cpp
+++ b/Engine/Source/Editor/MaterialEditor/Private/MaterialEditor.cpp
@@ -96,6 +96,7 @@
 #include "Materials/MaterialExpressionDoubleVectorParameter.h"
 #include "Materials/MaterialExpressionStaticBoolParameter.h"
 #include "Materials/MaterialExpressionCustomOutput.h"
+#include "Materials/MaterialExpressionShadingModel.h"
 #include "Materials/MaterialFunction.h"
 #include "Materials/MaterialFunctionInstance.h"
 #include "Materials/MaterialParameterCollection.h"
@@ -5106,7 +5107,6 @@ UClass* FMaterialEditor::GetOnPromoteToParameterClass(const UEdGraphPin* TargetP
 			case MP_AmbientOcclusion:
 			case MP_Refraction:
 			case MP_PixelDepthOffset:
-			case MP_ShadingModel:
 			case MP_OpacityMask:
 			case MP_SurfaceThickness:
 			case MP_Displacement:
@@ -5121,6 +5121,9 @@ UClass* FMaterialEditor::GetOnPromoteToParameterClass(const UEdGraphPin* TargetP
 			case MP_Tangent:
 				return UMaterialExpressionVectorParameter::StaticClass();
 
+			case MP_ShadingModel:
+				return UMaterialExpressionShadingModel::StaticClass();
+
 			case MP_FrontMaterial:
 				return nullptr;
 
