Don't display material parameters in the Instance Editor if they are marked as using custom primitive data.

- CPD parameters pull their data from the primitive, not from the material so displaying these parameters for the user to edit was confusing since they wouldn't do anything.
- Cache off transient state on the UDEditorParameterValue on whether or not this is a CPD parameter.
- Pull that info from the parameter metadata if PrimitiveDataIndex != INDEX_NONE
- Adds helper function UsesCustomPrimitiveData(UDEditorParameterValue* Parameter) that determines if the parameter uses CPD.

#tests Creating a material with a CPD parameter, making a MI from it, and verifying the parameter did not show up in the instance editor
[RN] minor, Rendering / Materials and Shaders

diff --git a/Engine/Source/Editor/MaterialEditor/Private/MaterialEditorInstanceDetailCustomization.cpp b/Engine/Source/Editor/MaterialEditor/Private/MaterialEditorInstanceDetailCustomization.cpp
--- a/Engine/Source/Editor/MaterialEditor/Private/MaterialEditorInstanceDetailCustomization.cpp
+++ b/Engine/Source/Editor/MaterialEditor/Private/MaterialEditorInstanceDetailCustomization.cpp
@@ -384,7 +384,7 @@ void FMaterialInstanceParameterDetails::CreateGroupsWidget(TSharedRef<IPropertyH
 			for (int32 ParamIdx = 0; ParamIdx < ParameterGroup.Parameters.Num() && !bCreateGroup; ++ParamIdx)
 			{
 				UDEditorParameterValue* Parameter = ParameterGroup.Parameters[ParamIdx];
-				const bool bIsVisible = MaterialEditorInstance->VisibleExpressions.Contains(Parameter->ParameterInfo);
+				const bool bIsVisible = MaterialEditorInstance->VisibleExpressions.Contains(Parameter->ParameterInfo) && !FMaterialPropertyHelpers::UsesCustomPrimitiveData(Parameter);
 				bCreateGroup = bIsVisible && (!MaterialEditorInstance->bShowOnlyOverrides || FMaterialPropertyHelpers::IsOverriddenExpression(Parameter));
 			}
 		
@@ -513,6 +513,14 @@ void FMaterialInstanceParameterDetails::CreateSingleGroupWidget(FEditorParameter
 			UDEditorSparseVolumeTextureParameterValue* SparseVolumeTextureParam = Cast<UDEditorSparseVolumeTextureParameterValue>(Parameter);
 			UDEditorVectorParameterValue* VectorParam = Cast<UDEditorVectorParameterValue>(Parameter);
 
+			// Don't display custom primitive data parameters in the details panel.
+			// This data is pulled from the primitive and can't be changed on the material.
+			if ((VectorParam && VectorParam->bUseCustomPrimitiveData) ||
+				(ScalarParam && ScalarParam->bUseCustomPrimitiveData))
+			{
+				continue;
+			}
+
 			if (Parameter->ParameterInfo.Association == EMaterialParameterAssociation::GlobalParameter || bForceShowParam)
 			{
 				if (VectorParam && VectorParam->bIsUsedAsChannelMask)
diff --git a/Engine/Source/Editor/MaterialEditor/Private/MaterialPropertyHelpers.cpp b/Engine/Source/Editor/MaterialEditor/Private/MaterialPropertyHelpers.cpp
--- a/Engine/Source/Editor/MaterialEditor/Private/MaterialPropertyHelpers.cpp
+++ b/Engine/Source/Editor/MaterialEditor/Private/MaterialPropertyHelpers.cpp
@@ -548,6 +548,20 @@ ECheckBoxState FMaterialPropertyHelpers::IsOverriddenExpressionCheckbox(UDEditor
 	return IsOverriddenExpression(Parameter) ? ECheckBoxState::Checked : ECheckBoxState::Unchecked;
 }
 
+bool FMaterialPropertyHelpers::UsesCustomPrimitiveData(UDEditorParameterValue* Parameter)
+{
+	if (UDEditorScalarParameterValue* ScalarParameter = Cast<UDEditorScalarParameterValue>(Parameter))
+	{
+		return ScalarParameter->bUseCustomPrimitiveData;
+	}
+	else if (UDEditorVectorParameterValue* VectorParameter = Cast<UDEditorVectorParameterValue>(Parameter))
+	{
+		return VectorParameter->bUseCustomPrimitiveData;
+	}
+
+	return false;
+}
+
 void FMaterialPropertyHelpers::OnOverrideParameter(bool NewValue, class UDEditorParameterValue* Parameter, UMaterialEditorInstanceConstant* MaterialEditorInstance)
 {
 	// If the material instance disallows the creation of new shader permutations, prevent overriding the static parameter.
diff --git a/Engine/Source/Editor/MaterialEditor/Public/MaterialPropertyHelpers.h b/Engine/Source/Editor/MaterialEditor/Public/MaterialPropertyHelpers.h
--- a/Engine/Source/Editor/MaterialEditor/Public/MaterialPropertyHelpers.h
+++ b/Engine/Source/Editor/MaterialEditor/Public/MaterialPropertyHelpers.h
@@ -154,6 +154,7 @@ public:
 	static bool IsOverriddenExpression(TObjectPtr<UDEditorParameterValue> Parameter) { return IsOverriddenExpression(Parameter.Get()); }
 	static bool IsOverriddenExpression(TWeakObjectPtr<UDEditorParameterValue> Parameter) { return IsOverriddenExpression(Parameter.Get()); }
 	static ECheckBoxState IsOverriddenExpressionCheckbox(UDEditorParameterValue* Parameter);
+	static bool UsesCustomPrimitiveData(UDEditorParameterValue* Parameter);
 
 	/** Gets the expression description of this parameter from the base material */
 	static	FText GetParameterExpressionDescription(UDEditorParameterValue* Parameter, UObject* MaterialEditorInstance);
diff --git a/Engine/Source/Editor/UnrealEd/Classes/MaterialEditor/DEditorScalarParameterValue.h b/Engine/Source/Editor/UnrealEd/Classes/MaterialEditor/DEditorScalarParameterValue.h
--- a/Engine/Source/Editor/UnrealEd/Classes/MaterialEditor/DEditorScalarParameterValue.h
+++ b/Engine/Source/Editor/UnrealEd/Classes/MaterialEditor/DEditorScalarParameterValue.h
@@ -38,6 +38,9 @@ public:
 	UPROPERTY(Transient)
 	FScalarParameterAtlasData AtlasData;
 
+	UPROPERTY(Transient)
+	bool bUseCustomPrimitiveData;
+
 	virtual FName GetDefaultGroupName() const override { return TEXT("Scalar Parameter Values"); }
 
 	virtual bool GetValue(FMaterialParameterMetadata& OutResult) const override
diff --git a/Engine/Source/Editor/UnrealEd/Classes/MaterialEditor/DEditorVectorParameterValue.h b/Engine/Source/Editor/UnrealEd/Classes/MaterialEditor/DEditorVectorParameterValue.h
--- a/Engine/Source/Editor/UnrealEd/Classes/MaterialEditor/DEditorVectorParameterValue.h
+++ b/Engine/Source/Editor/UnrealEd/Classes/MaterialEditor/DEditorVectorParameterValue.h
@@ -19,6 +19,9 @@ class UDEditorVectorParameterValue : public UDEditorParameterValue
 	UPROPERTY(Transient)
 	bool bIsUsedAsChannelMask;
 
+	UPROPERTY(Transient)
+	bool bUseCustomPrimitiveData;
+
 	UPROPERTY(Transient)
 	FParameterChannelNames ChannelNames;
 
diff --git a/Engine/Source/Editor/UnrealEd/Private/EditorParameters.cpp b/Engine/Source/Editor/UnrealEd/Private/EditorParameters.cpp
--- a/Engine/Source/Editor/UnrealEd/Private/EditorParameters.cpp
+++ b/Engine/Source/Editor/UnrealEd/Private/EditorParameters.cpp
@@ -94,6 +94,7 @@ static UDEditorParameterValue* CreateParameter_Scalar(UObject* Owner, const FMat
 		Parameter->AtlasData.bIsUsedAsAtlasPosition = Meta.bUsedAsAtlasPosition;
 		Parameter->AtlasData.Atlas = Meta.ScalarAtlas;
 		Parameter->AtlasData.Curve = Meta.ScalarCurve;
+		Parameter->bUseCustomPrimitiveData = Meta.PrimitiveDataIndex != INDEX_NONE;
 	}
 	return Parameter;
 }
@@ -105,6 +106,7 @@ static UDEditorParameterValue* CreateParameter_Vector(UObject* Owner, const FMat
 	{
 		Parameter->ParameterValue = Meta.Value.AsLinearColor();
 		Parameter->bIsUsedAsChannelMask = Meta.bUsedAsChannelMask;
+		Parameter->bUseCustomPrimitiveData = Meta.PrimitiveDataIndex != INDEX_NONE;
 		Parameter->ChannelNames = Meta.ChannelNames;
 	}
 	return Parameter;
