Substrate - preview tab now update when while material is edited, to show for instance the UI preview when editing Substrate UI material.

#rnx

diff --git a/Engine/Source/Editor/MaterialEditor/Private/MaterialEditor.cpp b/Engine/Source/Editor/MaterialEditor/Private/MaterialEditor.cpp
--- a/Engine/Source/Editor/MaterialEditor/Private/MaterialEditor.cpp
+++ b/Engine/Source/Editor/MaterialEditor/Private/MaterialEditor.cpp
@@ -6807,6 +6807,21 @@ void FMaterialEditor::UpdateMaterialAfterGraphChange()
 	Material->MaterialGraph->UpdatePinTypes();
 
 	UpdateSubstrateTopologyPreview();
+
+	if (Substrate::IsSubstrateEnabled())
+	{
+		// Update preview tab to adapt to the change of domain.
+		// This is because with Substrate, we do not receive any event that the material domain has changed from the UI combobox (no longer editable, set form the graph via RebuildShadingModelField).
+		if (ExpressionPreviewMaterial)
+		{
+			SetPreviewMaterial(ExpressionPreviewMaterial);
+		}
+		else
+		{
+			SetPreviewMaterial(Material);
+		}
+		UpdatePreviewViewportsVisibility();
+	}
 }
 
 void FMaterialEditor::MarkMaterialDirty()
