Make ShadingModel override on material instance only visible if parent material have a SubstrateShadingModel node.

diff --git a/Engine/Source/Editor/MaterialEditor/Private/MaterialEditorInstanceDetailCustomization.cpp b/Engine/Source/Editor/MaterialEditor/Private/MaterialEditorInstanceDetailCustomization.cpp
--- a/Engine/Source/Editor/MaterialEditor/Private/MaterialEditorInstanceDetailCustomization.cpp
+++ b/Engine/Source/Editor/MaterialEditor/Private/MaterialEditorInstanceDetailCustomization.cpp
@@ -52,6 +52,7 @@
 #include "IPropertyUtilities.h"
 #include "Engine/Texture.h"
 #include "HAL/PlatformApplicationMisc.h"
+#include "RenderUtils.h"
 
 #define LOCTEXT_NAMESPACE "MaterialInstanceEditor"
 
@@ -1458,7 +1459,8 @@ void FMaterialInstanceParameterDetails::CreateBasePropertyOverrideWidgets(IDetai
 		auto&& OverrideBoolEnabledMemberFn,
 		auto&& OverrideBoolChangedMemberFn,
 		auto&& IsResetPropertyVisibleLambda,
-		auto&& ResetPropertyHandlerLambda
+		auto&& ResetPropertyHandlerLambda,
+		auto&& OverrideAndVisibleMemberFn
 	)
 	{
 		TSharedPtr<IPropertyHandle> ValueProperty = BasePropertyOverridePropery->GetChildHandle(PropertyName);
@@ -1471,18 +1473,20 @@ void FMaterialInstanceParameterDetails::CreateBasePropertyOverrideWidgets(IDetai
 			.DisplayName(ValueProperty->GetPropertyDisplayName())
 			.ToolTip(bStaticParametersOverrideDisabled ? ParameterDisabledToolTipString : ValueProperty->GetToolTipText())
 			.EditCondition(OverrideBoolAttr, FOnBooleanValueChanged::CreateSP(this, OverrideBoolChangedMemberFn))
-			.Visibility(TAttribute<EVisibility>::Create(TAttribute<EVisibility>::FGetter::CreateSP(this, &FMaterialInstanceParameterDetails::IsOverriddenAndVisible, OverrideBoolAttr)))
+			.Visibility(TAttribute<EVisibility>::Create(TAttribute<EVisibility>::FGetter::CreateSP(this, OverrideAndVisibleMemberFn, OverrideBoolAttr)))
 			.OverrideResetToDefault(ResetPropertyOverride);
 	};
 
-#define CREATE_BASE_OVERRIDE_ROW_CUSTOM(PropertyName, PropertyVariableName, IsResetPropertyVisibleLambda, ResetPropertyHandlerLambda) \
+#define CREATE_BASE_OVERRIDE_ROW_CUSTOM(PropertyName, PropertyVariableName, IsResetPropertyVisibleLambda, ResetPropertyHandlerLambda, IsOverriddenAndVisibleFn) \
 	CreateBaseOverrideRow ( \
 		#PropertyVariableName, \
 		&FMaterialInstanceParameterDetails::Override ## PropertyName ## Enabled, \
 		&FMaterialInstanceParameterDetails::OnOverride ## PropertyName ## Changed, \
 		IsResetPropertyVisibleLambda, \
-		ResetPropertyHandlerLambda \
-	)
+		ResetPropertyHandlerLambda, \
+		IsOverriddenAndVisibleFn \
+		)
+
 #define CREATE_BASE_OVERRIDE_ROW(PropertyName, PropertyVariableName, ValueGetterName) \
 	CREATE_BASE_OVERRIDE_ROW_CUSTOM( \
 		PropertyName, \
@@ -1502,7 +1506,8 @@ void FMaterialInstanceParameterDetails::CreateBasePropertyOverrideWidgets(IDetai
 			{ \
 				MaterialEditorInstance->BasePropertyOverrides.PropertyVariableName = ParentMat->ValueGetterName(); \
 			} \
-		} \
+		}, \
+		&FMaterialInstanceParameterDetails::IsOverriddenAndVisible \
 	)
 #define CREATE_BASE_OVERRIDE_ROW_BASIC(PropertyName) \
 	CREATE_BASE_OVERRIDE_ROW(PropertyName, PropertyName, Get ## PropertyName)
@@ -1543,7 +1548,8 @@ void FMaterialInstanceParameterDetails::CreateBasePropertyOverrideWidgets(IDetai
 					MaterialEditorInstance->BasePropertyOverrides.ShadingModel = ParentMat->GetShadingModels().GetFirstShadingModel();
 				}
 			}
-		}
+		},
+		&FMaterialInstanceParameterDetails::IsOverriddenAndVisibleShadingModels
 	);
 	CREATE_BASE_OVERRIDE_ROW(TwoSided, TwoSided, IsTwoSided);
 	CREATE_BASE_OVERRIDE_ROW_BASIC_BOOL(IsThinSurface);
@@ -1574,6 +1580,22 @@ EVisibility FMaterialInstanceParameterDetails::IsOverriddenAndVisible(TAttribute
 	return bShouldBeVisible ? EVisibility::Visible : EVisibility::Collapsed;
 }
 
+EVisibility FMaterialInstanceParameterDetails::IsOverriddenAndVisibleShadingModels(TAttribute<bool> IsOverridden) const
+{
+	bool bShouldBeVisible = true;
+	if (MaterialEditorInstance->bShowOnlyOverrides)
+	{
+		bShouldBeVisible = IsOverridden.Get();
+	}
+	// If Substrate is enabled, only allows ShadingModel to be visible if the parent allows it
+	if (Substrate::IsSubstrateEnabled())
+	{
+		const UMaterial* ParentMaterial = MaterialEditorInstance->Parent ? MaterialEditorInstance->Parent->GetMaterial() : nullptr;
+		bShouldBeVisible = ParentMaterial ? ParentMaterial->SupportsShadingModelOverride() : false;
+	}
+	return bShouldBeVisible ? EVisibility::Visible : EVisibility::Collapsed;
+}
+
 /** Helper function used by some parameters to verify that they are allowed to be overridden. This
   * must be prevented if the source material instance disallows the creation of new static parameter
   * permutations as that would trigger a new shader creation. */
diff --git a/Engine/Source/Editor/MaterialEditor/Private/MaterialEditorInstanceDetailCustomization.h b/Engine/Source/Editor/MaterialEditor/Private/MaterialEditorInstanceDetailCustomization.h
--- a/Engine/Source/Editor/MaterialEditor/Private/MaterialEditorInstanceDetailCustomization.h
+++ b/Engine/Source/Editor/MaterialEditor/Private/MaterialEditorInstanceDetailCustomization.h
@@ -76,7 +76,7 @@ private:
 	/** Returns true if the refraction options should be displayed */
 	EVisibility ShouldShowMaterialRefractionSettings() const;
 
-	/** Returns true if the refraction options should be displayed */
+	/** Returns true if the Subsurface Profile options should be displayed */
 	EVisibility ShouldShowSubsurfaceProfile() const;
 
 	//Functions supporting copy/paste of entire parameter groups.
@@ -114,6 +114,7 @@ private:
 	void CreateBasePropertyOverrideWidgets(IDetailLayoutBuilder& DetailLayout, IDetailGroup& MaterialPropertyOverrideGroup);
 
 	EVisibility IsOverriddenAndVisible(TAttribute<bool> IsOverridden) const;
+	EVisibility IsOverriddenAndVisibleShadingModels(TAttribute<bool> IsOverridden) const;
 
 #define DECLARE_OVERRIDE_MEMBER_FUNCS(PropertyName) \
 	bool Override ## PropertyName ## Enabled() const; \
diff --git a/Engine/Source/Runtime/Engine/Private/Materials/Material.cpp b/Engine/Source/Runtime/Engine/Private/Materials/Material.cpp
--- a/Engine/Source/Runtime/Engine/Private/Materials/Material.cpp
+++ b/Engine/Source/Runtime/Engine/Private/Materials/Material.cpp
@@ -6110,6 +6110,24 @@ void UMaterial::GetAllExpressionsForCustomInterpolators(TArray<class UMaterialEx
 		}
 	}
 }
+
+bool UMaterial::SupportsShadingModelOverride() const
+{
+	// If the material contains a Substrate's SubstrateShadingModels node, then we can support shading model override
+	bool Out = true;
+	if (Substrate::IsSubstrateEnabled())
+	{
+		for (UMaterialExpression* Expression : GetExpressions())
+		{		
+			if (UMaterialExpressionSubstrateShadingModels* ShadingModelNode = Cast<UMaterialExpressionSubstrateShadingModels>(Expression))
+			{
+				return true;
+			}
+		}
+		Out = false;
+	}
+	return Out;
+}
 #endif // WITH_EDITORONLY_DATA
 
 #if WITH_EDITOR
diff --git a/Engine/Source/Runtime/Engine/Public/Materials/Material.h b/Engine/Source/Runtime/Engine/Public/Materials/Material.h
--- a/Engine/Source/Runtime/Engine/Public/Materials/Material.h
+++ b/Engine/Source/Runtime/Engine/Public/Materials/Material.h
@@ -1226,6 +1226,11 @@ public:
 	/** Checks to see if an input property should be active, based on the state of the material */
 	ENGINE_API virtual bool IsPropertyActive(EMaterialProperty InProperty) const override;
 
+#if WITH_EDITORONLY_DATA
+	/** Return true if the material support shading model override */
+	ENGINE_API bool SupportsShadingModelOverride() const;
+#endif
+
 #if WITH_EDITOR
 	/** 
 	* Like IsPropertyActive(), but should be used in context of editor.
