Hide IsThin property on material instance when Substrate is disabled.

diff --git a/Engine/Source/Editor/MaterialEditor/Private/MaterialEditorInstanceDetailCustomization.cpp b/Engine/Source/Editor/MaterialEditor/Private/MaterialEditorInstanceDetailCustomization.cpp
--- a/Engine/Source/Editor/MaterialEditor/Private/MaterialEditorInstanceDetailCustomization.cpp
+++ b/Engine/Source/Editor/MaterialEditor/Private/MaterialEditorInstanceDetailCustomization.cpp
@@ -1487,7 +1487,7 @@ void FMaterialInstanceParameterDetails::CreateBasePropertyOverrideWidgets(IDetai
 		IsOverriddenAndVisibleFn \
 		)
 
-#define CREATE_BASE_OVERRIDE_ROW(PropertyName, PropertyVariableName, ValueGetterName) \
+#define CREATE_BASE_OVERRIDE_ROW_OVERRIDEFN(PropertyName, PropertyVariableName, ValueGetterName, IsOverriddenAndVisibleFn) \
 	CREATE_BASE_OVERRIDE_ROW_CUSTOM( \
 		PropertyName, \
 		PropertyVariableName, \
@@ -1507,8 +1507,12 @@ void FMaterialInstanceParameterDetails::CreateBasePropertyOverrideWidgets(IDetai
 				MaterialEditorInstance->BasePropertyOverrides.PropertyVariableName = ParentMat->ValueGetterName(); \
 			} \
 		}, \
-		&FMaterialInstanceParameterDetails::IsOverriddenAndVisible \
+		IsOverriddenAndVisibleFn \
 	)
+
+#define CREATE_BASE_OVERRIDE_ROW(PropertyName, PropertyVariableName, ValueGetterName) \
+		CREATE_BASE_OVERRIDE_ROW_OVERRIDEFN(PropertyName, PropertyVariableName, ValueGetterName, &FMaterialInstanceParameterDetails::IsOverriddenAndVisible)
+
 #define CREATE_BASE_OVERRIDE_ROW_BASIC(PropertyName) \
 	CREATE_BASE_OVERRIDE_ROW(PropertyName, PropertyName, Get ## PropertyName)
 #define CREATE_BASE_OVERRIDE_ROW_BOOL(PropertyName, GetterName) \
@@ -1552,7 +1556,7 @@ void FMaterialInstanceParameterDetails::CreateBasePropertyOverrideWidgets(IDetai
 		&FMaterialInstanceParameterDetails::IsOverriddenAndVisibleShadingModels
 	);
 	CREATE_BASE_OVERRIDE_ROW(TwoSided, TwoSided, IsTwoSided);
-	CREATE_BASE_OVERRIDE_ROW_BASIC_BOOL(IsThinSurface);
+	CREATE_BASE_OVERRIDE_ROW_OVERRIDEFN(IsThinSurface, bIsThinSurface, IsThinSurface, &FMaterialInstanceParameterDetails::IsOverriddenAndVisibleSubstrateOnly);
 	CREATE_BASE_OVERRIDE_ROW(DitheredLODTransition, DitheredLODTransition, IsDitheredLODTransition);
 	CREATE_BASE_OVERRIDE_ROW_BOOL(OutputTranslucentVelocity, IsTranslucencyWritingVelocity);
 	CREATE_BASE_OVERRIDE_ROW_BASIC_BOOL(HasPixelAnimation);
@@ -1596,6 +1600,11 @@ EVisibility FMaterialInstanceParameterDetails::IsOverriddenAndVisibleShadingMode
 	return bShouldBeVisible ? EVisibility::Visible : EVisibility::Collapsed;
 }
 
+EVisibility FMaterialInstanceParameterDetails::IsOverriddenAndVisibleSubstrateOnly(TAttribute<bool> IsOverridden) const
+{
+	return Substrate::IsSubstrateEnabled() ? EVisibility::Visible : EVisibility::Collapsed;
+}
+
 /** Helper function used by some parameters to verify that they are allowed to be overridden. This
   * must be prevented if the source material instance disallows the creation of new static parameter
   * permutations as that would trigger a new shader creation. */
diff --git a/Engine/Source/Editor/MaterialEditor/Private/MaterialEditorInstanceDetailCustomization.h b/Engine/Source/Editor/MaterialEditor/Private/MaterialEditorInstanceDetailCustomization.h
--- a/Engine/Source/Editor/MaterialEditor/Private/MaterialEditorInstanceDetailCustomization.h
+++ b/Engine/Source/Editor/MaterialEditor/Private/MaterialEditorInstanceDetailCustomization.h
@@ -115,6 +115,7 @@ private:
 
 	EVisibility IsOverriddenAndVisible(TAttribute<bool> IsOverridden) const;
 	EVisibility IsOverriddenAndVisibleShadingModels(TAttribute<bool> IsOverridden) const;
+	EVisibility IsOverriddenAndVisibleSubstrateOnly(TAttribute<bool> IsOverridden) const;
 
 #define DECLARE_OVERRIDE_MEMBER_FUNCS(PropertyName) \
 	bool Override ## PropertyName ## Enabled() const; \
