[Viewport Toolbar] Material Editor Show Menu: use UToolMenu instead of FMenuBuilder

diff --git a/Engine/Source/Editor/MaterialEditor/Private/MaterialEditorViewportToolbarSections.cpp b/Engine/Source/Editor/MaterialEditor/Private/MaterialEditorViewportToolbarSections.cpp
--- a/Engine/Source/Editor/MaterialEditor/Private/MaterialEditorViewportToolbarSections.cpp
+++ b/Engine/Source/Editor/MaterialEditor/Private/MaterialEditorViewportToolbarSections.cpp
@@ -6,6 +6,9 @@
 #include "SEditorViewport.h"
 #include "SMaterialEditorViewport.h"
 #include "ToolMenu.h"
+#include "ToolMenus.h"
+#include "ViewportToolbar/UnrealEdViewportToolbar.h"
+#include "ViewportToolbar/UnrealEdViewportToolbarContext.h"
 
 #define LOCTEXT_NAMESPACE "MaterialEditorViewportToolbarSections"
 
@@ -14,53 +17,87 @@ TSharedRef<SWidget> UE::MaterialEditor::CreateShowMenuWidget(const TSharedRef<SM
 {
 	InMaterialEditorViewport->OnFloatingButtonClicked();
 
-	constexpr bool bInShouldCloseWindowAfterMenuSelection = true;
-	FMenuBuilder ShowMenuBuilder(bInShouldCloseWindowAfterMenuSelection, InMaterialEditorViewport->GetCommandList());
-	{
-		FMaterialEditorCommands Commands = FMaterialEditorCommands::Get();
+	// We generate a menu via UToolMenus, so we can use FillShowSubmenu call from both old and new toolbar
+	FName OldShowMenuName = "MaterialEditor.OldViewportToolbar.Show";
 
-		if (bInShowViewportStatsToggle)
-		{
-			ShowMenuBuilder.AddMenuEntry(
-				FEditorViewportCommands::Get().ToggleStats, "ViewportStats", LOCTEXT("ViewportStatsLabel", "Viewport Stats")
-			);
+	if (!UToolMenus::Get()->IsMenuRegistered(OldShowMenuName))
+	{
+		UToolMenu* Menu = UToolMenus::Get()->RegisterMenu(OldShowMenuName, NAME_None, EMultiBoxType::Menu, false);
+		Menu->AddDynamicSection(
+			"BaseSection",
+			FNewToolMenuDelegate::CreateLambda(
+				[ViewportWeak = InMaterialEditorViewport.ToWeakPtr(), bInShowViewportStatsToggle](UToolMenu* InMenu)
+				{
+					if (TSharedPtr<SMaterialEditor3DPreviewViewport> Viewport = ViewportWeak.Pin())
+					{
+						UUnrealEdViewportToolbarContext* const ContextObject = NewObject<UUnrealEdViewportToolbarContext>();
+						ContextObject->Viewport = ViewportWeak;
+						InMenu->Context.AddObject(ContextObject);
 
-			ShowMenuBuilder.AddMenuSeparator();
-		}
+						UE::MaterialEditor::FillShowSubmenu(InMenu, bInShowViewportStatsToggle);
+					}
+				}
+			)
+		);
+	}
 
-		ShowMenuBuilder.AddMenuEntry(Commands.ToggleMaterialStats);
+	FToolMenuContext MenuContext;
+	{
+		MenuContext.AppendCommandList(InMaterialEditorViewport->GetCommandList());
 
-		ShowMenuBuilder.AddMenuSeparator();
+		// Add the UnrealEd viewport toolbar context.
+		{
+			UUnrealEdViewportToolbarContext* const ContextObject =
+				UE::UnrealEd::CreateViewportToolbarDefaultContext(InMaterialEditorViewport);
 
-		ShowMenuBuilder.AddMenuEntry(Commands.TogglePreviewBackground);
+			MenuContext.AddObject(ContextObject);
+		}
 	}
 
-	return ShowMenuBuilder.MakeWidget();
+	return UToolMenus::Get()->GenerateWidget(OldShowMenuName, MenuContext);
 }
 
-FToolMenuEntry UE::MaterialEditor::CreateShowSubmenu(TWeakPtr<SMaterialEditor3DPreviewViewport> InViewport)
+FToolMenuEntry UE::MaterialEditor::CreateShowSubmenu()
 {
 	return FToolMenuEntry::InitSubMenu(
 		"Show",
 		LOCTEXT("ShowSubmenuLabel", "Show"),
 		LOCTEXT("ShowSubmenuTooltip", "Show options"),
 		FNewToolMenuDelegate::CreateLambda(
-			[InViewport](UToolMenu* Submenu) -> void
+			[](UToolMenu* Submenu) -> void
 			{
-				TSharedPtr<SMaterialEditor3DPreviewViewport> Viewport = InViewport.Pin();
-				if (!Viewport)
-				{
-					return;
-				}
-
-				FToolMenuSection& UnnamedSection = Submenu->FindOrAddSection("", FText());
-
-				UnnamedSection.AddEntry(FToolMenuEntry::InitWidget(
-					"ShowMenuItems", UE::MaterialEditor::CreateShowMenuWidget(Viewport.ToSharedRef()), FText(), true
-				));
+				UE::MaterialEditor::FillShowSubmenu(Submenu);
 			}
 		)
 	);
 }
 
+void UE::MaterialEditor::FillShowSubmenu(UToolMenu* InMenu, bool bInShowViewportStatsToggle)
+{
+	if (UUnrealEdViewportToolbarContext* const EditorViewportContext =
+			InMenu->FindContext<UUnrealEdViewportToolbarContext>())
+	{
+		if (TSharedPtr<SMaterialEditor3DPreviewViewport> StaticMeshEditorViewport =
+				StaticCastSharedPtr<SMaterialEditor3DPreviewViewport>(EditorViewportContext->Viewport.Pin()))
+		{
+			FToolMenuSection& UnnamedSection = InMenu->FindOrAddSection(NAME_None);
+
+			if (bInShowViewportStatsToggle)
+			{
+				UnnamedSection.AddMenuEntry(
+					FEditorViewportCommands::Get().ToggleStats, LOCTEXT("ViewportStatsLabel", "Viewport Stats")
+				);
+
+				UnnamedSection.AddSeparator(NAME_None);
+			}
+
+			UnnamedSection.AddMenuEntry(FMaterialEditorCommands::Get().ToggleMaterialStats);
+
+			UnnamedSection.AddSeparator(NAME_None);
+
+			UnnamedSection.AddMenuEntry(FMaterialEditorCommands::Get().TogglePreviewBackground);
+		}
+	}
+}
+
 #undef LOCTEXT_NAMESPACE
diff --git a/Engine/Source/Editor/MaterialEditor/Private/MaterialEditorViewportToolbarSections.h b/Engine/Source/Editor/MaterialEditor/Private/MaterialEditorViewportToolbarSections.h
--- a/Engine/Source/Editor/MaterialEditor/Private/MaterialEditorViewportToolbarSections.h
+++ b/Engine/Source/Editor/MaterialEditor/Private/MaterialEditorViewportToolbarSections.h
@@ -10,5 +10,6 @@ class SMaterialEditor3DPreviewViewport;
 namespace UE::MaterialEditor
 {
 TSharedRef<SWidget> CreateShowMenuWidget(const TSharedRef<SMaterialEditor3DPreviewViewport>& InMaterialEditorViewport, bool bInShowViewportStatsToggle = true);
-FToolMenuEntry CreateShowSubmenu(TWeakPtr<SMaterialEditor3DPreviewViewport> InViewport);
+FToolMenuEntry CreateShowSubmenu();
+void FillShowSubmenu(UToolMenu* InMenu, bool bInShowViewportStatsToggle = true);
 }
diff --git a/Engine/Source/Editor/MaterialEditor/Private/SMaterialEditorViewport.cpp b/Engine/Source/Editor/MaterialEditor/Private/SMaterialEditorViewport.cpp
--- a/Engine/Source/Editor/MaterialEditor/Private/SMaterialEditorViewport.cpp
+++ b/Engine/Source/Editor/MaterialEditor/Private/SMaterialEditorViewport.cpp
@@ -1001,7 +1001,7 @@ TSharedPtr<SWidget> SMaterialEditor3DPreviewViewport::BuildViewportToolbar()
 
 				// Add the Show submenu.
 				{
-					FToolMenuEntry ShowSubmenu = UE::MaterialEditor::CreateShowSubmenu(SharedThis(this));
+					FToolMenuEntry ShowSubmenu = UE::MaterialEditor::CreateShowSubmenu();
 					ShowSubmenu.InsertPosition.Position = EToolMenuInsertType::First;
 					RightSection.AddEntry(ShowSubmenu);
 				}
