Minor UI updates/fixes to Substrate Material Layer tree
- added an ellipsis menu with options to add a new material evaluation layer, relink to parent, save as sibling, save as child (existing options just moved here)
- added "+" button to add new material attribute layers
- Delete and Rename commands added, which also work with keys, with restrictions to allow only specific layers to be deleted/renamed
- Material attribute layers get/update their names based on the asset loaded into them

diff --git a/Engine/Source/Editor/MaterialEditor/Private/MaterialPropertyHelpers.cpp b/Engine/Source/Editor/MaterialEditor/Private/MaterialPropertyHelpers.cpp
--- a/Engine/Source/Editor/MaterialEditor/Private/MaterialPropertyHelpers.cpp
+++ b/Engine/Source/Editor/MaterialEditor/Private/MaterialPropertyHelpers.cpp
@@ -136,6 +136,10 @@ void FMaterialPropertyHelpers::OnMaterialLayerAssetChanged(const struct FAssetDa
 			{
 				InMaterialFunction->Layers[Index] = LayerFunction;
 				InMaterialFunction->UnlinkLayerFromParent(Index);
+#ifdef ENABLE_MATERIAL_LAYER_PROTOTYPE
+				// update name of the layer
+				InMaterialFunction->EditorOnly.LayerNames[Index] = FText::FromName(InAssetData.AssetName);
+#endif
 			}
 			break;
 		case EMaterialParameterAssociation::BlendParameter:
@@ -143,7 +147,7 @@ void FMaterialPropertyHelpers::OnMaterialLayerAssetChanged(const struct FAssetDa
 			{
 				InMaterialFunction->Blends[Index] = LayerFunction;
 #ifdef ENABLE_MATERIAL_LAYER_PROTOTYPE
-				InMaterialFunction->UnlinkLayerFromParent(Index); // 
+				InMaterialFunction->UnlinkLayerFromParent(Index);
 #else
 				InMaterialFunction->UnlinkLayerFromParent(Index + 1); // Blend indices are offset by 1, no blend for base layer
 #endif
diff --git a/Engine/Source/Editor/MaterialEditor/Private/SMaterialLayersFunctionsTree.cpp b/Engine/Source/Editor/MaterialEditor/Private/SMaterialLayersFunctionsTree.cpp
--- a/Engine/Source/Editor/MaterialEditor/Private/SMaterialLayersFunctionsTree.cpp
+++ b/Engine/Source/Editor/MaterialEditor/Private/SMaterialLayersFunctionsTree.cpp
@@ -36,7 +36,9 @@
 #include "Materials/MaterialFunctionInstance.h"
 #include "Materials/MaterialExpressionTextureSampleParameter.h"
 #include "Framework/Application/SlateApplication.h"
+#include "Framework/MultiBox/MultiBoxBuilder.h"
 #include "Widgets/Input/SEditableTextBox.h"
+#include "Widgets/Input/SComboButton.h"
 #include "Curves/CurveLinearColor.h"
 #include "Curves/CurveLinearColorAtlas.h"
 #include "Subsystems/AssetEditorSubsystem.h"
@@ -1798,9 +1800,7 @@ void SMaterialLayersFunctionsInstanceWrapper::Refresh()
 	
 	if (LayerParameter != nullptr)
 	{
-#if ENABLE_MATERIAL_LAYER_PROTOTYPE
-		FOnClicked OnRelinkToParent = FOnClicked::CreateSP(NestedTree.ToSharedRef(), &SMaterialSubstrateTree::RelinkLayersToParent);
-#else
+#if !ENABLE_MATERIAL_LAYER_PROTOTYPE
 		FOnClicked OnRelinkToParent = FOnClicked::CreateSP(NestedTree.ToSharedRef(), &SMaterialLayersFunctionsInstanceTree::RelinkLayersToParent);
 #endif
 		
@@ -1873,20 +1873,73 @@ void SMaterialLayersFunctionsInstanceWrapper::Refresh()
 		];
 		if (FMaterialPropertyHelpers::IsOverriddenExpression(NestedTree->FunctionParameter))
 		{
+#if ENABLE_MATERIAL_LAYER_PROTOTYPE
+			FMenuBuilder DetailViewOptions( true, nullptr);
+			DetailViewOptions.AddMenuEntry(
+			LOCTEXT("AddMaterialEvaluation", "Add Material Evaluation"),
+			LOCTEXT("AddMaterialEvaluation_ToolTip", "Add new Material Evaluation to the Substrate tree"),
+			FSlateIcon(),
+			FUIAction(FExecuteAction::CreateSP(NestedTree.Get(), &SMaterialSubstrateTree::AddRootNodeLayer)));
+
+			if (NestedTree->GetRelinkLayersToParentVisibility() == EVisibility::Visible)
+			{
+				DetailViewOptions.AddMenuEntry(
+				LOCTEXT("Relink", "Relink"),
+				LOCTEXT("RelinkToParentLayers", "Relink to Parent Layers and Blends"),
+				FSlateIcon(),
+				FUIAction(FExecuteAction::CreateSP(NestedTree.Get(), &SMaterialSubstrateTree::RelinkLayersToParent)));
+			}
+			DetailViewOptions.AddMenuEntry(
+			LOCTEXT("SaveSibling", "Save Sibling"),
+			LOCTEXT("SaveToSiblingInstance", "Save to Sibling Instance"),
+			FSlateIcon(),
+			FUIAction(FExecuteAction::CreateLambda([this, OnSiblingButtonClicked]()	{ OnSiblingButtonClicked.Execute(); })));
+
+			DetailViewOptions.AddMenuEntry(
+			LOCTEXT("SaveChild", "Save Child"),
+			LOCTEXT("SaveToChildInstance", "Save To Child Instance"),
+			FSlateIcon(),
+			FUIAction(FExecuteAction::CreateLambda([this, OnChildButtonClicked]()	{ OnChildButtonClicked.Execute(); })));
+
 			HeaderBox->AddSlot()
-				.HAlign(HAlign_Left)
+			.FillWidth(1.0f)
+			[
+				SNullWidget::NullWidget
+			];
+			HeaderBox->AddSlot()
+				.HAlign(HAlign_Right)
 				.AutoWidth()
+				.Padding(2.0f)
 				.VAlign(VAlign_Center)
 				[
-#if ENABLE_MATERIAL_LAYER_PROTOTYPE
-				PropertyCustomizationHelpers::MakeAddButton(FSimpleDelegate::CreateSP(NestedTree.Get(), &SMaterialSubstrateTree::AddRootNodeLayer))
+					SNew(SComboButton)
+						.ContentPadding(0.f)
+						.ComboButtonStyle(&FAppStyle::Get().GetWidgetStyle<FComboButtonStyle>("SimpleComboButton"))
+						.HasDownArrow(false)
+						.MenuContent()
+						[
+							DetailViewOptions.MakeWidget()
+						]
+						.ButtonContent()
+						[
+							SNew(SImage)
+							.ColorAndOpacity(FSlateColor::UseForeground())
+							.Image(FAppStyle::GetBrush( "PropertyEditor.VerticalDottedLine"))
+						]
+				];
+			
 #else
+			HeaderBox->AddSlot()
+				.HAlign(HAlign_Left)
+				.AutoWidth()
+				.VAlign(VAlign_Center)
+				[
 					PropertyCustomizationHelpers::MakeAddButton(FSimpleDelegate::CreateSP(NestedTree.Get(), &SMaterialLayersFunctionsInstanceTree::AddLayer))
-#endif
 				];
+#endif
 		}
 
-#ifndef ENABLE_MATERIAL_LAYER_PROTOTYPE
+#if !ENABLE_MATERIAL_LAYER_PROTOTYPE
 		HeaderBox->AddSlot()
 			.FillWidth(1.0f)
 			[
diff --git a/Engine/Source/Editor/MaterialEditor/Private/SMaterialSubstrateTree.cpp b/Engine/Source/Editor/MaterialEditor/Private/SMaterialSubstrateTree.cpp
--- a/Engine/Source/Editor/MaterialEditor/Private/SMaterialSubstrateTree.cpp
+++ b/Engine/Source/Editor/MaterialEditor/Private/SMaterialSubstrateTree.cpp
@@ -55,6 +55,7 @@
 #include "Styling/StyleColors.h"
 #include "MaterialEditorStyle.h"
 #include "DetailTreeNode.h"
+#include "Framework/Commands/GenericCommands.h"
 #include "Materials/MaterialFunctionMaterialLayerBlend.h"
 
 #define LOCTEXT_NAMESPACE "MaterialSubstrateTree"
@@ -153,33 +154,37 @@ ECheckBoxState SMaterialSubstrateTreeItem::GetFilterChecked(SMaterialSubstrateTr
 	return GetFilterState(InTree, InStackData) ? ECheckBoxState::Checked : ECheckBoxState::Unchecked;
 }
 
+int32 SMaterialSubstrateTreeItem::GetLayerFunctionIndex() const
+{
+	return Tree->GetLayerFunctionIndex(StackParameterData->ParameterInfo.Index);
+}
+
+int32 SMaterialSubstrateTreeItem::GetBlendFunctionIndex() const
+{
+	return Tree->GetBlendFunctionIndex(StackParameterData->ParameterInfo.Index);
+}
+
 FReply SMaterialSubstrateTreeItem::ToggleLayerVisibility()
 {
-	int32 LayerFuncIndex = Tree->FunctionInstance->GetLayerFuncIndex(StackParameterData->ParameterInfo.Index);
+	int32 LayerFuncIndex = GetLayerFunctionIndex();
 	return Tree->ToggleLayerVisibility(LayerFuncIndex);
 }
 
 bool SMaterialSubstrateTreeItem::IsLayerVisible() const
 {
-	int32 LayerFuncIndex = Tree->FunctionInstance->GetLayerFuncIndex(StackParameterData->ParameterInfo.Index);
+	int32 LayerFuncIndex = GetLayerFunctionIndex();
 	return Tree->IsLayerVisible(LayerFuncIndex);
 }
 
-FReply SMaterialSubstrateTreeItem::UnlinkLayer()
-{
-	int32 LayerFuncIndex = Tree->FunctionInstance->GetLayerFuncIndex(StackParameterData->ParameterInfo.Index);
-	return Tree->UnlinkLayer(LayerFuncIndex);
-}
-
 EVisibility SMaterialSubstrateTreeItem::GetUnlinkLayerVisibility() const
 {
-	int32 LayerFuncIndex = Tree->FunctionInstance->GetLayerFuncIndex(StackParameterData->ParameterInfo.Index);
+	int32 LayerFuncIndex = GetLayerFunctionIndex();
 	return Tree->GetUnlinkLayerVisibility(LayerFuncIndex);
 }
 
 FText SMaterialSubstrateTreeItem::GetLayerName() const
 {
-	int32 LayerFuncIndex = Tree->FunctionInstance->GetLayerFuncIndex(StackParameterData->ParameterInfo.Index);
+	int32 LayerFuncIndex = GetLayerFunctionIndex();
 	return Tree->FunctionInstance->GetLayerName(LayerFuncIndex);
 }
 
@@ -193,7 +198,7 @@ FText SMaterialSubstrateTreeItem::GetLayerDesc() const
 void SMaterialSubstrateTreeItem::OnNameChanged(const FText& InText, ETextCommit::Type CommitInfo)
 {
 	const FScopedTransaction Transaction(LOCTEXT("RenamedSection", "Renamed layer and blend section"));
-	int32 LayerFuncIndex = Tree->FunctionInstance->GetLayerFuncIndex(StackParameterData->ParameterInfo.Index);
+	int32 LayerFuncIndex = GetLayerFunctionIndex();
 	Tree->FunctionInstanceHandle->NotifyPreChange();
 	Tree->FunctionInstance->EditorOnly.LayerNames[LayerFuncIndex] = InText;
 	Tree->FunctionInstance->UnlinkLayerFromParent(LayerFuncIndex);
@@ -456,6 +461,12 @@ void  SMaterialSubstrateTreeItem::OnOverrideParameter(bool NewValue, TObjectPtr<
 	OnOverrideParameter(NewValue, Parameter.Get());
 }
 
+FReply SMaterialSubstrateTreeItem::OnAddMaterialAttributeClicked(int32 index)
+{
+	Tree->AddNodeLayer(index);
+	return FReply::Handled();
+}
+
 void SMaterialSubstrateTreeItem::Construct(const FArguments& InArgs, const TSharedRef<STableViewBase>& InOwnerTableView)
 {
 	
@@ -504,8 +515,7 @@ void SMaterialSubstrateTreeItem::Construct(const FArguments& InArgs, const TShar
 	if (StackParameterData->StackDataType == EStackDataType::Stack)
 	{
 #if WITH_EDITOR
-		int32 LayerFuncIndex = Tree->FunctionInstance->GetLayerFuncIndex(StackParameterData->ParameterInfo.Index);
-		int32 BlendFuncIndex = Tree->FunctionInstance->GetBlendFuncIndex(StackParameterData->ParameterInfo.Index);
+		int32 LayerFuncIndex = GetLayerFunctionIndex();
 		NameOverride = Tree->FunctionInstance->GetLayerName(LayerFuncIndex);
 #endif
 		TSharedRef<SHorizontalBox> HeaderRowWidget = SNew(SHorizontalBox);
@@ -640,9 +650,9 @@ void SMaterialSubstrateTreeItem::Construct(const FArguments& InArgs, const TShar
 					+ SVerticalBox::Slot()
 					.AutoHeight()
 					[
-						SNew(STextBlock)
-						.Text(NameOverride)
-						.TextStyle(FSubstrateMaterialEditorStyle::Get(), "LayerView.Row.HeaderText")
+						SAssignNew(InlineRenameWidget, SInlineEditableTextBlock)
+							.Text(this, &SMaterialSubstrateTreeItem::GetDisplayName)
+							.OnTextCommitted(this, &SMaterialSubstrateTreeItem::OnNameChanged)
 					]
 
 					+ SVerticalBox::Slot()
@@ -655,28 +665,6 @@ void SMaterialSubstrateTreeItem::Construct(const FArguments& InArgs, const TShar
 				
 				];
 		}
-
-		// Unlink UI
-		HeaderRowWidget->AddSlot()
-			.FillWidth(1.0f)
-			.VAlign(VAlign_Center)
-			[
-				SNullWidget::NullWidget
-			];
-
-		HeaderRowWidget->AddSlot()
-			.AutoWidth()
-			.VAlign(VAlign_Center)
-			.Padding(0.0f, 0.0f, 0.0f, 0.0f)
-			[
-				SNew(SButton)
-				.Text(LOCTEXT("Unlink", "Unlink"))
-				.HAlign(HAlign_Center)
-				.OnClicked(this, &SMaterialSubstrateTreeItem::UnlinkLayer)
-				.ToolTipText(LOCTEXT("UnlinkLayer", "Whether or not to unlink this layer/blend combination from the parent."))
-				.Visibility(this, &SMaterialSubstrateTreeItem::GetUnlinkLayerVisibility)
-			];
-
 		LeftSideWidget = HeaderRowWidget;
 	}
 // END STACK
@@ -691,7 +679,28 @@ void SMaterialSubstrateTreeItem::Construct(const FArguments& InArgs, const TShar
 			[
 				LeftSideWidget
 			];
-		
+
+		// Add button
+		if (bCanAppendSubLayer)
+		{
+			MainStack->AddSlot()
+					.VAlign(VAlign_Center)
+					.Padding(2,0,0,0)
+					.AutoWidth()
+					[
+						SNew(SButton)
+						.ContentPadding(FMargin(1, 0))
+						.ButtonStyle(FAppStyle::Get(), "SimpleButton")
+							.ToolTipText(LOCTEXT("AddMaterialAttribute", "Add Material Attribute"))
+							.OnClicked(this, &SMaterialSubstrateTreeItem::OnAddMaterialAttributeClicked, StackParameterData->ParameterInfo.Index)
+							[
+								SNew(SImage)
+								.Image(FAppStyle::Get().GetBrush("Icons.Plus"))
+								.ColorAndOpacity(FStyleColors::AccentGreen)
+							]
+					
+					];
+		}
 		if (bCanReorderLayer)
 		{
 			MainStack->AddSlot()
@@ -723,7 +732,11 @@ void SMaterialSubstrateTreeItem::Construct(const FArguments& InArgs, const TShar
 	
 	this->SetDesiredSizeScale(FVector2d(1.0f, 1.2f));
 }
-
+FText SMaterialSubstrateTreeItem::GetDisplayName() const
+{
+	int32 LayerFuncIndex = GetLayerFunctionIndex();
+	return Tree->FunctionInstance->GetLayerName(LayerFuncIndex);
+}
 int32 SMaterialSubstrateTreeItem::OnPaint(const FPaintArgs& Args, const FGeometry& AllottedGeometry, const FSlateRect& MyCullingRect, FSlateWindowElementList& OutDrawElements, int32 LayerId, const FWidgetStyle& InWidgetStyle, bool bParentEnabled) const
 {
 	const FSlateBrush* BackgroundBrushResource = Tree->BackgroundBrush.Get();
@@ -779,11 +792,19 @@ int32 SMaterialSubstrateTreeItem::OnPaint(const FPaintArgs& Args, const FGeometr
 	return STableRow<TSharedPtr<FSortedParamData>>::OnPaint(Args, BorderGeom, MyCullingRect, OutDrawElements, LayerId+1, InWidgetStyle, bParentEnabled);
 }
 
+void SMaterialSubstrateTreeItem::Rename()
+{
+	if (InlineRenameWidget)
+	{
+		InlineRenameWidget->EnterEditingMode();
+	}
+}
+
 
 FString SMaterialSubstrateTreeItem::GetInstancePath(SMaterialSubstrateTree* InTree) const
 {
-	int32 LayerFuncIndex = Tree->FunctionInstance->GetLayerFuncIndex(StackParameterData->ParameterInfo.Index);
-	int32 BlendFuncIndex = Tree->FunctionInstance->GetBlendFuncIndex(StackParameterData->ParameterInfo.Index);
+	int32 LayerFuncIndex = GetLayerFunctionIndex();
+	int32 BlendFuncIndex = GetBlendFunctionIndex();
 
 	FString InstancePath;
 	if (StackParameterData->ParameterInfo.Association == EMaterialParameterAssociation::BlendParameter && InTree->FunctionInstance->Blends.IsValidIndex(BlendFuncIndex))
@@ -799,33 +820,62 @@ FString SMaterialSubstrateTreeItem::GetInstancePath(SMaterialSubstrateTree* InTr
 
 TSharedPtr<SWidget> SMaterialSubstrateTree::CreateContextMenu()
 {
-	const bool bCloseAfterSelection = true;
-	
-	FMenuBuilder MenuBuilder(bCloseAfterSelection, TSharedPtr<FUICommandList>());
+	FMenuBuilder MenuBuilder(true, CommandList);
 
 	TArray<FSortedParamDataPtr> SelectedItemsArray = GetSelectedItems();
 	if (SelectedItemsArray.Num() > 0)
 	{
 		FSortedParamDataPtr StackParameterData = SelectedItemsArray[0];
-		
+		bool bCanUnlinkLayer = GetUnlinkLayerVisibility(StackParameterData->ParameterInfo.Index) == EVisibility::Visible;
 		bool bCanAppendSubLayer = FunctionInstance->CanAppendLayerNode(StackParameterData->ParameterInfo.Index);
-		bool bCanRemoveLayer = FunctionInstance->CanRemoveLayerNode(StackParameterData->ParameterInfo.Index);
-	
 		if (bCanAppendSubLayer)
 		{
 			const FSlateIcon PlusIcon(FAppStyle::GetAppStyleSetName(), "Plus");
 			MenuBuilder.AddMenuEntry(LOCTEXT("AddNewLayer", "Add New Layer"), FText(), PlusIcon, FUIAction(FExecuteAction::CreateSP(this, &SMaterialSubstrateTree::AddNodeLayer, StackParameterData->ParameterInfo.Index)));
 		}
-
-		if (bCanRemoveLayer)
+		MenuBuilder.AddMenuEntry(FGenericCommands::Get().Delete);
+		MenuBuilder.AddMenuEntry(FGenericCommands::Get().Rename);
+		
+		if (bCanUnlinkLayer)
 		{
-			const FSlateIcon MinusIcon(FAppStyle::GetAppStyleSetName(), "Icons.Minus");
-			MenuBuilder.AddMenuEntry(LOCTEXT("RemoveLayer", "Remove Layer"), FText(), MinusIcon, FUIAction(FExecuteAction::CreateSP(this, &SMaterialSubstrateTree::RemoveNodeLayer, StackParameterData->ParameterInfo.Index)));
+			int32 LayerFuncIndex;
+			LayerFuncIndex = GetLayerFunctionIndex(StackParameterData->ParameterInfo.Index);
+			MenuBuilder.AddMenuEntry(LOCTEXT("UnlinkLayer", "Unlink Layer"),
+			                         LOCTEXT("UnlinkLayer", "Whether or not to unlink this layer/blend combination from the parent."),
+			                         FSlateIcon(),
+			                         FUIAction(FExecuteAction::CreateSP(this, &SMaterialSubstrateTree::UnlinkLayer, LayerFuncIndex)));
 		}
 	}
 	return MenuBuilder.MakeWidget();
 }
 
+void SMaterialSubstrateTree::CreateCommandList()
+{
+	CommandList = MakeShareable(new FUICommandList);
+	CommandList->MapAction(FGenericCommands::Get().Delete,
+		FUIAction(
+			FExecuteAction::CreateSP(this, &SMaterialSubstrateTree::OnDeleteSelectedTreeViewItems),
+			FCanExecuteAction::CreateSP(this, &SMaterialSubstrateTree::CanDeleteSelectedTreeViewItems)
+		)
+	);
+	CommandList->MapAction(FGenericCommands::Get().Rename,
+		FUIAction(
+			FExecuteAction::CreateSP(this, &SMaterialSubstrateTree::OnRenameSelectedTreeViewItems),
+			FCanExecuteAction::CreateSP(this, &SMaterialSubstrateTree::CanRenameSelectedTreeViewItem)
+		)
+	);
+}
+
+
+FReply SMaterialSubstrateTree::OnKeyDown(const FGeometry& MyGeometry, const FKeyEvent& InKeyEvent)
+{
+	if (CommandList.IsValid() && CommandList->ProcessCommandBindings(InKeyEvent))
+	{
+		return FReply::Handled();
+	}
+	return FReply::Unhandled();
+}
+
 void SMaterialSubstrateTree::Construct(const FArguments& InArgs)
 {
 	ColumnSizeData.SetValueColumnWidth(0.5f);
@@ -835,7 +885,7 @@ void SMaterialSubstrateTree::Construct(const FArguments& InArgs)
 	Wrapper = InArgs._InWrapper;
 	ShowHiddenDelegate = InArgs._InShowHiddenDelegate;
 	CreateGroupsWidget();
-
+	CreateCommandList();
 #if WITH_EDITOR
 	//Fixup for adding new bool arrays to the class
 	if (FunctionInstance)
@@ -938,6 +988,16 @@ void SMaterialSubstrateTree::SetParentsExpansionState()
 	}
 }
 
+int32 SMaterialSubstrateTree::GetLayerFunctionIndex(int32 NodeIndex) const
+{
+	return FunctionInstance ? FunctionInstance->GetLayerFuncIndex(NodeIndex) : -1;
+}
+
+int32 SMaterialSubstrateTree::GetBlendFunctionIndex(int32 NodeIndex) const
+{
+	return FunctionInstance ? FunctionInstance->GetBlendFuncIndex(NodeIndex) : -1;
+}
+
 void SMaterialSubstrateTree::RefreshOnAssetChange(const struct FAssetData& InAssetData, int32 InNodeId, EMaterialParameterAssociation MaterialType)
 {
 	auto NodePayload = FunctionInstance->GetNodePayload(InNodeId);
@@ -948,6 +1008,7 @@ void SMaterialSubstrateTree::RefreshOnAssetChange(const struct FAssetData& InAss
 		return;
 
 	FMaterialPropertyHelpers::OnMaterialLayerAssetChanged(InAssetData, Index, MaterialType, FunctionInstanceHandle, FunctionInstance);
+	
 	//set their overrides back to 0
 	MaterialEditorInstance->CleanParameterStack(Index, MaterialType);
 	CreateGroupsWidget();
@@ -1009,17 +1070,17 @@ void SMaterialSubstrateTree::RemoveNodeLayer(int32 InNodeId)
 	RequestTreeRefresh();
 }
 
-FReply SMaterialSubstrateTree::UnlinkLayer(int32 Index)
+void SMaterialSubstrateTree::UnlinkLayer(int32 Index)
 {
 	const FScopedTransaction Transaction(LOCTEXT("UnlinkLayerFromParent", "Unlink a layer from the parent"));
 	FunctionInstanceHandle->NotifyPreChange();
 	FunctionInstance->UnlinkLayerFromParent(Index);
 	FunctionInstanceHandle->NotifyPostChange(EPropertyChangeType::ValueSet);
 	CreateGroupsWidget();
-	return FReply::Handled();
+	Wrapper->Refresh();
 }
 
-FReply SMaterialSubstrateTree::RelinkLayersToParent()
+void SMaterialSubstrateTree::RelinkLayersToParent()
 {
 	const FScopedTransaction Transaction(LOCTEXT("RelinkLayersToParent", "Relink layers to parent"));
 	FunctionInstanceHandle->NotifyPreChange();
@@ -1027,7 +1088,7 @@ FReply SMaterialSubstrateTree::RelinkLayersToParent()
 	FunctionInstanceHandle->NotifyPostChange(EPropertyChangeType::ValueSet);
 	MaterialEditorInstance->RegenerateArrays();
 	CreateGroupsWidget();
-	return FReply::Handled();
+	Wrapper->Refresh();
 }
 
 EVisibility SMaterialSubstrateTree::GetUnlinkLayerVisibility(int32 Index) const
@@ -1319,7 +1380,7 @@ void SMaterialSubstrateTree::CreateGroupsWidget()
 		DeferredSearches.Empty();
 
 		// Create the hierarchy of Sorted items recursivelly following the LayerFunctions Tree
-#ifdef ENABLE_MATERIAL_LAYER_PROTOTYPE
+#if ENABLE_MATERIAL_LAYER_PROTOTYPE
 		if (FunctionParameterHandle)
 		{
 			TSharedPtr<IPropertyHandle>	LayerHandle = FunctionParameterHandle->GetChildHandle("Layers").ToSharedRef();
@@ -1529,4 +1590,51 @@ void SMaterialSubstrateTree::ShowSubParameters(TSharedPtr<FSortedParamData> Pare
 	}
 }
 
-#undef LOCTEXT_NAMESPACE
\ No newline at end of file
+void SMaterialSubstrateTree::OnDeleteSelectedTreeViewItems()
+{
+	TArray<FSortedParamDataPtr> SelectedItemsArray = GetSelectedItems();
+	if (SelectedItemsArray.Num() > 0)
+	{
+		FSortedParamDataPtr StackParameterData = SelectedItemsArray[0];
+		RemoveNodeLayer(StackParameterData->ParameterInfo.Index);
+	}
+}
+
+bool SMaterialSubstrateTree::CanDeleteSelectedTreeViewItems() const
+{
+	bool bCanRemoveLayer = false;
+	TArray<FSortedParamDataPtr> SelectedItemsArray = GetSelectedItems();
+	if (SelectedItemsArray.Num() > 0)
+	{
+		FSortedParamDataPtr StackParameterData = SelectedItemsArray[0];
+		bCanRemoveLayer = FunctionInstance->CanRemoveLayerNode(StackParameterData->ParameterInfo.Index);
+	}
+	return bCanRemoveLayer;
+}
+
+void SMaterialSubstrateTree::OnRenameSelectedTreeViewItems()
+{
+	TArray<FSortedParamDataPtr> SelectedItemsArray = GetSelectedItems();
+	if (SelectedItemsArray.Num() > 0)
+	{
+		FSortedParamDataPtr StackParameterData = SelectedItemsArray[0];
+
+		TSharedPtr<SMaterialSubstrateTreeItem> TreeItem = StaticCastSharedPtr<SMaterialSubstrateTreeItem>(WidgetFromItem(StackParameterData));
+		TreeItem->Rename();
+	}
+}
+
+bool SMaterialSubstrateTree::CanRenameSelectedTreeViewItem() const
+{
+	bool bCanRenameLayer = false;
+	TArray<FSortedParamDataPtr> SelectedItemsArray = GetSelectedItems();
+	if (SelectedItemsArray.Num() > 0)
+	{
+		FSortedParamDataPtr StackParameterData = SelectedItemsArray[0];
+		// only allow Material Evaluation layers to be renamed for now
+		bCanRenameLayer = FunctionInstance->GetNodeDepth(StackParameterData->ParameterInfo.Index) == 1;
+	}
+	return bCanRenameLayer;
+}
+
+#undef LOCTEXT_NAMESPACE
diff --git a/Engine/Source/Editor/MaterialEditor/Private/SMaterialSubstrateTree.h b/Engine/Source/Editor/MaterialEditor/Private/SMaterialSubstrateTree.h
--- a/Engine/Source/Editor/MaterialEditor/Private/SMaterialSubstrateTree.h
+++ b/Engine/Source/Editor/MaterialEditor/Private/SMaterialSubstrateTree.h
@@ -49,7 +49,12 @@ public:
 	void OnSelectionChangedMaterialSubstrateView(TSharedPtr<FSortedParamData> InSelectedItem, ESelectInfo::Type SelectInfo);
 	void SetParentsExpansionState();
 
+	int32 GetLayerFunctionIndex(int32 NodeIndex) const;
+	int32 GetBlendFunctionIndex(int32 NodeIndex) const;
 	TSharedPtr<SWidget> CreateContextMenu();
+	void CreateCommandList();
+	virtual FReply OnKeyDown(const FGeometry& MyGeometry, const FKeyEvent& InKeyEvent) override;
+
 	void ShowHiddenValues(bool& bShowHiddenParameters) { bShowHiddenParameters = true; }
 	TWeakObjectPtr<class UDEditorParameterValue> FunctionParameter;
 	struct FMaterialLayersFunctions* FunctionInstance;
@@ -61,8 +66,8 @@ public:
 	void AddNodeLayer(int32 InParent = -1);
 	void RemoveNodeLayer(int32 InNodeId);
 
-	FReply UnlinkLayer(int32 Index);
-	FReply RelinkLayersToParent();
+	void UnlinkLayer(int32 Index);
+	void RelinkLayersToParent();
 	EVisibility GetUnlinkLayerVisibility(int32 Index) const;
 	EVisibility GetRelinkLayersToParentVisibility() const;
 	FReply ToggleLayerVisibility(int32 Index);
@@ -93,8 +98,12 @@ public:
 protected:
 
 	void ShowSubParameters(TSharedPtr<FSortedParamData> ParentParameter);
-
+	void OnDeleteSelectedTreeViewItems();
+	bool CanDeleteSelectedTreeViewItems() const;
+	void OnRenameSelectedTreeViewItems();
+	bool CanRenameSelectedTreeViewItem() const;
 private:
+	TSharedPtr<FUICommandList> CommandList;
 	TArray<TSharedPtr<FSortedParamData>> LayerProperties;
 	
 	TArray<FUnsortedParamData> NonLayerProperties;
@@ -131,10 +140,18 @@ public:
 
 	bool bIsBeingDragged = false;
 
+
 private:
 	bool bIsHoveredDragTarget = false;
+	
+	/** Widget to display the name of the asset item and allows for renaming */
+	TSharedPtr<class SInlineEditableTextBlock> InlineRenameWidget;
+	
+	/** Handles verifying name changes */
+	bool VerifyNameChanged(const FText& InName, FText& OutError) const;
 
-
+	/** Returns false if this folder is in the process of being created */
+	bool IsReadOnly() const;
 	FString GetCurvePath(class UDEditorScalarParameterValue* Parameter) const;
 	const FSlateBrush* GetBorderImage() const;
 
@@ -145,6 +162,8 @@ public:
 	bool GetFilterState(SMaterialSubstrateTree* InTree, TSharedPtr<FSortedParamData> InStackData) const;
 	void FilterClicked(const ECheckBoxState NewCheckedState, SMaterialSubstrateTree* InTree, TSharedPtr<FSortedParamData> InStackData);
 	ECheckBoxState GetFilterChecked(SMaterialSubstrateTree* InTree, TSharedPtr<FSortedParamData> InStackData) const;
+	int32 GetLayerFunctionIndex() const;
+	int32 GetBlendFunctionIndex() const;
 	FText GetLayerName() const;
 	FText GetLayerDesc() const;
 	void OnNameChanged(const FText& InText, ETextCommit::Type CommitInfo);
@@ -179,13 +198,16 @@ public:
 	void OnOverrideParameter(bool NewValue, class UDEditorParameterValue* Parameter);
 	void OnOverrideParameter(bool NewValue, TObjectPtr<UDEditorParameterValue> Parameter);
 
+	FReply OnAddMaterialAttributeClicked(int32 ParentIndex);
 	/**
 	* Construct the widget
 	*
 	* @param InArgs   A declaration from which to construct the widget
 	*/
 	void Construct(const FArguments& InArgs, const TSharedRef<STableViewBase>& InOwnerTableView);
+	FText GetDisplayName() const;
 	int32 OnPaint( const FPaintArgs& Args, const FGeometry& AllottedGeometry, const FSlateRect& MyCullingRect, FSlateWindowElementList& OutDrawElements, int32 LayerId, const FWidgetStyle& InWidgetStyle, bool bParentEnabled ) const override;
+	void Rename();
 	/** The node info to build the tree view row from. */
 	
 	FSortedParamDataPtr StackParameterData;
