Fix Text conflict in Build Engine Localization - SMaterialSubstrateTree.cpp

diff --git a/Engine/Source/Editor/MaterialEditor/Private/SMaterialSubstrateTree.cpp b/Engine/Source/Editor/MaterialEditor/Private/SMaterialSubstrateTree.cpp
--- a/Engine/Source/Editor/MaterialEditor/Private/SMaterialSubstrateTree.cpp
+++ b/Engine/Source/Editor/MaterialEditor/Private/SMaterialSubstrateTree.cpp
@@ -841,7 +841,7 @@ TSharedPtr<SWidget> SMaterialSubstrateTree::CreateContextMenu()
 			int32 LayerFuncIndex;
 			LayerFuncIndex = GetLayerFunctionIndex(StackParameterData->ParameterInfo.Index);
 			MenuBuilder.AddMenuEntry(LOCTEXT("UnlinkLayer", "Unlink Layer"),
-			                         LOCTEXT("UnlinkLayer", "Whether or not to unlink this layer/blend combination from the parent."),
+			                         LOCTEXT("UnlinkLayerTooltip", "Whether or not to unlink this layer/blend combination from the parent."),
 			                         FSlateIcon(),
 			                         FUIAction(FExecuteAction::CreateSP(this, &SMaterialSubstrateTree::UnlinkLayer, LayerFuncIndex)));
 		}
