Substrate - now marked as beta.

[RN] Substrate is now BETA!
#tests none

diff --git a/Engine/Source/Editor/MaterialEditor/Private/SMaterialEditorSubstrateWidget.cpp b/Engine/Source/Editor/MaterialEditor/Private/SMaterialEditorSubstrateWidget.cpp
--- a/Engine/Source/Editor/MaterialEditor/Private/SMaterialEditorSubstrateWidget.cpp
+++ b/Engine/Source/Editor/MaterialEditor/Private/SMaterialEditorSubstrateWidget.cpp
@@ -407,7 +407,7 @@ void SMaterialEditorSubstrateWidget::Construct(const FArguments& InArgs, TWeakPt
 					.ColorAndOpacity(FLinearColor::Yellow)
 					.ShadowColorAndOpacity(FLinearColor::Black)
 					.ShadowOffset(FVector2D::UnitVector)
-					.Text(LOCTEXT("SubstrateWidgetNotEnable", "Details cannot be shown: Substrate (experimental) is not enabled for this project (See the project settings window, rendering settings section)."))
+					.Text(LOCTEXT("SubstrateWidgetNotEnable", "Details cannot be shown: Substrate (Beta) is not enabled for this project (See the project settings window, rendering settings section)."))
 				]
 			]
 		];
diff --git a/Engine/Source/Runtime/Engine/Classes/Engine/RendererSettings.h b/Engine/Source/Runtime/Engine/Classes/Engine/RendererSettings.h
--- a/Engine/Source/Runtime/Engine/Classes/Engine/RendererSettings.h
+++ b/Engine/Source/Runtime/Engine/Classes/Engine/RendererSettings.h
@@ -1195,8 +1195,8 @@ class URendererSettings : public UDeveloperSettings
 	"Enable Substrate materials (Beta)."
 	*/
 	UPROPERTY(config, EditAnywhere, Category = Substrate, meta = (
-		ConsoleVariable = "r.Substrate", DisplayName = "Substrate materials (Experimental)",
-		ToolTip = "Enable Substrate materials (Experimental).",
+		ConsoleVariable = "r.Substrate", DisplayName = "Substrate materials (Beta)",
+		ToolTip = "Enable Substrate materials (Beta).",
 		ConfigRestartRequired = true))
 		uint32 bEnableSubstrate : 1;
 
diff --git a/Engine/Source/Runtime/Engine/Private/RendererSettings.cpp b/Engine/Source/Runtime/Engine/Private/RendererSettings.cpp
--- a/Engine/Source/Runtime/Engine/Private/RendererSettings.cpp
+++ b/Engine/Source/Runtime/Engine/Private/RendererSettings.cpp
@@ -367,7 +367,7 @@ void URendererSettings::PostEditChangeProperty(FPropertyChangedEvent& PropertyCh
 		{
 			if (bEnableSubstrate)
 			{
-				FMessageDialog::Open(EAppMsgType::Ok, LOCTEXT("Substrate Experimental", "Warning: Substrate is experimental. Be aware that any materials saved when Substrate is enabled won't be rendered correctly if Substrate is disabled later on."));
+				FMessageDialog::Open(EAppMsgType::Ok, LOCTEXT("Substrate Beta", "Warning: Substrate is in a Beta state. Be aware that any materials saved when Substrate is enabled won't be rendered correctly if Substrate is disabled later on."));
 			}
 		}
 
