Refreshing the material 3D preview viewport when the profile is changed for the viewport. Previously the profile change wouldn't take effect until the viewport was manually updated by for example resizing the viewport.

#rnx

diff --git a/Engine/Source/Editor/MaterialEditor/Private/SMaterialEditorViewport.cpp b/Engine/Source/Editor/MaterialEditor/Private/SMaterialEditorViewport.cpp
--- a/Engine/Source/Editor/MaterialEditor/Private/SMaterialEditorViewport.cpp
+++ b/Engine/Source/Editor/MaterialEditor/Private/SMaterialEditorViewport.cpp
@@ -904,6 +904,11 @@ void SMaterialEditor3DPreviewViewport::OnAssetViewerSettingsChanged(const FName&
 			EditorViewportClient->SetRealtime(true);
 		}
 	}
+	else if (InPropertyName == GET_MEMBER_NAME_CHECKED(FPreviewSceneProfile, ProfileName))
+	{
+		// If the profile changes, refresh the viewport for it to take effect
+		RefreshViewport();
+	}
 }
 
 TSharedRef<class SEditorViewport> SMaterialEditor3DPreviewViewport::GetViewportWidget()
