Drag/Drop updates
- Drag and drop material evals (reorder)
- Drag and drop layers within a material eval (reorder)
- Drag and drop layers between material evals
-     Move by default
-     Hold ctrl or shift to duplicate
Used correct Ellipsis icon for options button

diff --git a/Engine/Source/Editor/MaterialEditor/Private/SMaterialLayersFunctionsTree.cpp b/Engine/Source/Editor/MaterialEditor/Private/SMaterialLayersFunctionsTree.cpp
--- a/Engine/Source/Editor/MaterialEditor/Private/SMaterialLayersFunctionsTree.cpp
+++ b/Engine/Source/Editor/MaterialEditor/Private/SMaterialLayersFunctionsTree.cpp
@@ -1914,18 +1914,12 @@ void SMaterialLayersFunctionsInstanceWrapper::Refresh()
 				[
 					SNew(SComboButton)
 						.ContentPadding(0.f)
-						.ComboButtonStyle(&FAppStyle::Get().GetWidgetStyle<FComboButtonStyle>("SimpleComboButton"))
-						.HasDownArrow(false)
+						.ComboButtonStyle(&FAppStyle::Get().GetWidgetStyle<FComboButtonStyle>("StatusBar.StatusBarEllipsisComboButton"))
+						.HasDownArrow(true)
 						.MenuContent()
 						[
 							DetailViewOptions.MakeWidget()
 						]
-						.ButtonContent()
-						[
-							SNew(SImage)
-							.ColorAndOpacity(FSlateColor::UseForeground())
-							.Image(FAppStyle::GetBrush( "PropertyEditor.VerticalDottedLine"))
-						]
 				];
 			
 #else
diff --git a/Engine/Source/Editor/MaterialEditor/Private/SMaterialSubstrateTree.cpp b/Engine/Source/Editor/MaterialEditor/Private/SMaterialSubstrateTree.cpp
--- a/Engine/Source/Editor/MaterialEditor/Private/SMaterialSubstrateTree.cpp
+++ b/Engine/Source/Editor/MaterialEditor/Private/SMaterialSubstrateTree.cpp
@@ -60,16 +60,6 @@
 
 #define LOCTEXT_NAMESPACE "MaterialSubstrateTree"
 
-
-FString CreateNodeKeyForBackgroundParameter(const FString& InNodeKeyBase)
-{
-	return FString(InNodeKeyBase + TEXT("_BackgroundParameter"));
-}
-bool IsNodeKeyForBackgroundParameter(const FString InNodeKey)
-{
-	return (InNodeKey.Find(TEXT("_BackgroundParameter")) >= 0);
-}
-
 /// ===========================================================================================================
 /// SMaterialSubstrateItem Methods
 /// ===========================================================================================================
@@ -210,10 +200,38 @@ TOptional<EItemDropZone> SMaterialSubstrateTreeItem::CanAcceptDrop(const FDragDr
 {
 	TSharedPtr<FLayerDragDropOp> LayerDragDropOperation = DragDropEvent.GetOperationAs<FLayerDragDropOp>();
 	TSharedPtr<FAssetDragDropOp> AssetDragDropOperation = DragDropEvent.GetOperationAs<FAssetDragDropOp>();
-	if (LayerDragDropOperation.IsValid())
+
+	
+	// Drop above or below could CREATE a new layer node:
+	int32 TargetNodeId = StackParameterData->ParameterInfo.Index;
+	int32 ParentNodeId = Tree->FunctionInstance->GetNodeParent(TargetNodeId);
+	auto ChildrenNodeId = Tree->FunctionInstance->GetNodeChildren(ParentNodeId);
+	int32 SiblingIdx = -1;
+	bool FoundTarget = ChildrenNodeId.Find(TargetNodeId, SiblingIdx);
+	const int LastSiblingIdx = ChildrenNodeId.Num() - 1;
+	
+	if (LayerDragDropOperation.IsValid() && LayerDragDropOperation->OwningStack.IsValid())
 	{
-		// PREVENT LAYER DRAG FOR NOW
-		//return DropZone;
+		TSharedPtr<SMaterialSubstrateTreeItem> LayerBeingDraggedPtr = StaticCastWeakPtr<SMaterialSubstrateTreeItem>(LayerDragDropOperation->OwningStack).Pin();
+		int32 SourceNodeId = LayerBeingDraggedPtr->StackParameterData->ParameterInfo.Index;
+		
+		int32 SourceDepth = Tree->FunctionInstance->GetNodeDepth(SourceNodeId);
+		int32 DestDepth = Tree->FunctionInstance->GetNodeDepth(TargetNodeId);
+		
+		// Allow layers at same depth to be moved
+		if (SourceDepth == DestDepth)
+		{
+			// Allow dropping above items or below last item (sibling index 0)
+			if (DropZone == EItemDropZone::AboveItem || ((SiblingIdx == 0) && DropZone == EItemDropZone::BelowItem))
+			{
+				return DropZone;
+			}
+		}
+		// else if a Mat Attribute layer is dragged and dropped Onto a Mat Eval layer, we allow it
+		else if (SourceDepth > DestDepth && DropZone == EItemDropZone::OntoItem)
+		{
+			return DropZone;
+		}
 	}
 	else if (AssetDragDropOperation.IsValid())
 	{
@@ -233,13 +251,6 @@ TOptional<EItemDropZone> SMaterialSubstrateTreeItem::CanAcceptDrop(const FDragDr
 			}
 		}
 
-		// Drop above or below could CREATE a new layer node:
-		int32 TargetNodeId = StackParameterData->ParameterInfo.Index;
-		int32 ParentNodeId = Tree->FunctionInstance->GetNodeParent(TargetNodeId);
-		auto ChildrenNodeId = Tree->FunctionInstance->GetNodeChildren(ParentNodeId);
-		int32 SiblingIdx = -1;
-		bool FoundTarget = ChildrenNodeId.Find(TargetNodeId, SiblingIdx);
-
 		switch (DropZone)
 		{
 		case EItemDropZone::AboveItem:
@@ -281,11 +292,76 @@ FReply SMaterialSubstrateTreeItem::OnLayerDrop(const FDragDropEvent& DragDropEve
 	FScopedTransaction Transaction(NSLOCTEXT("UnrealEd", "MoveLayer", "Move Layer"));
 	Tree->FunctionInstanceHandle->NotifyPreChange();
 	bIsHoveredDragTarget = false;
-	TSharedPtr<FLayerDragDropOp> ArrayDropOp = DragDropEvent.GetOperationAs< FLayerDragDropOp >();
+	TSharedPtr<FLayerDragDropOp> LayerDragDropOp = DragDropEvent.GetOperationAs< FLayerDragDropOp >();
 	TSharedPtr<SMaterialSubstrateTreeItem> LayerPtr = nullptr;
-	if (ArrayDropOp.IsValid() && ArrayDropOp->OwningStack.IsValid())
+	if (LayerDragDropOp.IsValid() && LayerDragDropOp->OwningStack.IsValid())
 	{
-		LayerPtr = StaticCastWeakPtr<SMaterialSubstrateTreeItem>(ArrayDropOp->OwningStack).Pin();
+		LayerPtr = StaticCastWeakPtr<SMaterialSubstrateTreeItem>(LayerDragDropOp->OwningStack).Pin();
+		if (LayerPtr.IsValid())
+		{
+			LayerPtr->bIsBeingDragged = false;
+			TSharedPtr<FSortedParamData> SourcePropertyData = LayerPtr->StackParameterData;
+			TSharedPtr<FSortedParamData> DestPropertyData = StackParameterData;
+			if (SourcePropertyData.IsValid() && DestPropertyData.IsValid())
+			{
+				int32 SourceNodeId = SourcePropertyData->ParameterInfo.Index;
+				int32 DestNodeId = DestPropertyData->ParameterInfo.Index;
+				
+				bool bShouldDuplicate = (DragDropEvent.GetModifierKeys().IsShiftDown() || DragDropEvent.GetModifierKeys().IsControlDown());
+				
+				if (SourceNodeId != DestNodeId || bShouldDuplicate)
+				{
+					int32 SourceParent = Tree->FunctionInstance->GetNodeParent(SourceNodeId);
+					int32 DestParent = Tree->FunctionInstance->GetNodeParent(DestNodeId);
+
+					int32 DestSiblingIdx = -1;
+					int32 SrcSiblingIdx = -1;
+					auto ChildrenNodeId = Tree->FunctionInstance->GetNodeChildren(DestParent);
+
+					ChildrenNodeId.Find(DestNodeId, DestSiblingIdx);
+					ChildrenNodeId.Find(SourceNodeId, SrcSiblingIdx);
+
+					//  Adjust the sibling index if the node is being moved within the same parent
+					if (DropZone == EItemDropZone::BelowItem)
+					{
+						// only allowed to go below first index (or lowest in UI view)
+						DestSiblingIdx = 0;
+					}
+					else if (DropZone == EItemDropZone::AboveItem)
+					{
+						if (SourceParent == DestParent)
+						{
+							// If the original index was before the target index, adjust the target index
+							if (DestSiblingIdx < SrcSiblingIdx || bShouldDuplicate)
+							{
+								++DestSiblingIdx;
+							}
+						}
+						else
+						// if coming from a different parent (moving attribute from one Mat Eval to another)
+						{
+							++DestSiblingIdx;
+						}
+					}
+					else if (DropZone == EItemDropZone::OntoItem)
+					{
+						int32 SourceDepth = Tree->FunctionInstance->GetNodeDepth(SourceNodeId);
+						int32 DestDepth = Tree->FunctionInstance->GetNodeDepth(DestNodeId);
+						check (SourceDepth > DestDepth)
+
+						// if a material attr layer is dropped on to a mat eval layer, we add it to the end (top) of it's children
+						DestSiblingIdx = ChildrenNodeId.Num() - 1;
+						DestParent = DestNodeId;
+					}
+					
+					Tree->FunctionInstance->MoveLayerNode(SourceNodeId, DestParent, DestSiblingIdx, bShouldDuplicate);
+					Tree->FunctionInstanceHandle->NotifyPostChange(EPropertyChangeType::ValueSet);
+					Tree->CreateGroupsWidget();
+					Tree->RequestTreeRefresh();
+				}
+			}
+			return FReply::Handled();
+		}
 	}
 	else
 	{
@@ -375,64 +451,9 @@ FReply SMaterialSubstrateTreeItem::OnLayerDrop(const FDragDropEvent& DragDropEve
 			{
 				return FReply::Handled();
 			}
-			
-		}
-		return FReply::Unhandled();
-	}
-	
-	if (!LayerPtr.IsValid())
-	{
-		return FReply::Unhandled();
-	}
-	
-	LayerPtr->bIsBeingDragged = false;
-	TSharedPtr<FSortedParamData> SwappingPropertyData = LayerPtr->StackParameterData;
-	TSharedPtr<FSortedParamData> SwappablePropertyData = StackParameterData;
-	if (SwappingPropertyData.IsValid() && SwappablePropertyData.IsValid())
-	{
-		if (SwappingPropertyData != SwappablePropertyData)
-		{
-			int32 OriginalIndex = SwappingPropertyData->ParameterInfo.Index;
-			if (SwappingPropertyData->ParameterInfo.Association == EMaterialParameterAssociation::BlendParameter)
-			{
-				OriginalIndex++;
-			}
-
-			int32 NewIndex = SwappablePropertyData->ParameterInfo.Index;
-			if (SwappablePropertyData->ParameterInfo.Association == EMaterialParameterAssociation::BlendParameter)
-			{
-				NewIndex++;
-			}
-
-			if (OriginalIndex != NewIndex)
-			{
-				if (Tree->MaterialEditorInstance->GetMaterialInterface()->IsA<UMaterialInstanceConstant>())
-				{
-					UMaterialInstanceConstant* MaterialConstant = Cast<UMaterialInstanceConstant>(Tree->MaterialEditorInstance->GetMaterialInterface());
-					MaterialConstant->SwapLayerParameterIndices(OriginalIndex, NewIndex);	
-				}
-				
-
-				// Need to save the moving and target expansion states before swapping
-				const bool bOriginalSwappableExpansion = IsItemExpanded();
-				const bool bOriginalSwappingExpansion = LayerPtr->IsItemExpanded();
-
-				TArray<void*> StructPtrs;
-				Tree->FunctionInstanceHandle->AccessRawData(StructPtrs);
-				FMaterialLayersFunctions* MaterialLayersFunctions = static_cast<FMaterialLayersFunctions*>(StructPtrs[0]);
-				MaterialLayersFunctions->MoveBlendedLayer(OriginalIndex, NewIndex);
-
-				Tree->OnExpansionChanged(SwappablePropertyData, bOriginalSwappingExpansion);
-				Tree->OnExpansionChanged(SwappingPropertyData, bOriginalSwappableExpansion);
-				Tree->FunctionInstanceHandle->NotifyPostChange(EPropertyChangeType::ValueSet);
-				Tree->CreateGroupsWidget();
-				Tree->RequestTreeRefresh();
-				Tree->SetParentsExpansionState();
-			}
 		}
 	}
-
-	return FReply::Handled();
+	return FReply::Unhandled();
 }
 
 
@@ -507,9 +528,7 @@ void SMaterialSubstrateTreeItem::Construct(const FArguments& InArgs, const TShar
 		
 	EHorizontalAlignment ValueAlignment = HAlign_Left;
 
-	bool bIsBackgroundItem = IsNodeKeyForBackgroundParameter(StackParameterData->NodeKey);
-	bool bCanReorderLayer = !(bIsBackgroundItem || StackParameterData->ParameterInfo.Index == 0);
-	bool bCanAppendSubLayer = Tree->FunctionInstance->CanAppendLayerNode(StackParameterData->ParameterInfo.Index) && !bIsBackgroundItem;
+	bool bCanAppendSubLayer = Tree->FunctionInstance->CanAppendLayerNode(StackParameterData->ParameterInfo.Index);
 		
 // STACK --------------------------------------------------
 	if (StackParameterData->StackDataType == EStackDataType::Stack)
@@ -701,17 +720,16 @@ void SMaterialSubstrateTreeItem::Construct(const FArguments& InArgs, const TShar
 					
 					];
 		}
-		if (bCanReorderLayer)
-		{
-			MainStack->AddSlot()
-				.HAlign(HAlign_Center)
-				.VAlign(VAlign_Center)
-				.Padding(2.5f, 0)
-				.AutoWidth()
-				[
-					FMaterialPropertyHelpers::MakeStackReorderHandle(SharedThis(this))
-				];
-		}
+		
+		MainStack->AddSlot()
+			.HAlign(HAlign_Center)
+			.VAlign(VAlign_Center)
+			.Padding(2.5f, 0)
+			.AutoWidth()
+			[
+				FMaterialPropertyHelpers::MakeStackReorderHandle(SharedThis(this))
+			];
+		
 
 		if (bCanAppendSubLayer)
 		{
diff --git a/Engine/Source/Editor/MaterialEditor/Public/MaterialPropertyHelpers.h b/Engine/Source/Editor/MaterialEditor/Public/MaterialPropertyHelpers.h
--- a/Engine/Source/Editor/MaterialEditor/Public/MaterialPropertyHelpers.h
+++ b/Engine/Source/Editor/MaterialEditor/Public/MaterialPropertyHelpers.h
@@ -114,6 +114,9 @@ public:
 
 	FLayerDragDropOp(TSharedPtr<IDraggableItem> InOwningStack)
 	{
+#if ENABLE_MATERIAL_LAYER_PROTOTYPE
+		MouseCursor = EMouseCursor::GrabHandClosed;
+#endif
 		OwningStack = InOwningStack;
 		DecoratorWidget = SNew(SBorder)
 			.BorderImage(FAppStyle::GetBrush("Graph.ConnectorFeedback.Border"))
diff --git a/Engine/Source/Runtime/Engine/Private/Materials/MaterialExpressions.cpp b/Engine/Source/Runtime/Engine/Private/Materials/MaterialExpressions.cpp
--- a/Engine/Source/Runtime/Engine/Private/Materials/MaterialExpressions.cpp
+++ b/Engine/Source/Runtime/Engine/Private/Materials/MaterialExpressions.cpp
@@ -18681,6 +18681,91 @@ FLayerNodeIdArray FMaterialLayersFunctionsTree::RemoveNode(FNodeId RemovedNodeId
 	return RemovedIds;
 }
 
+void FMaterialLayersFunctionsTree::MoveNode(FNodeId InNodeId, FNodeId DstParentId, int32 InSiblingIndex)
+{
+	if (InNodeId == InvalidId || InNodeId == DstParentId)
+    {
+        return;
+    }
+
+    FNode* NodeToMove = GetNode(InNodeId);
+    FNodeId OldParentId = NodeToMove->Parent;
+
+    FNode* OldParentNode = GetNode(OldParentId);
+	FNodeId PrevSiblingId = GetSiblingHeadId(InNodeId);
+	
+	// If this node is the head child, update the parent's ChildrenHead
+    if (PrevSiblingId == InNodeId)
+    {
+        // 1. Detach the node from its old parent
+    	if (OldParentNode)
+        {
+            OldParentNode->ChildrenHead = NodeToMove->NextSibling;
+        }
+        else
+        {
+            Root = NodeToMove->NextSibling;
+        }
+    }
+    else
+    // Otherwise, find the previous sibling and update its NextSibling
+	 {
+        while (PrevSiblingId != InvalidId && GetNode(PrevSiblingId)->NextSibling != InNodeId)
+        {
+            PrevSiblingId = GetNode(PrevSiblingId)->NextSibling;
+        }
+
+        if (PrevSiblingId != InvalidId)
+        {
+            // Previous sibling found, update its NextSibling to skip the node being moved
+            GetNode(PrevSiblingId)->NextSibling = NodeToMove->NextSibling;
+        }
+    }
+
+    // 2. Update the node's parent to the new parent
+    NodeToMove->Parent = DstParentId;
+	NodeToMove->NextSibling = InvalidId; // Reset the NextSibling for now
+	
+	FNode* NewParentNode = GetNode(DstParentId);
+	FNodeId SiblingHead = NewParentNode ? NewParentNode->ChildrenHead : Root;
+	
+	// 3. Insert the node at the given sibling index
+	if (InSiblingIndex <= 0 || SiblingHead == InvalidId)
+	{
+		// Insert as the first child if the index is 0 or if there are no children
+		NodeToMove->NextSibling = SiblingHead;
+		if (NewParentNode)
+		{
+			NewParentNode->ChildrenHead = InNodeId;
+		}
+		else
+		{
+			// this is the first depth level with no parent, we set the root in this case
+			Root = InNodeId;
+		}
+	}
+	else
+	{
+		// Insert after the specified sibling index
+		PrevSiblingId = SiblingHead;
+		int32 CurrentIndex = 0;
+
+		// Traverse to find the sibling at InSiblingIndex - 1 (previous sibling)
+		while (PrevSiblingId != InvalidId && CurrentIndex < InSiblingIndex - 1)
+		{
+			PrevSiblingId = GetNode(PrevSiblingId)->NextSibling;
+			++CurrentIndex;
+		}
+
+		// Insert the node after the previous sibling
+		if (PrevSiblingId != InvalidId)
+		{
+			NodeToMove->NextSibling = GetNode(PrevSiblingId)->NextSibling;
+			GetNode(PrevSiblingId)->NextSibling = InNodeId;
+		}
+	}
+}
+
 FString FMaterialLayersFunctionsTree::Log(FString InTab) const
 {
 	FString NewLine = TEXT("\r\n");
@@ -19688,9 +19773,22 @@ void FMaterialLayersFunctions::RemoveLayerNodeAt(FLayerNodeId InNodeId)
 	
 
 
-void FMaterialLayersFunctions::MoveLayerNode(FLayerNodeId InNodeId, FLayerNodeId DstParentId, int32 InSiblingIndex)
+void FMaterialLayersFunctions::MoveLayerNode(FLayerNodeId InNodeId, FLayerNodeId DstParentId, int32 InSiblingIndex, bool bShouldDuplicate)
 {
-
+	if (bShouldDuplicate)
+	{
+		// create new layer function and blend function copies
+		FLayerNodeId NewNode = AppendLayerNode(DstParentId, InSiblingIndex);
+		FLayerNodePayload SourcePayload = GetNodePayload(InNodeId);
+		FLayerNodePayload DestPayload = GetNodePayload(NewNode);
+		Layers[DestPayload.Layer] = DuplicateObject<UMaterialFunctionInterface>(Layers[SourcePayload.Layer], nullptr);
+		Blends[DestPayload.Blend] = DuplicateObject<UMaterialFunctionInterface>(Blends[SourcePayload.Blend], nullptr);
+		EditorOnly.LayerNames[DestPayload.Layer] = EditorOnly.LayerNames[SourcePayload.Layer];
+	}
+	else
+	{
+		Tree.MoveNode(InNodeId, DstParentId, InSiblingIndex);
+	}
 }
 
 #endif // WITH_EDITOR
diff --git a/Engine/Source/Runtime/Engine/Public/Materials/MaterialLayersFunctions.h b/Engine/Source/Runtime/Engine/Public/Materials/MaterialLayersFunctions.h
--- a/Engine/Source/Runtime/Engine/Public/Materials/MaterialLayersFunctions.h
+++ b/Engine/Source/Runtime/Engine/Public/Materials/MaterialLayersFunctions.h
@@ -272,6 +272,9 @@ public:
 	// the payloads are NOT cleared and still hold on their value
 	// return the array of Id of the removed nodes
 	TArray<FNodeId> RemoveNode(FNodeId RemovedNodeId);
+	
+	// Move node from one position to another
+	void MoveNode(FNodeId InNodeId, FNodeId DstParentId, int32 InSiblingIndex);
 
 	// Traverser Utility
 	using FNodeVisitor = void (*)(FNodeId InNodeId, int32 InDepth, int32 InSiblingNum, const FPayload& InPayload);
@@ -627,7 +630,7 @@ struct FMaterialLayersFunctions : public FMaterialLayersFunctionsRuntimeData
 	ENGINE_API bool	 CanRemoveLayerNode(FLayerNodeId InNodeId) const;
 	ENGINE_API void	 RemoveLayerNodeAt(FLayerNodeId InNodeId);
 
-	ENGINE_API void MoveLayerNode(FLayerNodeId InNodeId, FLayerNodeId DstParentId, int32 InSiblingIndex = -1);
+	ENGINE_API void MoveLayerNode(FLayerNodeId InNodeId, FLayerNodeId DstParentId, int32 InSiblingIndex, bool bShouldDuplicate);
 
 #endif // WITH_EDITOR
 
