Substrate Material Layer UI updates
- size/spacing for layer items
- rounded edges
- various other style update

diff --git a/Engine/Source/Editor/MaterialEditor/Private/MaterialEditorStyle.cpp b/Engine/Source/Editor/MaterialEditor/Private/MaterialEditorStyle.cpp
--- a/Engine/Source/Editor/MaterialEditor/Private/MaterialEditorStyle.cpp
+++ b/Engine/Source/Editor/MaterialEditor/Private/MaterialEditorStyle.cpp
@@ -199,67 +199,85 @@ TSharedRef<FSlateStyleSet> FSubstrateMaterialEditorStyle::Create()
 void FSubstrateMaterialEditorStyle::SetupLayerViewStyles(const TSharedRef<FSlateStyleSet>& Style)
 {
 	Style->Set("LayerView.Background", new FSlateRoundedBoxBrush(
-		FStyleColors::Recessed/*FLinearColor(0, 0, 0, 0.25f)*/, 6.0f,
-		/*FStyleColors::Header.GetSpecifiedColor()*/FStyleColors::Recessed, 0.0f));
-
-	Style->Set("LayerView.Details.Background", new FSlateRoundedBoxBrush(
-		/*FStyleColors::Recessed.GetSpecifiedColor()*/FLinearColor(FColor::FromHex(TEXT("#575757"))), 6.0f,
-		FStyleColors::Header.GetSpecifiedColor()/*FLinearColor(1, 1, 1, 0.2f)*/, 20.0f));
+		FStyleColors::Recessed, 6.0f, FStyleColors::Recessed, 0.0f));
 
 	Style->Set("LayerView", FTableViewStyle()
 		.SetBackgroundBrush(*Style->GetBrush("LayerView.Background"))
 	);
-	
-	const float LayerViewItemCornerRadius = 10.0f;
-	const float LayerViewItemBorderWidth = 1.0f;
 
-	const FLinearColor LayerViewItemFillColor = FLinearColor(FColor::FromHex(TEXT("#383838")));//FLinearColor::Transparent;
+	constexpr float LayerViewItemCornerRadius = 10.0f;
+	constexpr float LayerViewItemBorderWidth = 1.0f;
+
+	const FLinearColor LayerViewItem1FillColor = FStyleColors::Dropdown.GetSpecifiedColor();
+	const FLinearColor LayerViewItem2FillColor = FStyleColors::DropdownOutline.GetSpecifiedColor();
 	constexpr FLinearColor LayerViewItemBorderColor = FLinearColor(0, 0, 0, 1.f);
 
-	const FLinearColor LayerItemHoverFillColor = FStyleColors::Recessed.GetSpecifiedColor();
+	const FLinearColor LayerItemHoverFillColor = FStyleColors::Hover.GetSpecifiedColor();
 	constexpr FLinearColor LayerItemHoverBorderColor = FLinearColor(1, 1, 1, 0.2f);
 
-	const FLinearColor LayerItemSelectFillColor = FStyleColors::Header.GetSpecifiedColor();
 	const FLinearColor LayerItemSelectBorderColor = ReplaceColorAlpha(FStyleColors::Select.GetSpecifiedColor(), 0.9f);
 
-	Style->Set("LayerView.Row.Item", new FSlateRoundedBoxBrush(
-		LayerViewItemFillColor, LayerViewItemCornerRadius,
-		LayerViewItemBorderColor, LayerViewItemBorderWidth));
-
+	Style->Set("LayerView.Row1.Item", new FSlateRoundedBoxBrush(
+		LayerViewItem1FillColor, LayerViewItemCornerRadius));
+	Style->Set("LayerView.Row2.Item", new FSlateRoundedBoxBrush(
+				LayerViewItem2FillColor, LayerViewItemCornerRadius,
+				LayerViewItemBorderColor, LayerViewItemBorderWidth));
 	Style->Set("LayerView.Row.Hovered", new FSlateRoundedBoxBrush(
 		LayerItemHoverFillColor, LayerViewItemCornerRadius,
 		LayerItemHoverBorderColor, LayerViewItemBorderWidth));
-
-	Style->Set("LayerView.Row.Selected", new FSlateRoundedBoxBrush(
-		LayerItemSelectFillColor, LayerViewItemCornerRadius,
+	Style->Set("LayerView.Row1.ActiveBrush", new FSlateRoundedBoxBrush(
+		LayerViewItem1FillColor, LayerViewItemCornerRadius,
+		LayerItemSelectBorderColor, LayerViewItemBorderWidth));
+	Style->Set("LayerView.Row1.ActiveHoveredBrush", new FSlateRoundedBoxBrush(
+		LayerViewItem1FillColor, LayerViewItemCornerRadius,
 		LayerItemSelectBorderColor, LayerViewItemBorderWidth));
+	Style->Set("LayerView.Row1.InactiveBrush", new FSlateRoundedBoxBrush(
+		LayerViewItem1FillColor, LayerViewItemCornerRadius));
+	Style->Set("LayerView.Row1.InactiveHoveredBrush", new FSlateRoundedBoxBrush(
+		LayerViewItem1FillColor, LayerViewItemCornerRadius));
 
-	Style->Set("LayerView.Row.ActiveBrush", new FSlateRoundedBoxBrush(
-		LayerItemSelectFillColor, LayerViewItemCornerRadius,
+	Style->Set("LayerView.Row2.ActiveBrush", new FSlateRoundedBoxBrush(
+		LayerViewItem2FillColor, LayerViewItemCornerRadius,
 		LayerItemSelectBorderColor, LayerViewItemBorderWidth));
-	Style->Set("LayerView.Row.ActiveHoveredBrush", new FSlateRoundedBoxBrush(
-		LayerItemSelectFillColor, LayerViewItemCornerRadius,
+	Style->Set("LayerView.Row2.ActiveHoveredBrush", new FSlateRoundedBoxBrush(
+		LayerViewItem2FillColor, LayerViewItemCornerRadius,
 		LayerItemSelectBorderColor, LayerViewItemBorderWidth));
-	Style->Set("LayerView.Row.InactiveBrush", new FSlateRoundedBoxBrush(
-		LayerItemSelectFillColor, LayerViewItemCornerRadius,
+	Style->Set("LayerView.Row2.InactiveBrush", new FSlateRoundedBoxBrush(
+		LayerViewItem2FillColor, LayerViewItemCornerRadius,
 		LayerItemSelectBorderColor, LayerViewItemBorderWidth));
-	Style->Set("LayerView.Row.InactiveHoveredBrush", new FSlateRoundedBoxBrush(
-		LayerItemSelectFillColor, LayerViewItemCornerRadius,
+	Style->Set("LayerView.Row2.InactiveHoveredBrush", new FSlateRoundedBoxBrush(
+		LayerViewItem2FillColor, LayerViewItemCornerRadius,
 		LayerItemSelectBorderColor, LayerViewItemBorderWidth));
 
 	const float DropZoneMargin = 0.25f; 
-	Style->Set("LayerView.Row", FTableRowStyle()
+	Style->Set("LayerView.Row1", FTableRowStyle()
 		.SetTextColor(FStyleColors::Foreground)
 		.SetSelectedTextColor(FStyleColors::ForegroundHover)
-		.SetEvenRowBackgroundBrush(*Style->GetBrush(TEXT("LayerView.Row.Item")))
+		.SetEvenRowBackgroundBrush(*Style->GetBrush(TEXT("LayerView.Row1.Item")))
 		.SetEvenRowBackgroundHoveredBrush(*Style->GetBrush(TEXT("LayerView.Row.Hovered")))
-		.SetOddRowBackgroundBrush(*Style->GetBrush(TEXT("LayerView.Row.Item")))
+		.SetOddRowBackgroundBrush(*Style->GetBrush(TEXT("LayerView.Row1.Item")))
 		.SetOddRowBackgroundHoveredBrush(*Style->GetBrush(TEXT("LayerView.Row.Hovered")))
-		.SetSelectorFocusedBrush(*Style->GetBrush(TEXT("LayerView.Row.Item")))
-		.SetActiveBrush(*Style->GetBrush(TEXT("LayerView.Row.ActiveBrush")))
-		.SetActiveHoveredBrush(*Style->GetBrush(TEXT("LayerView.Row.ActiveHoveredBrush")))
-		.SetInactiveBrush(*Style->GetBrush(TEXT("LayerView.Row.InactiveBrush")))
-		.SetInactiveHoveredBrush(*Style->GetBrush(TEXT("LayerView.Row.InactiveHoveredBrush")))
+		.SetSelectorFocusedBrush(*Style->GetBrush(TEXT("LayerView.Row1.ActiveBrush")))
+		.SetActiveBrush(*Style->GetBrush(TEXT("LayerView.Row1.ActiveBrush")))
+		.SetActiveHoveredBrush(*Style->GetBrush(TEXT("LayerView.Row1.ActiveHoveredBrush")))
+		.SetInactiveBrush(*Style->GetBrush(TEXT("LayerView.Row1.InactiveBrush")))
+		.SetInactiveHoveredBrush(*Style->GetBrush(TEXT("LayerView.Row1.InactiveHoveredBrush")))
+		.SetDropIndicator_Onto(BOX_BRUSH("Common/DropZoneIndicator_Onto", FMargin(4.0f / 16.0f), Style->GetColor(TEXT("Color.Select.Hover"))))
+		.SetDropIndicator_Above(BORDER_BRUSH("Common/LayersDropZoneDashed_Above", FMargin(DropZoneMargin, DropZoneMargin, 0.f, 0.f), Style->GetColor(TEXT("Color.Select.Hover"))))
+		.SetDropIndicator_Below(BORDER_BRUSH("Common/LayersDropZoneDashed_Below", FMargin(DropZoneMargin, 0, 0, DropZoneMargin), Style->GetColor(TEXT("Color.Select.Hover"))))
+	);
+	Style->Set("LayerView.Row2", FTableRowStyle()
+		.SetTextColor(FStyleColors::Foreground)
+		.SetSelectedTextColor(FStyleColors::ForegroundHover)
+		.SetEvenRowBackgroundBrush(*Style->GetBrush(TEXT("LayerView.Row2.Item")))
+		.SetEvenRowBackgroundHoveredBrush(*Style->GetBrush(TEXT("LayerView.Row.Hovered")))
+		.SetOddRowBackgroundBrush(*Style->GetBrush(TEXT("LayerView.Row2.Item")))
+		.SetOddRowBackgroundHoveredBrush(*Style->GetBrush(TEXT("LayerView.Row.Hovered")))
+		.SetSelectorFocusedBrush(*Style->GetBrush(TEXT("LayerView.Row2.ActiveBrush")))
+		.SetActiveBrush(*Style->GetBrush(TEXT("LayerView.Row2.ActiveBrush")))
+		.SetActiveHoveredBrush(*Style->GetBrush(TEXT("LayerView.Row2.ActiveHoveredBrush")))
+		.SetInactiveBrush(*Style->GetBrush(TEXT("LayerView.Row2.InactiveBrush")))
+		.SetInactiveHoveredBrush(*Style->GetBrush(TEXT("LayerView.Row2.InactiveHoveredBrush")))
 		.SetDropIndicator_Onto(BOX_BRUSH("Common/DropZoneIndicator_Onto", FMargin(4.0f / 16.0f), Style->GetColor(TEXT("Color.Select.Hover"))))
 		.SetDropIndicator_Above(BORDER_BRUSH("Common/LayersDropZoneDashed_Above", FMargin(DropZoneMargin, DropZoneMargin, 0.f, 0.f), Style->GetColor(TEXT("Color.Select.Hover"))))
 		.SetDropIndicator_Below(BORDER_BRUSH("Common/LayersDropZoneDashed_Below", FMargin(DropZoneMargin, 0, 0, DropZoneMargin), Style->GetColor(TEXT("Color.Select.Hover"))))
@@ -269,6 +287,8 @@ void FSubstrateMaterialEditorStyle::SetupLayerViewStyles(const TSharedRef<FSlate
 	Style->Set("LayerView.RemoveIcon", new IMAGE_BRUSH("Icons/EditorIcons/LayerRemove", Icon16x16));
 
 	Style->Set("LayerView.Row.Handle", new IMAGE_BRUSH_SVG("Icons/DragHandle", Icon16x16));
+
+	Style->Set("LayerView.Row.OuterRoundBrush", new FSlateRoundedBoxBrush(FStyleColors::Recessed, 15.0f, FStyleColors::Recessed, 1.0f));
 }
 
 void FSubstrateMaterialEditorStyle::SetupTextStyles(const TSharedRef<FSlateStyleSet>& Style)
@@ -314,8 +334,6 @@ void FSubstrateMaterialEditorStyle::SetupTextStyles(const TSharedRef<FSlateStyle
 	Style->Set("LayerView.Row.HandleFont", RegularTextStyle);
 
 	FTextBlockStyle LayerViewItemTextStyle(NormalTextStyle);
-	LayerViewItemTextStyle.SetShadowOffset(FVector2D(1.0f, 1.0f));
-	LayerViewItemTextStyle.SetColorAndOpacity(LayerViewItemTextShadowColor);
 
 	Style->Set("LayerView.Row.HeaderText",
 		FTextBlockStyle(LayerViewItemTextStyle)
diff --git a/Engine/Source/Editor/MaterialEditor/Private/MaterialInstanceEditor.cpp b/Engine/Source/Editor/MaterialEditor/Private/MaterialInstanceEditor.cpp
--- a/Engine/Source/Editor/MaterialEditor/Private/MaterialInstanceEditor.cpp
+++ b/Engine/Source/Editor/MaterialEditor/Private/MaterialInstanceEditor.cpp
@@ -1106,11 +1106,15 @@ TSharedRef<SDockTab> FMaterialInstanceEditor::SpawnTab_LayerProperties(const FSp
 	TSharedRef<SDockTab> SpawnedTab = SNew(SDockTab)
 		.Label(LOCTEXT("MaterialLayerPropertiesTitle", "Layer Parameters"))
 		[
+#if ENABLE_MATERIAL_LAYER_PROTOTYPE
+			MaterialLayersFunctionsInstance.ToSharedRef()
+#else
 			SNew(SBorder)
 			.Padding(4)
 		[
 			MaterialLayersFunctionsInstance.ToSharedRef()
 		]
+#endif
 		];
 
 	AddToSpawnedToolPanels(Args.GetTabId().TabType, SpawnedTab);
diff --git a/Engine/Source/Editor/MaterialEditor/Private/SMaterialLayersFunctionsTree.cpp b/Engine/Source/Editor/MaterialEditor/Private/SMaterialLayersFunctionsTree.cpp
--- a/Engine/Source/Editor/MaterialEditor/Private/SMaterialLayersFunctionsTree.cpp
+++ b/Engine/Source/Editor/MaterialEditor/Private/SMaterialLayersFunctionsTree.cpp
@@ -1789,7 +1789,7 @@ void SMaterialLayersFunctionsInstanceWrapper::Refresh()
 	FOnClicked	OnSiblingButtonClicked = FOnClicked::CreateStatic(&FMaterialPropertyHelpers::OnClickedSaveNewMaterialInstance, MaterialEditorInstance->GetParentMaterialInterface(), ImplicitConv<UObject*>(MaterialEditorInstance));
 
 #if ENABLE_MATERIAL_LAYER_PROTOTYPE
-	const float ThumbnailSize = 64.0f;
+	const float ThumbnailSize = 48.0f;
 	TSharedPtr<SBox> ThumbnailBox;
 	UObject* ThumbnailObject = MaterialEditorInstance->GetMaterialInterface().Get();
 	const TSharedPtr<FAssetThumbnail> AssetThumbnail = MakeShareable(new FAssetThumbnail(ThumbnailObject, ThumbnailSize, ThumbnailSize, NestedTree->GetTreeThumbnailPool()));
@@ -1833,7 +1833,7 @@ void SMaterialLayersFunctionsInstanceWrapper::Refresh()
 					+ SHorizontalBox::Slot()
 					.VAlign(VAlign_Center)
 					.AutoWidth()
-					.Padding(5.0f)
+					.Padding(8.0f, 0.0f)
 					[
 						SNew(SVerticalBox)
 						+ SVerticalBox::Slot()
@@ -1845,6 +1845,7 @@ void SMaterialLayersFunctionsInstanceWrapper::Refresh()
 						]
 
 						+ SVerticalBox::Slot()
+						.Padding(0, 4)
 						.AutoHeight()
 						[
 							SNew(STextBlock)
@@ -1866,7 +1867,7 @@ void SMaterialLayersFunctionsInstanceWrapper::Refresh()
 #endif
 				]
 				+ SVerticalBox::Slot()
-				.Padding(FMargin(0.0f))
+				.Padding(FMargin(2.0f))
 				[
 					NestedTree.ToSharedRef()
 				]
diff --git a/Engine/Source/Editor/MaterialEditor/Private/SMaterialSubstrateTree.cpp b/Engine/Source/Editor/MaterialEditor/Private/SMaterialSubstrateTree.cpp
--- a/Engine/Source/Editor/MaterialEditor/Private/SMaterialSubstrateTree.cpp
+++ b/Engine/Source/Editor/MaterialEditor/Private/SMaterialSubstrateTree.cpp
@@ -19,6 +19,7 @@
 #include "PropertyEditorModule.h"
 #include "ISinglePropertyView.h"
 #include "Materials/MaterialInstanceConstant.h"
+#include "Materials/MaterialExpressionMaterialFunctionCall.h"
 #include "Styling/AppStyle.h"
 #include "PropertyCustomizationHelpers.h"
 #include "ScopedTransaction.h"
@@ -181,9 +182,9 @@ FText SMaterialSubstrateTreeItem::GetLayerName() const
 
 FText SMaterialSubstrateTreeItem::GetLayerDesc() const
 {
-	const FString LayerDescText[] = { TEXT("Material Evaluation"), TEXT("Material Instance Attributes") };
+	const FText LayerDescText[] = { LOCTEXT("Slab", "Slab"), LOCTEXT("Attributes", "Attributes") };
 	
-	return FText::FromString(LayerDescText[GetIndentLevel()]);
+	return LayerDescText[GetIndentLevel()];
 }
 
 void SMaterialSubstrateTreeItem::OnNameChanged(const FText& InText, ETextCommit::Type CommitInfo)
@@ -491,24 +492,24 @@ FReply SMaterialSubstrateTreeItem::OnAddMaterialAttributeClicked(int32 index)
 
 void SMaterialSubstrateTreeItem::Construct(const FArguments& InArgs, const TSharedRef<STableViewBase>& InOwnerTableView)
 {
-	
 	FOnTableRowDragEnter LayerDragDelegate = FOnTableRowDragEnter::CreateSP(this, &SMaterialSubstrateTreeItem::OnLayerDragEnter);
 	FOnTableRowDragLeave LayerDragLeaveDelegate = FOnTableRowDragLeave::CreateSP(this, &SMaterialSubstrateTreeItem::OnLayerDragLeave);
 
+	StackParameterData = InArgs._StackParameterData;
+	MaterialEditorInstance = InArgs._MaterialEditorInstance;
+	Tree = InArgs._InTree;
+	int32 NodeId = StackParameterData->ParameterInfo.Index;
+	int32 Depth = Tree->FunctionInstance->GetNodeDepth(NodeId);
+	FName StyleName = FName(FString::Format(TEXT("LayerView.Row{0}"), {Depth}));
 	STableRow< TSharedPtr<FSortedParamData> >::ConstructInternal(
 		STableRow< TSharedPtr<FSortedParamData> >::FArguments()
-		.Style(FSubstrateMaterialEditorStyle::Get(), "LayerView.Row")
-		.ShowSelection(true)
+		.Style(FSubstrateMaterialEditorStyle::Get(), StyleName)
 		.OnCanAcceptDrop(this, &SMaterialSubstrateTreeItem::CanAcceptDrop)
 		.OnAcceptDrop(this, &SMaterialSubstrateTreeItem::OnLayerDrop)
 		.OnDragEnter(LayerDragDelegate)
 		.OnDragLeave(LayerDragLeaveDelegate),
 		InOwnerTableView
 	);
-	
-	StackParameterData = InArgs._StackParameterData;
-	MaterialEditorInstance = InArgs._MaterialEditorInstance;
-	Tree = InArgs._InTree;
 
 	TSharedRef<SWidget> LeftSideWidget = SNullWidget::NullWidget;
 	TSharedRef<SWidget> RightSideWidget = SNullWidget::NullWidget;
@@ -521,7 +522,7 @@ void SMaterialSubstrateTreeItem::Construct(const FArguments& InArgs, const TShar
 	if (StackParameterData->StackDataType == EStackDataType::Stack)
 	{
 		WrapperWidget->AddSlot()
-		.Padding(10.0f)/*.AutoHeight()*/
+		.Padding(2.0f)
 		[
 			SAssignNew(MainStack, SHorizontalBox)
 		];
@@ -554,7 +555,7 @@ void SMaterialSubstrateTreeItem::Construct(const FArguments& InArgs, const TShar
 					PropertyCustomizationHelpers::MakeVisibilityButton(VisibilityClickedDelegate, FText(), IsEnabledAttribute)
 				];
 		}
-		const float ThumbnailSize = 64.0f;
+		const float ThumbnailSize = 40.0f;
 		TArray<TSharedPtr<FSortedParamData>> AssetChildren = StackParameterData->Children;
 		// Extract the asset elements to represent them as thumbnail boxes
 		for (TSharedPtr<FSortedParamData> AssetChild : AssetChildren)
@@ -595,14 +596,15 @@ void SMaterialSubstrateTreeItem::Construct(const FArguments& InArgs, const TShar
 				// if blend asset, we set it up in the Wrapper Widget at the bottom of the VerticalBox
                 else if (AssetChild->ParameterInfo.Association == BlendParameter)
                 {
-                		
-                	 WrapperWidget->AddSlot()
-                	 .Padding(2.0f)
-                	 .AutoHeight()
-                	 [
-                		SNew(SSeparator)
-                		.Thickness(2.0f) // Set the thickness of the separator
-                	 ];
+                	// only show separator for attributes
+                	EVisibility SeparatorVisibility = (Depth > 1) ? EVisibility::Visible : EVisibility::Hidden;
+                	WrapperWidget->AddSlot()
+					.AutoHeight()
+					[
+					   SNew(SSeparator)
+						.Visibility(SeparatorVisibility)
+					   .Thickness(2.0f) // Set the thickness of the separator
+					];
                 	IDetailTreeNode& Node = *AssetChild->ParameterNode;
                 	TSharedPtr<IDetailPropertyRow> GeneratedRow = StaticCastSharedPtr<IDetailPropertyRow>(Node.GetRow());
                 	IDetailPropertyRow& Row = *GeneratedRow.Get();
@@ -629,7 +631,7 @@ void SMaterialSubstrateTreeItem::Construct(const FArguments& InArgs, const TShar
                 		.DisplayBrowse(false);
                 	
                 	 WrapperWidget->AddSlot()
-                		.Padding(5.0f)
+                		.Padding(8.0f)
                 		.AutoHeight()
                 	 [
                 		SNew(SHorizontalBox)
@@ -641,19 +643,18 @@ void SMaterialSubstrateTreeItem::Construct(const FArguments& InArgs, const TShar
                 		// ]
                 		
                 		+ SHorizontalBox::Slot()
-                		.Padding(5.0f)
-                		.HAlign(HAlign_Left)
+                		.Padding(8.0f, 0.f, 0.0f, 0.0f)
+	                	.AutoWidth()
                 		.VAlign(VAlign_Center)
                 		[
                 			SNew(STextBlock)
-                			.Justification(ETextJustify::Center)
                 			.Text(LOCTEXT("BlendLabel", "Blend"))
                 		]
                 		+ SHorizontalBox::Slot()
-                		.Padding(5.0f)
+                		.Padding(16.0f, 0.f, 0.0f, 0.0f)
+	                	.FillWidth(1.0f)
                 		[
                 			AssetPickerWidget
-                			// AssetChild->ParameterHandle->CreatePropertyValueWidget(false)	
                 		]
                 	 ];
                 }
@@ -662,13 +663,12 @@ void SMaterialSubstrateTreeItem::Construct(const FArguments& InArgs, const TShar
 
 		{
 			HeaderRowWidget->AddSlot()
-				.VAlign(VAlign_Center)
 				.AutoWidth()
-				.Padding(5.0f)
+				.Padding(8.0f, 2.0f, 0.0f, 0.0f)
 				[
 					SNew(SVerticalBox)
 					+ SVerticalBox::Slot()
-					.AutoHeight()
+					.VAlign(VAlign_Center)
 					[
 						SAssignNew(InlineRenameWidget, SInlineEditableTextBlock)
 							.Text(this, &SMaterialSubstrateTreeItem::GetDisplayName)
@@ -676,7 +676,7 @@ void SMaterialSubstrateTreeItem::Construct(const FArguments& InArgs, const TShar
 					]
 
 					+ SVerticalBox::Slot()
-					.AutoHeight()
+					.Padding(0, 4, 0, 0)
 					[
 						SNew(STextBlock)
 						.Text(TAttribute<FText>::Create(TAttribute<FText>::FGetter::CreateSP(this, &SMaterialSubstrateTreeItem::GetLayerDesc)))
@@ -758,6 +758,13 @@ FText SMaterialSubstrateTreeItem::GetDisplayName() const
 }
 int32 SMaterialSubstrateTreeItem::OnPaint(const FPaintArgs& Args, const FGeometry& AllottedGeometry, const FSlateRect& MyCullingRect, FSlateWindowElementList& OutDrawElements, int32 LayerId, const FWidgetStyle& InWidgetStyle, bool bParentEnabled) const
 {
+	int32 TargetNodeId = StackParameterData->ParameterInfo.Index;
+	int32 ParentNodeId = Tree->FunctionInstance->GetNodeParent(TargetNodeId);
+	auto ChildrenNodeId = Tree->FunctionInstance->GetNodeChildren(ParentNodeId);
+	int32 SiblingIdx = -1;
+	ChildrenNodeId.Find(TargetNodeId, SiblingIdx);
+	bool bFirstChild = SiblingIdx == 0; 
+	
 	const FSlateBrush* BackgroundBrushResource = Tree->BackgroundBrush.Get();
 	const int32 IndentLevel = GetIndentLevel();
 	// calculate children + self size
@@ -776,38 +783,58 @@ int32 SMaterialSubstrateTreeItem::OnPaint(const FPaintArgs& Args, const FGeometr
 	}
 	
 	FVector2d OuterBorderSize = AllottedGeometry.GetLocalSize() * FVector2d(1, HeightAdjFactor);
-	static float OffsetX = 30.0f;
-	FVector2D Offset = FVector2D(OffsetX * (IndentLevel+1), 0.0f );
-	FSlateDrawElement::MakeBox(
-			OutDrawElements,
-			++LayerId,
-			AllottedGeometry.ToPaintGeometry(OuterBorderSize - Offset, FSlateLayoutTransform(Offset/2)),
-			BackgroundBrushResource,
-			ESlateDrawEffect::None,
-			FStyleColors::Black.GetSpecifiedColor() * InWidgetStyle.GetColorAndOpacityTint() 
-	);
-	
-	// only need to draw gray part for first level
-	if (IndentLevel <= 0)
+	float OffsetX = 20.0f;
+	FVector2D Offset = FVector2D(OffsetX * (IndentLevel+1), 0.0f);
+
+	const FSlateBrush* RoundedBoxBrush = FSubstrateMaterialEditorStyle::GetBrush("LayerView.Row.OuterRoundBrush");
+	if (IndentLevel > 0)
 	{
-		Offset = FVector2d(OffsetX, -5.0f);
 		FSlateDrawElement::MakeBox(
 				OutDrawElements,
-				++LayerId,
+				LayerId,
 				AllottedGeometry.ToPaintGeometry(OuterBorderSize - Offset, FSlateLayoutTransform(Offset/2)),
+				RoundedBoxBrush,
+				ESlateDrawEffect::None,
+				FStyleColors::Recessed.GetSpecifiedColor() * InWidgetStyle.GetColorAndOpacityTint() 
+		);
+		// TODO: This causes culling issues when the parent item drawing the outer gray for the children goes out of the scroll rect scope
+		
+		// we also need to draw to fill in the rounded edges of the items in between (other than last or first) 
+		// unless they're the first or last items in the list which would have a rounded top and a rounded bottom respectively
+		if (!bFirstChild)
+		{
+			FVector2D NonRoundBorderOffset = (OuterBorderSize - Offset) * FVector2D(1, 0.25f);
+			FVector2D TranslationOffset = Offset/2 + FVector2D(0, OuterBorderSize.Y - 15.0f);
+			FSlateDrawElement::MakeBox(
+				OutDrawElements,
+				LayerId,
+				AllottedGeometry.ToPaintGeometry(NonRoundBorderOffset, FSlateLayoutTransform(TranslationOffset)),
 				BackgroundBrushResource,
 				ESlateDrawEffect::None,
-				FStyleColors::Header.GetSpecifiedColor() * InWidgetStyle.GetColorAndOpacityTint() 
+				FStyleColors::Recessed.GetSpecifiedColor() * InWidgetStyle.GetColorAndOpacityTint() 
+			);
+		}
+	}
+	// only need to draw gray part for first level to cover itself and it's children
+	if (IndentLevel <= 0)
+	{
+    	FVector2D ExpandedOffset = IsItemExpanded() ? FVector2d(0.0f, 10.0f) : FVector2d();
+		Offset = FVector2D(15.0f, IsItemExpanded() ? 10.0f : 20.0f);
+		FSlateDrawElement::MakeBox(
+				OutDrawElements,
+				LayerId,
+				AllottedGeometry.ToPaintGeometry(OuterBorderSize - Offset, FSlateLayoutTransform(Offset/2 + ExpandedOffset)),
+				RoundedBoxBrush,
+				ESlateDrawEffect::None,
+				FStyleColors::Dropdown.GetSpecifiedColor() * InWidgetStyle.GetColorAndOpacityTint() 
 		);
 	}
-	
-	
-	OffsetX = 20.0f;
-	const float ReductionFactorX = 50.0f;
-	const float ReductionFactorY = 15.0f;
-	Offset = FVector2d(OffsetX * IndentLevel + ReductionFactorX, ReductionFactorY);
+		
+	OffsetX = 40.0f;
+	const float ReductionFactorX = 15.0f;
+	const float ReductionFactorY = 10.0f;
+	Offset = FVector2d((OffsetX * IndentLevel) + ReductionFactorX, ReductionFactorY * (2 - IndentLevel));
 	FGeometry BorderGeom = AllottedGeometry.MakeChild(AllottedGeometry.GetLocalSize() - Offset, FSlateLayoutTransform(Offset/2));
-	
 	return STableRow<TSharedPtr<FSortedParamData>>::OnPaint(Args, BorderGeom, MyCullingRect, OutDrawElements, LayerId+1, InWidgetStyle, bParentEnabled);
 }
 
