Improve performance for large scale conversion of virtual textures to and from regular textures.
Texture changes no longer trigger immediate cascading material updates per texture. Material updates are anyway handled separately later.
Similarly material expression updates don't immediated trigger material recompilation.
Material and material functions now recompile as an array instead of individually.
This means we have only 2 material update contexts processed. We used to have one per texture/material expression/material/material function.
Batch conversion of 100s of textures now takes seconds instead of minutes.

diff --git a/Engine/Source/Editor/MaterialEditor/Private/MaterialEditingLibrary.cpp b/Engine/Source/Editor/MaterialEditor/Private/MaterialEditingLibrary.cpp
--- a/Engine/Source/Editor/MaterialEditor/Private/MaterialEditingLibrary.cpp
+++ b/Engine/Source/Editor/MaterialEditor/Private/MaterialEditingLibrary.cpp
@@ -680,45 +680,89 @@ bool UMaterialEditingLibrary::ConnectMaterialExpressions(UMaterialExpression* Fr
 	return bResult;
 }
 
-void UMaterialEditingLibrary::RecompileMaterial(UMaterial* Material)
+
+void UMaterialEditingLibrary::RecompileMaterialInternal(FMaterialUpdateContext& UpdateContext, UMaterial* Material)
 {
-	if (ensureAsRuntimeWarning(Material != nullptr))
+	UpdateContext.AddMaterial(Material);
+
+	// Propagate the change to this material
+	Material->PreEditChange(nullptr);
+	Material->PostEditChange();
+
+	Material->MarkPackageDirty();
+
+	// Force particle components to update their view relevance.
+	for (TObjectIterator<UParticleSystemComponent> It(/*AdditionalExclusionFlags = */RF_ClassDefaultObject, /*bIncludeDerivedClasses = */true, /*InInternalExclusionFlags = */EInternalObjectFlags::Garbage); It; ++It)
 	{
+		It->bIsViewRelevanceDirty = true;
+	}
+
+	// Update parameter names on any child material instances
+	for (TObjectIterator<UMaterialInstance> It(/*AdditionalExclusionFlags = */RF_ClassDefaultObject, /*bIncludeDerivedClasses = */true, /*InInternalExclusionFlags = */EInternalObjectFlags::Garbage); It; ++It)
+	{
+		if (It->Parent == Material)
 		{
-			FMaterialUpdateContext UpdateContext;
+			It->UpdateParameterNames();
+		}
+	}
 
-			UpdateContext.AddMaterial(Material);
+	UMaterialEditingLibrary::RebuildMaterialInstanceEditors(Material);
+}
 
-			// Propagate the change to this material
-			Material->PreEditChange(nullptr);
-			Material->PostEditChange();
+void UMaterialEditingLibrary::RecompileMaterial(UMaterial* Material)
+{
+	if (!ensureAsRuntimeWarning(Material != nullptr))
+	{
+		return;
+	}
+
+	{
+		// Create a material update context so we can safely update materials using this function.
+		FMaterialUpdateContext UpdateContext;
 
-			Material->MarkPackageDirty();
+		RecompileMaterialInternal(UpdateContext, Material);
+	}
 
-			// update the world's viewports
-			FEditorDelegates::RefreshEditor.Broadcast();
-			FEditorSupportDelegates::RedrawAllViewports.Broadcast();
+	// Update the world's viewports
+	FEditorDelegates::RefreshEditor.Broadcast();
+	FEditorSupportDelegates::RedrawAllViewports.Broadcast();
 
-			// Force particle components to update their view relevance.
-			for (TObjectIterator<UParticleSystemComponent> It(/*AdditionalExclusionFlags = */RF_ClassDefaultObject, /*bIncludeDerivedClasses = */true, /*InInternalExclusionFlags = */EInternalObjectFlags::Garbage); It; ++It)
-			{
-				It->bIsViewRelevanceDirty = true;
-			}
+	// Update streaming data
+	FMaterialEditorUtilities::BuildTextureStreamingData(Material);
+}
+
+void UMaterialEditingLibrary::RecompileMaterials(TArray<UMaterial*>& Materials, const FOnItemComplete& OnItemComplete)
+{
+	if (Materials.Num() == 0)
+	{
+		return;
+	}
 
-			// Update parameter names on any child material instances
-			for (TObjectIterator<UMaterialInstance> It(/*AdditionalExclusionFlags = */RF_ClassDefaultObject, /*bIncludeDerivedClasses = */true, /*InInternalExclusionFlags = */EInternalObjectFlags::Garbage); It; ++It)
+	{
+		// Create a material update context so we can safely update materials using this function.
+		FMaterialUpdateContext UpdateContext;
+
+		for (UMaterial* Material : Materials)
+		{
+			if (ensureAsRuntimeWarning(Material != nullptr))
 			{
-				if (It->Parent == Material)
-				{
-					It->UpdateParameterNames();
-				}
+				RecompileMaterialInternal(UpdateContext, Material);
+				OnItemComplete.Execute();
 			}
-
-			// Leaving this scope will update all dependent material instances.
 		}
+	}
+
+	// Update the world's viewports
+	FEditorDelegates::RefreshEditor.Broadcast();
+	FEditorSupportDelegates::RedrawAllViewports.Broadcast();
 
-		UMaterialEditingLibrary::RebuildMaterialInstanceEditors(Material);
-		FMaterialEditorUtilities::BuildTextureStreamingData(Material);
+	// Update streaming data
+	for (UMaterial* Material : Materials)
+	{
+		if (Material != nullptr)
+		{
+			FMaterialEditorUtilities::BuildTextureStreamingData(Material);
+		}
 	}
 }
 
@@ -970,60 +1014,91 @@ void UMaterialEditingLibrary::DeleteMaterialExpressionInFunction(UMaterialFuncti
 	}
 }
 
-
-void UMaterialEditingLibrary::UpdateMaterialFunction(UMaterialFunctionInterface* MaterialFunction, UMaterial* PreviewMaterial)
+void UMaterialEditingLibrary::UpdateMaterialFunctionInternal(FMaterialUpdateContext& UpdateContext, UMaterialFunctionInterface* MaterialFunction, UMaterial* PreviewMaterial)
 {
 	TRACE_CPUPROFILER_EVENT_SCOPE(UMaterialEditingLibrary::UpdateMaterialFunction)
 
-	if (MaterialFunction)
+	// mark the function as changed
+	MaterialFunction->ForceRecompileForRendering(UpdateContext, PreviewMaterial);
+	MaterialFunction->MarkPackageDirty();
+
 	{
-		// Create a material update context so we can safely update materials using this function.
-		{
-			FMaterialUpdateContext UpdateContext;
+		TRACE_CPUPROFILER_EVENT_SCOPE(UpdateAllMaterialInstances)
 
-			// mark the function as changed
-			MaterialFunction->ForceRecompileForRendering(UpdateContext, PreviewMaterial);
-			MaterialFunction->MarkPackageDirty();
+		// Go through all function instances in memory and recompile them if they are children
+		for (TObjectIterator<UMaterialFunctionInstance> It(/*AdditionalExclusionFlags = */RF_ClassDefaultObject, /*bIncludeDerivedClasses = */true, /*InInternalExclusionFlags = */EInternalObjectFlags::Garbage); It; ++It)
+		{
+			UMaterialFunctionInstance* FunctionInstance = *It;
 
+			TArray<UMaterialFunctionInterface*> Functions;
+			FunctionInstance->GetDependentFunctions(Functions);
+			if (Functions.Contains(MaterialFunction))
 			{
-				TRACE_CPUPROFILER_EVENT_SCOPE(UpdateAllMaterialInstances)
+				FunctionInstance->UpdateParameterSet();
+				FunctionInstance->ForceRecompileForRendering(UpdateContext, PreviewMaterial);
 
-				// Go through all function instances in memory and recompile them if they are children
-				for (TObjectIterator<UMaterialFunctionInstance> It(/*AdditionalExclusionFlags = */RF_ClassDefaultObject, /*bIncludeDerivedClasses = */true, /*InInternalExclusionFlags = */EInternalObjectFlags::Garbage); It; ++It)
-				{
-					UMaterialFunctionInstance* FunctionInstance = *It;
+				// ForceRecompileForRendering will update StateId, so need to mark the package as dirty
+				FunctionInstance->MarkPackageDirty();
+			}
+		}
+	}
 
-					TArray<UMaterialFunctionInterface*> Functions;
-					FunctionInstance->GetDependentFunctions(Functions);
-					if (Functions.Contains(MaterialFunction))
-					{
-						FunctionInstance->UpdateParameterSet();
-						FunctionInstance->ForceRecompileForRendering(UpdateContext, PreviewMaterial);
+	// Notify material editor for any materials that we are updating
+	for (UMaterialInterface* CurrentMaterial : UpdateContext.GetUpdatedMaterials())
+	{
+		if (IMaterialEditor* MaterialEditor = MaterialEditingLibraryImpl::FindMaterialEditorForAsset(CurrentMaterial))
+		{
+			MaterialEditor->NotifyExternalMaterialChange();
+		}
+	}
 
-						// ForceRecompileForRendering will update StateId, so need to mark the package as dirty
-						FunctionInstance->MarkPackageDirty();
-					}
-				}
-			}
+	UMaterialFunction* BaseFunction = MaterialFunction->GetBaseFunction();
+	UMaterialEditingLibrary::RebuildMaterialInstanceEditors(BaseFunction);
+}
+
+void UMaterialEditingLibrary::UpdateMaterialFunction(UMaterialFunctionInterface* MaterialFunction, UMaterial* PreviewMaterial)
+{
+	if (MaterialFunction == nullptr)
+	{
+		return;
+	}
+	
+	{
+		// Create a material update context so we can safely update materials using this function.
+		FMaterialUpdateContext UpdateContext;
+	
+		UpdateMaterialFunctionInternal(UpdateContext, MaterialFunction, PreviewMaterial);
+	}
+
+	// Update the world's viewports	
+	FEditorDelegates::RefreshEditor.Broadcast();
+	FEditorSupportDelegates::RedrawAllViewports.Broadcast();
+}
+
+void UMaterialEditingLibrary::UpdateMaterialFunctions(TArray<UMaterialFunctionInterface*>& MaterialFunctions, const FOnItemComplete& OnItemComplete)
+{
+	if (MaterialFunctions.Num() == 0)
+	{
+		return;
+	}
 
-			// Notify material editor for any materials that we are updating
-			for (UMaterialInterface* CurrentMaterial : UpdateContext.GetUpdatedMaterials())
+	{
+		// Create a material update context so we can safely update materials using this function.
+		FMaterialUpdateContext UpdateContext;
+
+		for (UMaterialFunctionInterface* MaterialFunction : MaterialFunctions)
+		{
+			if (MaterialFunction != nullptr)
 			{
-				if (IMaterialEditor* MaterialEditor = MaterialEditingLibraryImpl::FindMaterialEditorForAsset(CurrentMaterial))
-				{
-					MaterialEditor->NotifyExternalMaterialChange();
-				}
+				UpdateMaterialFunctionInternal(UpdateContext, MaterialFunction, nullptr);
+				OnItemComplete.Execute();
 			}
 		}
-
-		// update the world's viewports	
-		UMaterialFunction* BaseFunction = MaterialFunction->GetBaseFunction();
-
-		UMaterialEditingLibrary::RebuildMaterialInstanceEditors(BaseFunction);
-		FEditorDelegates::RefreshEditor.Broadcast();
-		FEditorSupportDelegates::RedrawAllViewports.Broadcast();
 	}
 
+	// Update the world's viewports	
+	FEditorDelegates::RefreshEditor.Broadcast();
+	FEditorSupportDelegates::RedrawAllViewports.Broadcast();
 }
 
 void UMaterialEditingLibrary::LayoutMaterialFunctionExpressions(UMaterialFunction* MaterialFunction)
diff --git a/Engine/Source/Editor/MaterialEditor/Public/MaterialEditingLibrary.h b/Engine/Source/Editor/MaterialEditor/Public/MaterialEditingLibrary.h
--- a/Engine/Source/Editor/MaterialEditor/Public/MaterialEditingLibrary.h
+++ b/Engine/Source/Editor/MaterialEditor/Public/MaterialEditingLibrary.h
@@ -164,6 +164,16 @@ public:
 	UFUNCTION(BlueprintCallable, Category = "MaterialEditing")
 	static void RecompileMaterial(UMaterial* Material);
 
+	/** Item complete callback for use when processing arrays of items. Allows higher level to tick UI etc. */
+	DECLARE_DELEGATE(FOnItemComplete)
+
+	/** 
+	 *	Trigger a recompile of an array of materials.
+	 * 	Faster than individual calls to RecompileMaterial().
+	 *  Calls OnItemComplete after each material is processed.
+	 */
+	static void RecompileMaterials(TArray<UMaterial*>& Materials, FOnItemComplete const& OnItemComplete);
+
 	/**
 	 *	Layouts the expressions in a grid pattern
 	 */
@@ -254,6 +264,14 @@ public:
 	UFUNCTION(BlueprintCallable, Category = "MaterialEditing", meta = (HidePin = "PreviewMaterial"))
 	static void UpdateMaterialFunction(UMaterialFunctionInterface* MaterialFunction, UMaterial* PreviewMaterial = nullptr);
 
+	/**
+	 *	Update an array of Material Functions after edits have been made.
+	 *	Will recompile any Materials that use the supplied Material Function.
+	 *  Faster than individual calls to UpdateMaterialFunction().
+	 *  Calls OnItemComplete after each material function is processed.
+	 */
+	static void UpdateMaterialFunctions(TArray<UMaterialFunctionInterface*>& MaterialFunctions, FOnItemComplete const& OnItemComplete);
+
 	/**
 	 *	Layouts the expressions in a grid pattern
 	 */
@@ -394,4 +412,8 @@ public:
 	/** Returns any nanite override material for the given material */
 	UFUNCTION(BlueprintCallable, Category = "MaterialEditing")
 	static UMaterialInterface* GetNaniteOverrideMaterial(UMaterialInterface* Material);
+
+private:
+	static void RecompileMaterialInternal(FMaterialUpdateContext& UpdateContext, UMaterial* Material);
+	static void UpdateMaterialFunctionInternal(FMaterialUpdateContext& UpdateContext, UMaterialFunctionInterface* MaterialFunction, UMaterial* PreviewMaterial);
 };
\ No newline at end of file
diff --git a/Engine/Source/Editor/VirtualTexturingEditor/Private/VirtualTextureConversionWorker.cpp b/Engine/Source/Editor/VirtualTexturingEditor/Private/VirtualTextureConversionWorker.cpp
--- a/Engine/Source/Editor/VirtualTexturingEditor/Private/VirtualTextureConversionWorker.cpp
+++ b/Engine/Source/Editor/VirtualTexturingEditor/Private/VirtualTextureConversionWorker.cpp
@@ -515,21 +515,6 @@ void FVirtualTextureConversionWorker::FilterList(int32 SizeThreshold)
 	FindAllTexturesAndMaterials(Materials, Functions, Textures);
 }
 
-static void MarkTextureExpressionModified(UMaterialExpressionTextureBase* Expression)
-{
-	FPropertyChangedEvent SamplerTypeChangeEvent(UMaterialExpressionTextureBase::StaticClass()->FindPropertyByName(GET_MEMBER_NAME_CHECKED(UMaterialExpressionTextureBase, SamplerType)));
-	FPropertyChangedEvent TextureChangeEvent(UMaterialExpressionTextureBase::StaticClass()->FindPropertyByName(GET_MEMBER_NAME_CHECKED(UMaterialExpressionTextureBase, Texture)));
-
-	Expression->Modify();
-
-	// Nofity that we changed SamplerType
-	Expression->PostEditChangeProperty(SamplerTypeChangeEvent);
-
-	// Also notify that we changed Texture (technically didn't modify this property)
-	// This way code in FMaterialExpressionTextureBaseDetails will see the event, and refresh the list of valid sampler types for the updated texture
-	Expression->PostEditChangeProperty(TextureChangeEvent);
-}
-
 void FVirtualTextureConversionWorker::DoConvert()
 {
 	const bool bVirtualTextureEnable = !bConvertBackward;
@@ -594,10 +579,9 @@ void FVirtualTextureConversionWorker::DoConvert()
 
 				if (TextureToUpdate->VirtualTextureStreaming != bVirtualTextureEnable)
 				{
-					FPropertyChangedEvent PropertyChangeEvent(UTexture::StaticClass()->FindPropertyByName(GET_MEMBER_NAME_CHECKED(UTexture, VirtualTextureStreaming)));
-					TextureToUpdate->Modify();
+					TextureToUpdate->PreEditChange(nullptr);
 					TextureToUpdate->VirtualTextureStreaming = bVirtualTextureEnable;
-					TextureToUpdate->PostEditChangeProperty(PropertyChangeEvent);
+					TextureToUpdate->PostEditChange();
 				}
 			}
 		}
@@ -605,10 +589,12 @@ void FVirtualTextureConversionWorker::DoConvert()
 
 	SlowTask.EnterProgressFrame();
 	{
-		FScopedSlowTask MaterialTask(static_cast<float>(Materials.Num() + Functions.Num()), LOCTEXT("ConvertToVT_Progress_MaterialTask", "Updating materials..."));
-		FMaterialUpdateContext UpdateContext;
+		FScopedSlowTask MaterialTask(static_cast<float>((Materials.Num() + Functions.Num()) * 2), LOCTEXT("ConvertToVT_Progress_MaterialTask", "Updating materials..."));
 
-		TMap<UMaterialFunctionInterface*, TArray<UMaterial*>> FunctionToMaterialMap;
+		TArray<UMaterial*> ModifiedMaterials;
+		ModifiedMaterials.Reserve(Materials.Num());
+		TArray<UMaterialFunctionInterface*> ModifiedMaterialFunctions;
+		ModifiedMaterialFunctions.Reserve(Functions.Num());
 
 		for (UMaterial *Mat : Materials)
 		{
@@ -650,24 +636,17 @@ void FVirtualTextureConversionWorker::DoConvert()
 						}
 						if (TextureCopy != nullptr || OldType != TexExpr->SamplerType)
 						{
-							MarkTextureExpressionModified(TexExpr);
+							TexExpr->Modify();
 							MatModified = true;
 						}
 					}
 				}
 			}
 
-			TArray<UMaterialFunctionInterface *>MaterialFunctions;
-			Mat->GetDependentFunctions(MaterialFunctions);
-			for (auto Function : MaterialFunctions)
-			{
-				FunctionToMaterialMap.FindOrAdd(Function).AddUnique(Mat);
-			}
-
 			if (MatModified)
 			{
 				UE_LOG(LogVirtualTextureConversion, Display, TEXT("Material %s added to update list."), *Mat->GetName());
-				UMaterialEditingLibrary::RecompileMaterial(Mat);
+				ModifiedMaterials.Add(Mat);
 			}
 			else
 			{
@@ -675,6 +654,8 @@ void FVirtualTextureConversionWorker::DoConvert()
 			}
 		}
 
+		UMaterialEditingLibrary::RecompileMaterials(ModifiedMaterials, UMaterialEditingLibrary::FOnItemComplete::CreateLambda([&MaterialTask]() { MaterialTask.EnterProgressFrame(); }));
+
 		for (UMaterialFunctionInterface *Func : Functions)
 		{
 			UE_LOG(LogVirtualTextureConversion, Display, TEXT("Function %s"), *Func->GetName());
@@ -715,7 +696,7 @@ void FVirtualTextureConversionWorker::DoConvert()
 						}
 						if (TextureCopy != nullptr || TexExpr->SamplerType != OldType)
 						{
-							MarkTextureExpressionModified(TexExpr);
+							TexExpr->Modify();
 							FuncModified = true;
 						}
 					}
@@ -724,8 +705,8 @@ void FVirtualTextureConversionWorker::DoConvert()
 
 			if (FuncModified)
 			{
-				UMaterialEditingLibrary::UpdateMaterialFunction(Func, nullptr);
 				UE_LOG(LogVirtualTextureConversion, Display, TEXT("Function %s added to update list."), *Func->GetName());
+				ModifiedMaterialFunctions.Add(Func);
 			}
 			else
 			{
@@ -733,6 +714,8 @@ void FVirtualTextureConversionWorker::DoConvert()
 			}
 		}
 
+		UMaterialEditingLibrary::UpdateMaterialFunctions(ModifiedMaterialFunctions, UMaterialEditingLibrary::FOnItemComplete::CreateLambda([&MaterialTask]() { MaterialTask.EnterProgressFrame(); }));
+
 		// update the world's viewports
 		UE_LOG(LogVirtualTextureConversion, Display, TEXT("Broadcasting to editor."));
 		FEditorDelegates::RefreshEditor.Broadcast();
