Fixed various crashes caused by deleting expressions inside Material Function as well performing undo/redo operations after deletion.

diff --git a/Engine/Source/Editor/MaterialEditor/Private/MaterialEditor.cpp b/Engine/Source/Editor/MaterialEditor/Private/MaterialEditor.cpp
--- a/Engine/Source/Editor/MaterialEditor/Private/MaterialEditor.cpp
+++ b/Engine/Source/Editor/MaterialEditor/Private/MaterialEditor.cpp
@@ -5235,7 +5235,7 @@ void FMaterialEditor::OnResetToDefault(const FToolMenuContext& InMenuContext) co
 		UMaterialEditorOnlyData* EditorOnlyData = Material->GetEditorOnlyData();
 		switch (PropertyId)
 		{
-		case MP_BaseColor:
+			case MP_BaseColor:
 				EditorOnlyData->BaseColor.Constant = FLinearColor::Gray;
 				TargetPin->GetSchema()->TrySetDefaultValue(*TargetPin, EditorOnlyData->BaseColor.GetDefaultValue());
 				break;
@@ -5937,6 +5937,12 @@ UMaterialExpression* FMaterialEditor::CreateNewMaterialExpression(UClass* NewExp
 		const FScopedTransaction Transaction( NSLOCTEXT("UnrealEd", "MaterialEditorNewExpression", "Material Editor: New Expression") );
 		Material->Modify();
 
+		if (MaterialFunction)
+		{
+			MaterialFunction->Modify();
+			MaterialFunction->GetEditorOnlyData()->Modify();
+		}
+
 		UObject* SelectedAsset = nullptr;
 		if (bAutoAssignResource)
 		{
@@ -6223,10 +6229,12 @@ void FMaterialEditor::DeleteNodes(const TArray<UEdGraphNode*>& NodesToDelete, bo
 			}
 			RegenerateCodeView();
 		}
+
 		InvalidateSubstrateConversionVersion(Material);
 		UpdatePreviewMaterial();
 		Material->MarkPackageDirty();
 		SetMaterialDirty();
+		UpdateOriginalMaterial();
 
 		if ( bHaveExpressionsToDelete )
 		{
@@ -6299,6 +6307,12 @@ void FMaterialEditor::DeleteNodesInternal(const TArray<class UEdGraphNode*>& Nod
 {
 	Material->Modify();
 
+	if (MaterialFunction)
+	{
+		MaterialFunction->Modify();
+		MaterialFunction->GetEditorOnlyData()->Modify();
+	}
+
 	for (int32 Index = 0; Index < NodesToDelete.Num(); ++Index)
 	{
 		if (NodesToDelete[Index]->CanUserDeleteNode())
@@ -7230,6 +7244,11 @@ void FMaterialEditor::PostUndo(bool bSuccess)
 {
 	if (bSuccess)
 	{	
+		if (MaterialFunction)
+		{
+			Material->AssignExpressionCollection(MaterialFunction->GetExpressionCollection());
+		}
+
 		if (TSharedPtr<SGraphEditor> FocusedGraphEd = FocusedGraphEdPtr.Pin())
 		{
 			FocusedGraphEd->ClearSelectionSet();
@@ -7252,10 +7271,9 @@ void FMaterialEditor::PostUndo(bool bSuccess)
 
 		// Update the current preview material.
 		UpdatePreviewMaterial();
-
 		UpdatePreviewViewportsVisibility();
-
 		RefreshExpressionPreviews();
+		UpdateOriginalMaterial();
 
 		// Remove any tabs are that are pending kill or otherwise invalid UObject pointers.
 		bool bNeedOpenGraphEditor = false;
