- Removed dubious code in MaterialEditor.cpp that prevented orphaned nodes to be deleted.
- Unified calls to Material->Modify() to MaterialModify() to make sure Modify() is called on the material and its editor only data, as well as the edited material function (if any) and its editor only data.

diff --git a/Engine/Source/Editor/MaterialEditor/Private/MaterialEditor.cpp b/Engine/Source/Editor/MaterialEditor/Private/MaterialEditor.cpp
--- a/Engine/Source/Editor/MaterialEditor/Private/MaterialEditor.cpp
+++ b/Engine/Source/Editor/MaterialEditor/Private/MaterialEditor.cpp
@@ -4185,7 +4185,7 @@ void FMaterialEditor::OnPromoteObjects()
 	if (SelectedNodes.Num() > 0)
 	{
 		const FScopedTransaction Transaction(LOCTEXT("MaterialEditorPromote", "Material Editor: Promote"));
-		Material->Modify();
+		ModifyMaterial();
 		Material->MaterialGraph->Modify();
 		TArray<class UEdGraphNode*> NodesToDelete;
 		TArray<class UEdGraphNode*> NodesToSelect;
@@ -4277,7 +4277,7 @@ void FMaterialEditor::OnConvertObjects()
 	if (SelectedNodes.Num() > 0)
 	{
 		const FScopedTransaction Transaction( LOCTEXT("MaterialEditorConvert", "Material Editor: Convert") );
-		Material->Modify();
+		ModifyMaterial();
 		Material->MaterialGraph->Modify();
 		TArray<class UEdGraphNode*> NodesToDelete;
 		TArray<class UEdGraphNode*> NodesToSelect;
@@ -4533,7 +4533,7 @@ void FMaterialEditor::OnConvertTextures()
 	if (SelectedNodes.Num() > 0)
 	{
 		const FScopedTransaction Transaction( LOCTEXT("MaterialEditorConvertTexture", "Material Editor: Convert to Texture") );
-		Material->Modify();
+		ModifyMaterial();
 		Material->MaterialGraph->Modify();
 		TArray<class UEdGraphNode*> NodesToDelete;
 		TArray<class UEdGraphNode*> NodesToSelect;
@@ -5935,13 +5935,7 @@ UMaterialExpression* FMaterialEditor::CreateNewMaterialExpression(UClass* NewExp
 	UMaterialExpression* NewExpression = NULL;
 	{
 		const FScopedTransaction Transaction( NSLOCTEXT("UnrealEd", "MaterialEditorNewExpression", "Material Editor: New Expression") );
-		Material->Modify();
-
-		if (MaterialFunction)
-		{
-			MaterialFunction->Modify();
-			MaterialFunction->GetEditorOnlyData()->Modify();
-		}
+		ModifyMaterial();
 
 		UObject* SelectedAsset = nullptr;
 		if (bAutoAssignResource)
@@ -5989,7 +5983,7 @@ UMaterialExpressionComposite* FMaterialEditor::CreateNewMaterialExpressionCompos
 	UMaterialExpressionComposite* NewComposite = nullptr;
 	{
 		const FScopedTransaction Transaction(NSLOCTEXT("UnrealEd", "MaterialEditorNewComposite", "Material Editor: New Composite"));
-		Material->Modify();
+		ModifyMaterial();
 
 		UObject* ExpressionOuter = Material;
 		if (MaterialFunction)
@@ -6050,7 +6044,7 @@ UMaterialExpressionComment* FMaterialEditor::CreateNewMaterialExpressionComment(
 	UMaterialExpressionComment* NewComment = NULL;
 	{
 		const FScopedTransaction Transaction(NSLOCTEXT("UnrealEd", "MaterialEditorCreateComment", "Material Editor: Create comment"));
-		Material->Modify();
+		ModifyMaterial();
 
 		UObject* ExpressionOuter = Material;
 		if (MaterialFunction)
@@ -6158,6 +6152,18 @@ void FMaterialEditor::JumpToExpression(UMaterialExpression* Expression)
 	JumpToNode(ExpressionNode);
 }
 
+void FMaterialEditor::ModifyMaterial()
+{
+	Material->Modify();
+	Material->GetEditorOnlyData()->Modify();
+
+	if (MaterialFunction)
+	{
+		MaterialFunction->Modify();
+		MaterialFunction->GetEditorOnlyData()->Modify();
+	}
+}
+
 void FMaterialEditor::SelectAllNodes()
 {
 	if (TSharedPtr<SGraphEditor> FocusedGraphEd = FocusedGraphEdPtr.Pin())
@@ -6305,24 +6311,12 @@ void FMaterialEditor::DeleteSelectedDuplicatableNodes()
 
 void FMaterialEditor::DeleteNodesInternal(const TArray<class UEdGraphNode*>& NodesToDelete, bool& bHaveExpressionsToDelete, bool& bPreviewExpressionDeleted)
 {
-	Material->Modify();
-
-	if (MaterialFunction)
-	{
-		MaterialFunction->Modify();
-		MaterialFunction->GetEditorOnlyData()->Modify();
-	}
+	ModifyMaterial();
 
 	for (int32 Index = 0; Index < NodesToDelete.Num(); ++Index)
 	{
 		if (NodesToDelete[Index]->CanUserDeleteNode())
 		{
-			// If this is a user-selected pinbase, don't allow the delete to pass
-			if (Cast<UMaterialGraphNode_PinBase>(NodesToDelete[Index]) && GetSelectedNodes().Contains(NodesToDelete[Index]))
-			{
-				continue;
-			}
-
 			if (UMaterialGraphNode* GraphNode = Cast<UMaterialGraphNode>(NodesToDelete[Index]))
 			{
 				// Break all node links first so that we don't update the material before deleting
@@ -6524,7 +6518,8 @@ void FMaterialEditor::PasteNodesHereFromBuffer(const FVector2D& Location, const
 	// Undo/Redo support
 	const FScopedTransaction Transaction( NSLOCTEXT("UnrealEd", "MaterialEditorPaste", "Material Editor: Paste") );
 	Material->MaterialGraph->Modify();
-	Material->Modify();
+	ModifyMaterial();
+	Material->GetEditorOnlyData()->Modify();
 
 	UMaterialGraph* ExpressionGraph = Graph ? ToRawPtr(CastChecked<UMaterialGraph>(const_cast<UEdGraph*>(Graph))) : ToRawPtr(Material->MaterialGraph);
 	ExpressionGraph->Modify();
@@ -7041,7 +7036,7 @@ void FMaterialEditor::OnCollapseNodes()
 	if (CollapsableNodes.Num())
 	{
 		const FScopedTransaction Transaction(FGraphEditorCommands::Get().CollapseNodes->GetDescription());
-		Material->Modify();
+		ModifyMaterial();
 		Material->MaterialGraph->Modify();
 
 		CollapseNodes(CollapsableNodes);
@@ -7070,7 +7065,7 @@ bool FMaterialEditor::CanCollapseNodes() const
 void FMaterialEditor::OnExpandNodes()
 {
 	const FScopedTransaction Transaction(FGraphEditorCommands::Get().ExpandNodes->GetLabel());
-	Material->Modify();
+	ModifyMaterial();
 	Material->MaterialGraph->Modify();
 
 	TSet<UEdGraphNode*> ExpandedNodes;
@@ -7965,7 +7960,7 @@ void FMaterialEditor::CleanUnusedExpressions()
 			// Kill off expressions referenced by the material that aren't reachable.
 			const FScopedTransaction Transaction( NSLOCTEXT("UnrealEd", "MaterialEditorCleanUnusedExpressions", "Material Editor: Clean Unused Expressions") );
 				
-			Material->Modify();
+			ModifyMaterial();
 			Material->MaterialGraph->Modify();
 
 			for (int32 Index = 0; Index < UnusedNodes.Num(); ++Index)
diff --git a/Engine/Source/Editor/MaterialEditor/Private/MaterialEditor.h b/Engine/Source/Editor/MaterialEditor/Private/MaterialEditor.h
--- a/Engine/Source/Editor/MaterialEditor/Private/MaterialEditor.h
+++ b/Engine/Source/Editor/MaterialEditor/Private/MaterialEditor.h
@@ -555,6 +555,9 @@ protected:
 	 */
 	FReply OnSpawnGraphNodeByShortcut(FInputChord InChord, const FVector2D& InPosition, UEdGraph* InGraph);
 
+	/** Calls Modify() on the Material and its editor only data, as well as MaterialFunction (and its editor only data) if present */
+	void ModifyMaterial();
+
 	/** Select every node in the graph */
 	void SelectAllNodes();
 	/** Whether we can select every node */
