Update thumbnails when parameters change value

diff --git a/Engine/Source/Editor/MaterialEditor/Private/MaterialEditor.cpp b/Engine/Source/Editor/MaterialEditor/Private/MaterialEditor.cpp
--- a/Engine/Source/Editor/MaterialEditor/Private/MaterialEditor.cpp
+++ b/Engine/Source/Editor/MaterialEditor/Private/MaterialEditor.cpp
@@ -1911,7 +1911,7 @@ void FMaterialEditor::OnChangeBreadCrumbGraph(UEdGraph* InGraph)
 void FMaterialEditor::GeneratorRowsRefreshed()
 {
 	MaterialParametersOverviewWidget->Refresh();
-#ifndef ENABLE_MATERIAL_LAYER_PROTOTYPE
+#if !ENABLE_MATERIAL_LAYER_PROTOTYPE
 	if (MaterialLayersFunctionsInstance)
 	{
 		MaterialLayersFunctionsInstance->Refresh();
diff --git a/Engine/Source/Editor/MaterialEditor/Private/MaterialEditorInstanceDetailCustomization.cpp b/Engine/Source/Editor/MaterialEditor/Private/MaterialEditorInstanceDetailCustomization.cpp
--- a/Engine/Source/Editor/MaterialEditor/Private/MaterialEditorInstanceDetailCustomization.cpp
+++ b/Engine/Source/Editor/MaterialEditor/Private/MaterialEditorInstanceDetailCustomization.cpp
@@ -70,19 +70,6 @@ FMaterialInstanceParameterDetails::FMaterialInstanceParameterDetails(UMaterialEd
 {
 }
 
-void FMaterialInstanceParameterDetails::CollectStackItemsRecursively(TSharedPtr<FSortedParamData> Item, TArray<TSharedPtr<FSortedParamData>>& OutGroupsContainer)
-{
-	for(TSharedPtr<FSortedParamData> Child : Item->Children)
-	{
-		if (Child->StackDataType == EStackDataType::Stack)
-		{
-			OutGroupsContainer.Add(Child);
-		}
-		
-		CollectStackItemsRecursively(Child, OutGroupsContainer);
-	}
-}
-
 TOptional<float> FMaterialInstanceParameterDetails::OnGetValue(TSharedRef<IPropertyHandle> PropertyHandle)
 {
 	float Value = 0.0f;
@@ -130,44 +117,36 @@ void FMaterialInstanceParameterDetails::CustomizeDetails(IDetailLayoutBuilder& D
 		check (SelectedItem->StackDataType == EStackDataType::Stack)
 
 		// we should now gather all sub-stack items to loop through all together
-		TArray<TSharedPtr<FSortedParamData>> StacksCollection;
-		CollectStackItemsRecursively(SelectedItem, StacksCollection);
+		TArray<TSharedPtr<FSortedParamData>> AssetStacksCollection;
+		MaterialLayersFunctionsInstance->NestedTree->CollectAssetStackItemsRecursively(SelectedItem, AssetStacksCollection);
 
-		if (StacksCollection.IsEmpty())
-		{
-			// we add the selected item in this case, as this indicates an inner stack item was selected
-			StacksCollection.Add(SelectedItem);
-		}
-		
-		for (TSharedPtr<FSortedParamData> StackItem : StacksCollection)
+		// we go through list of assets
+		for (TSharedPtr<FSortedParamData> AssetItem : AssetStacksCollection)
 		{
-			// // we go through list of assets
-			for(TSharedPtr<FSortedParamData> ChildAsset : StackItem->Children)
+
+			for (TSharedPtr<FSortedParamData> GroupParamData : AssetItem->Children)
 			{
-				for(TSharedPtr<FSortedParamData> GroupParamData : ChildAsset->Children)
+				if (GroupParamData->StackDataType == EStackDataType::Group)
 				{
-					if (GroupParamData->StackDataType == EStackDataType::Group)
-					{
-						int32 GroupIdx = MaterialEditorInstance->ParameterGroups.IndexOfByPredicate(
-							[&](const FEditorParameterGroup& Group)
-								{
-									return Group.GroupName == GroupParamData->Group.GroupName;
-								});
-						if (GroupIdx != INDEX_NONE)
+					int32 GroupIdx = MaterialEditorInstance->ParameterGroups.IndexOfByPredicate(
+						[&](const FEditorParameterGroup& Group)
 						{
-							FEditorParameterGroup& ParameterGroup = GroupParamData->Group;
-							IDetailGroup& DetailGroup = GroupsCategory.AddGroup(ParameterGroup.GroupName, FText::FromName(ParameterGroup.GroupName), false, true);
-							TSharedPtr<IPropertyHandle> GroupPropertyHandle = ParameterGroupsProperty->GetChildHandle(GroupIdx);
-						
-							CreateSingleGroupWidget(ParameterGroup, GroupPropertyHandle, DetailGroup, GroupParamData->ParameterInfo.Index, true);
-						
-							FSimpleDelegate UpdateThumbnails = FSimpleDelegate::CreateLambda([=, this]()
+							return Group.GroupName == GroupParamData->Group.GroupName;
+						});
+					if (GroupIdx != INDEX_NONE)
+					{
+						FEditorParameterGroup& ParameterGroup = GroupParamData->Group;
+						IDetailGroup& DetailGroup = GroupsCategory.AddGroup(ParameterGroup.GroupName, FText::FromName(ParameterGroup.GroupName), false, true);
+						TSharedPtr<IPropertyHandle> GroupPropertyHandle = ParameterGroupsProperty->GetChildHandle(GroupIdx);
+
+						CreateSingleGroupWidget(ParameterGroup, GroupPropertyHandle, DetailGroup, GroupParamData->ParameterInfo.Index, true);
+
+						FSimpleDelegate UpdateThumbnails = FSimpleDelegate::CreateLambda([=, this]()
 							{
-								this->MaterialLayersFunctionsInstance->NestedTree->UpdateThumbnailMaterial(ChildAsset->ParameterInfo.Association, ChildAsset->ParameterInfo.Index);
+								this->MaterialLayersFunctionsInstance->NestedTree->UpdateThumbnailMaterial(AssetItem->ParameterInfo.Association, AssetItem->ParameterInfo.Index);
 							});
-							GroupPropertyHandle->SetOnPropertyValueChanged(UpdateThumbnails);
-							GroupPropertyHandle->SetOnChildPropertyValueChanged(UpdateThumbnails);
-						}
+						GroupPropertyHandle->SetOnPropertyValueChanged(UpdateThumbnails);
+						GroupPropertyHandle->SetOnChildPropertyValueChanged(UpdateThumbnails);
 					}
 				}
 			}
diff --git a/Engine/Source/Editor/MaterialEditor/Private/MaterialEditorInstanceDetailCustomization.h b/Engine/Source/Editor/MaterialEditor/Private/MaterialEditorInstanceDetailCustomization.h
--- a/Engine/Source/Editor/MaterialEditor/Private/MaterialEditorInstanceDetailCustomization.h
+++ b/Engine/Source/Editor/MaterialEditor/Private/MaterialEditorInstanceDetailCustomization.h
@@ -32,7 +32,6 @@ public:
 	/** Constructor */
 	FMaterialInstanceParameterDetails(UMaterialEditorInstanceConstant* MaterialInstance, SMaterialLayersFunctionsInstanceWrapper* MaterialLayersFunctionsInstance, FGetShowHiddenParameters InShowHiddenDelegate);
 
-	void CollectStackItemsRecursively(TSharedPtr<FSortedParamData> Item, TArray<TSharedPtr<FSortedParamData>>& OutGroupsContainer);
 	/** IDetailCustomization interface */
 	virtual void CustomizeDetails(IDetailLayoutBuilder& DetailLayout) override;
 
diff --git a/Engine/Source/Editor/MaterialEditor/Private/MaterialPropertyHelpers.cpp b/Engine/Source/Editor/MaterialEditor/Private/MaterialPropertyHelpers.cpp
--- a/Engine/Source/Editor/MaterialEditor/Private/MaterialPropertyHelpers.cpp
+++ b/Engine/Source/Editor/MaterialEditor/Private/MaterialPropertyHelpers.cpp
@@ -136,7 +136,7 @@ void FMaterialPropertyHelpers::OnMaterialLayerAssetChanged(const struct FAssetDa
 			{
 				InMaterialFunction->Layers[Index] = LayerFunction;
 				InMaterialFunction->UnlinkLayerFromParent(Index);
-#ifdef ENABLE_MATERIAL_LAYER_PROTOTYPE
+#if ENABLE_MATERIAL_LAYER_PROTOTYPE
 				// update name of the layer
 				InMaterialFunction->EditorOnly.LayerNames[Index] = FText::FromName(InAssetData.AssetName);
 #endif
@@ -146,7 +146,7 @@ void FMaterialPropertyHelpers::OnMaterialLayerAssetChanged(const struct FAssetDa
 			if (InMaterialFunction->Blends[Index] != LayerFunction)
 			{
 				InMaterialFunction->Blends[Index] = LayerFunction;
-#ifdef ENABLE_MATERIAL_LAYER_PROTOTYPE
+#if ENABLE_MATERIAL_LAYER_PROTOTYPE
 				InMaterialFunction->UnlinkLayerFromParent(Index);
 #else
 				InMaterialFunction->UnlinkLayerFromParent(Index + 1); // Blend indices are offset by 1, no blend for base layer
diff --git a/Engine/Source/Editor/MaterialEditor/Private/SMaterialLayersFunctionsTree.h b/Engine/Source/Editor/MaterialEditor/Private/SMaterialLayersFunctionsTree.h
--- a/Engine/Source/Editor/MaterialEditor/Private/SMaterialLayersFunctionsTree.h
+++ b/Engine/Source/Editor/MaterialEditor/Private/SMaterialLayersFunctionsTree.h
@@ -172,6 +172,9 @@ public:
 	/** Builds the custom parameter groups category */
 	void CreateGroupsWidget();
 
+	// Adding this noop method for compatibility with the Substrate enabled Material Layer. 
+	void CollectAssetStackItemsRecursively(TSharedPtr<FSortedParamData> Item, TArray<TSharedPtr<FSortedParamData>>& OutGroupsContainer) {}
+
 	bool IsLayerVisible(int32 Index) const;
 
 	SMaterialLayersFunctionsInstanceWrapper* GetWrapper() { return Wrapper; }
diff --git a/Engine/Source/Editor/MaterialEditor/Private/SMaterialSubstrateTree.cpp b/Engine/Source/Editor/MaterialEditor/Private/SMaterialSubstrateTree.cpp
--- a/Engine/Source/Editor/MaterialEditor/Private/SMaterialSubstrateTree.cpp
+++ b/Engine/Source/Editor/MaterialEditor/Private/SMaterialSubstrateTree.cpp
@@ -1224,6 +1224,7 @@ struct FRecursiveCreateWidgetsContext
 	UDEditorParameterValue* Parameter;
 	TSharedPtr<IPropertyHandle>	LayerHandle;
 	TSharedPtr<IPropertyHandle>	BlendHandle;
+	TArray<FUnsortedParamData> NonLayerProperties;
 };
 
 void SMaterialSubstrateTree::RecursiveCreateWidgets(FRecursiveCreateWidgetsContext* InContext, FNodeId InNodeId, TArray<TSharedPtr<FSortedParamData>>& InParentContainer, bool GenerateChildren)
@@ -1284,7 +1285,7 @@ void SMaterialSubstrateTree::RecursiveCreateWidgets(FRecursiveCreateWidgetsConte
 				if (LayerTreeCache)
 				{
 					UMaterialExpressionMaterialFunctionCall* Call = LayerTreeCache->NodeMaterialGraphExpressions[Payload.Layer];
-					if (Call->MaterialFunction)
+					if (Call && Call->MaterialFunction)
 					{
 						if (MaterialEditorInstance->StoredLayerPreviews[Payload.Layer] == nullptr)
 						{
@@ -1304,7 +1305,7 @@ void SMaterialSubstrateTree::RecursiveCreateWidgets(FRecursiveCreateWidgetsConte
 		}
 
 		StackProperty->Children.Add(ChildProperty);
-		ShowSubParameters(ChildProperty);
+		ShowSubParameters(InContext, ChildProperty);
 
 	}
 
@@ -1336,7 +1337,7 @@ void SMaterialSubstrateTree::RecursiveCreateWidgets(FRecursiveCreateWidgetsConte
 		}
 
 		StackProperty->Children.Add(ChildProperty);
-		ShowSubParameters(ChildProperty);
+		ShowSubParameters(InContext, ChildProperty);
 	}
 
 	InParentContainer.Add(StackProperty);
@@ -1349,7 +1350,8 @@ void SMaterialSubstrateTree::CreateGroupsWidget()
 	{
 		MaterialEditorInstance->RegenerateArrays();
 	}
-	NonLayerProperties.Empty();
+	
+	TArray<FUnsortedParamData> NonLayerProperties; // NonLayerProperties.Empty();
 	LayerProperties.Empty();
 	FunctionParameter = nullptr;
 	TSharedPtr<IPropertyHandle> FunctionParameterHandle;
@@ -1479,7 +1481,8 @@ void SMaterialSubstrateTree::CreateGroupsWidget()
 			FRecursiveCreateWidgetsContext Context{
 				.Parameter = StrongFunctionParameter.Get(),
 				.LayerHandle = LayerHandle,
-				.BlendHandle = BlendHandle
+				.BlendHandle = BlendHandle,
+				.NonLayerProperties = NonLayerProperties,
 			};
 
 			auto RootChildren = FunctionInstance->GetNodeChildren(-1);
@@ -1525,6 +1528,11 @@ TSharedRef<SWidget> SMaterialSubstrateTree::CreateThumbnailWidget(EMaterialParam
 
 void SMaterialSubstrateTree::UpdateThumbnailMaterial(TEnumAsByte<EMaterialParameterAssociation> InAssociation, int32 InIndex, bool bAlterBlendIndex)
 {
+	// From the notification, we get the NodeId triggering a parameter change, or just a refresh and we require a render for that preview material
+	TArray<int32> NodeParentsIds = FunctionInstance->GetNodeParents(InIndex);
+	TArray<int32> NodeChildrenIds = FunctionInstance->GetNodeChildren(InIndex);
+
+
 	UMaterialInstanceConstant* MaterialToUpdate = nullptr;
 	int32 ParameterIndex = InIndex;
 	if (InAssociation == EMaterialParameterAssociation::LayerParameter)
@@ -1533,45 +1541,51 @@ void SMaterialSubstrateTree::UpdateThumbnailMaterial(TEnumAsByte<EMaterialParame
 	}
 	if (InAssociation == EMaterialParameterAssociation::BlendParameter)
 	{
-		if (bAlterBlendIndex)
-		{
-			ParameterIndex--;
-		}
 		MaterialToUpdate = MaterialEditorInstance->StoredBlendPreviews[ParameterIndex];
 	}
-	return;
-	if (MaterialToUpdate != nullptr)
+
+	if (MaterialToUpdate != nullptr && FunctionInstance)
 	{
-		// Need to invert index b/c layer properties is generated in reverse order
-		TArray<TSharedPtr<FSortedParamData>> AssetChildren = LayerProperties[LayerProperties.Num() - 1 - InIndex]->Children;
+		// if parents grab the parent preview mat:
+		UMaterialInstanceConstant* ParentMaterialToUpdate = nullptr; 
+		if (NodeParentsIds.Num() > 1)
+		{
+			ParentMaterialToUpdate = MaterialEditorInstance->StoredLayerPreviews[NodeParentsIds[0]];
+		}
+
+		TArray<TSharedPtr<FSortedParamData>> AssetItemsContainer;
+		CollectStackItemsForMaterialFunctionAsset(InAssociation, InIndex, AssetItemsContainer);
 
 		TArray<FEditorParameterGroup> ParameterGroups;
-		for (TSharedPtr<FSortedParamData> AssetChild : AssetChildren)
+		for (TSharedPtr<FSortedParamData> AssetItem : AssetItemsContainer)
 		{
-			for (TSharedPtr<FSortedParamData> Group : AssetChild->Children)
+			if (AssetItem->StackDataType == EStackDataType::Asset)
 			{
-				if (Group->ParameterInfo.Association == InAssociation)
+				for (TSharedPtr<FSortedParamData> Group : AssetItem->Children)
 				{
-					FEditorParameterGroup DuplicatedGroup = FEditorParameterGroup();
-					DuplicatedGroup.GroupAssociation = Group->Group.GroupAssociation;
-					DuplicatedGroup.GroupName = Group->Group.GroupName;
-					DuplicatedGroup.GroupSortPriority = Group->Group.GroupSortPriority;
-					for (UDEditorParameterValue* Parameter : Group->Group.Parameters)
+					if (Group->ParameterInfo.Association == InAssociation)
 					{
-						if (Parameter->ParameterInfo.Index == ParameterIndex)
+						FEditorParameterGroup DuplicatedGroup = FEditorParameterGroup();
+						DuplicatedGroup.GroupAssociation = Group->Group.GroupAssociation;
+						DuplicatedGroup.GroupName = Group->Group.GroupName;
+						DuplicatedGroup.GroupSortPriority = Group->Group.GroupSortPriority;
+						for (UDEditorParameterValue* Parameter : Group->Group.Parameters)
 						{
-							DuplicatedGroup.Parameters.Add(Parameter);
+							if (Parameter->ParameterInfo.Index == ParameterIndex)
+							{
+								DuplicatedGroup.Parameters.Add(Parameter);
+							}
 						}
+						ParameterGroups.Add(DuplicatedGroup);
 					}
-					ParameterGroups.Add(DuplicatedGroup);
 				}
-			}
-
-	
-
+			}	
 		}
 
 		FMaterialPropertyHelpers::TransitionAndCopyParameters(MaterialToUpdate, ParameterGroups, true);
+
+		if (ParentMaterialToUpdate)
+			FMaterialPropertyHelpers::TransitionAndCopyParameters(ParentMaterialToUpdate, ParameterGroups, true);
 	}
 }
 
@@ -1594,9 +1608,9 @@ FReply SMaterialSubstrateTree::OnThumbnailDoubleClick(const FGeometry& Geometry,
 	return FReply::Unhandled();
 }
 
-void SMaterialSubstrateTree::ShowSubParameters(TSharedPtr<FSortedParamData> ParentParameter)
+void SMaterialSubstrateTree::ShowSubParameters(FRecursiveCreateWidgetsContext* InContext, TSharedPtr<FSortedParamData> ParentParameter)
 {
-	for (FUnsortedParamData Property : NonLayerProperties)
+	for (FUnsortedParamData Property : InContext->NonLayerProperties)
 	{
 		UDEditorParameterValue* Parameter = Property.Parameter;
 		if (Parameter->ParameterInfo.Index == ParentParameter->ParameterInfo.Index
@@ -1663,6 +1677,54 @@ void SMaterialSubstrateTree::ShowSubParameters(TSharedPtr<FSortedParamData> Pare
 	}
 }
 
+
+void SMaterialSubstrateTree::CollectStackItemsRecursively(TSharedPtr<FSortedParamData> Item, TArray<TSharedPtr<FSortedParamData>>& OutGroupsContainer)
+{
+	for (TSharedPtr<FSortedParamData> Child : Item->Children)
+	{
+		if (Child->StackDataType == EStackDataType::Stack)
+		{
+			OutGroupsContainer.Add(Child);
+		}
+
+		CollectStackItemsRecursively(Child, OutGroupsContainer);
+	}
+}
+
+void SMaterialSubstrateTree::CollectAssetStackItemsRecursively(TSharedPtr<FSortedParamData> Item, TArray<TSharedPtr<FSortedParamData>>& OutGroupsContainer)
+{
+	for (TSharedPtr<FSortedParamData> Child : Item->Children)
+	{
+		if (Child->StackDataType == EStackDataType::Asset)
+		{
+			OutGroupsContainer.Add(Child);
+		}
+
+		CollectAssetStackItemsRecursively(Child, OutGroupsContainer);
+	}
+}
+
+void SMaterialSubstrateTree::CollectStackItemsForMaterialFunctionAsset(EMaterialParameterAssociation InAssociation, int32 InIndex, TArray<TSharedPtr<FSortedParamData>>& OutGroupsContainer)
+{
+	TArray<TSharedPtr<FSortedParamData>> CollectedAssetItems;
+	for (TSharedPtr<FSortedParamData> Child : LayerProperties)
+	{
+		CollectAssetStackItemsRecursively(Child, CollectedAssetItems);
+	}
+
+	for (TSharedPtr<FSortedParamData> AssetItem : CollectedAssetItems)
+	{
+		if (AssetItem->ParameterInfo.Association == InAssociation)
+		{
+			if (AssetItem->ParameterInfo.Index == InIndex)
+			{
+				OutGroupsContainer.Add(AssetItem);
+			}
+		}
+	}
+}
+
+
 void SMaterialSubstrateTree::OnDeleteSelectedTreeViewItems()
 {
 	TArray<FSortedParamDataPtr> SelectedItemsArray = GetSelectedItems();
diff --git a/Engine/Source/Editor/MaterialEditor/Private/SMaterialSubstrateTree.h b/Engine/Source/Editor/MaterialEditor/Private/SMaterialSubstrateTree.h
--- a/Engine/Source/Editor/MaterialEditor/Private/SMaterialSubstrateTree.h
+++ b/Engine/Source/Editor/MaterialEditor/Private/SMaterialSubstrateTree.h
@@ -80,8 +80,6 @@ public:
 
 	/** Builds the custom parameter groups category */
 	void CreateGroupsWidget();
-	using FNodeId = int32;
-	void RecursiveCreateWidgets(struct FRecursiveCreateWidgetsContext* Context, FNodeId InNodeId, TArray<TSharedPtr<FSortedParamData>>& InParentContainer, bool GenerateChildren);
 
 
 	SMaterialLayersFunctionsInstanceWrapper* GetWrapper() { return Wrapper; }
@@ -95,9 +93,18 @@ public:
 	bool IsOverriddenExpression(TObjectPtr<UDEditorParameterValue> Parameter, int32 InIndex);
 	
 	FGetShowHiddenParameters GetShowHiddenDelegate() const;
-protected:
 
-	void ShowSubParameters(TSharedPtr<FSortedParamData> ParentParameter);
+
+	void CollectStackItemsRecursively(TSharedPtr<FSortedParamData> Item, TArray<TSharedPtr<FSortedParamData>>& OutGroupsContainer);
+
+	void CollectAssetStackItemsRecursively(TSharedPtr<FSortedParamData> Item, TArray<TSharedPtr<FSortedParamData>>& OutGroupsContainer);
+
+	void CollectStackItemsForMaterialFunctionAsset(EMaterialParameterAssociation InAssociation, int32 InIndex, TArray<TSharedPtr<FSortedParamData>>& OutGroupsContainer);
+
+protected:
+	using FNodeId = int32;
+	void RecursiveCreateWidgets(struct FRecursiveCreateWidgetsContext* Context, FNodeId InNodeId, TArray<TSharedPtr<FSortedParamData>>& InParentContainer, bool GenerateChildren);
+	void ShowSubParameters(FRecursiveCreateWidgetsContext* InContext, TSharedPtr<FSortedParamData> ParentParameter);
 	void OnDeleteSelectedTreeViewItems();
 	bool CanDeleteSelectedTreeViewItems() const;
 	void OnRenameSelectedTreeViewItems();
@@ -106,8 +113,6 @@ private:
 	TSharedPtr<FUICommandList> CommandList;
 	TArray<TSharedPtr<FSortedParamData>> LayerProperties;
 	
-	TArray<FUnsortedParamData> NonLayerProperties;
-	
 	FDetailColumnSizeData ColumnSizeData;
 	
 	SMaterialLayersFunctionsInstanceWrapper* Wrapper;
diff --git a/Engine/Source/Runtime/Engine/Private/Materials/MaterialExpressions.cpp b/Engine/Source/Runtime/Engine/Private/Materials/MaterialExpressions.cpp
--- a/Engine/Source/Runtime/Engine/Private/Materials/MaterialExpressions.cpp
+++ b/Engine/Source/Runtime/Engine/Private/Materials/MaterialExpressions.cpp
@@ -18563,6 +18563,19 @@ int32 FMaterialLayersFunctionsTree::GetDepth(FNodeId InNodeId) const
 	return Depth;
 }
 
+FLayerNodeIdArray FMaterialLayersFunctionsTree::GetParentIds(FNodeId InNodeId) const
+{
+	FLayerNodeIdArray ParentIds;
+
+	FNode* Node = GetNode(InNodeId);
+	while (Node)
+	{
+		ParentIds.Add(Node->Parent);
+		Node = GetNode(Node->Parent);
+	}
+
+	return ParentIds;
+}
 
 // Find the node which is the head of the list of siblings where InNodeID belongs
 FLayerNodeId FMaterialLayersFunctionsTree::GetSiblingHeadId(FNodeId InNodeId) const
@@ -19699,6 +19712,11 @@ FLayerNodeId FMaterialLayersFunctions::GetNodeParent(FLayerNodeId InNodeId) cons
 	return FMaterialLayersFunctionsTree::InvalidId;
 }
 
+FLayerNodeIdArray FMaterialLayersFunctions::GetNodeParents(FLayerNodeId InNode) const
+{
+	return Tree.GetParentIds(InNode);
+}
+
 FMaterialLayersFunctions::FLayerNodeIdArray FMaterialLayersFunctions::GetNodeChildren(FLayerNodeId InNodeId) const
 {
 	return Tree.GetChildrenIds(InNodeId);
diff --git a/Engine/Source/Runtime/Engine/Public/Materials/MaterialLayersFunctions.h b/Engine/Source/Runtime/Engine/Public/Materials/MaterialLayersFunctions.h
--- a/Engine/Source/Runtime/Engine/Public/Materials/MaterialLayersFunctions.h
+++ b/Engine/Source/Runtime/Engine/Public/Materials/MaterialLayersFunctions.h
@@ -263,6 +263,13 @@ public:
 	// the first valid node in the tree under the root is depth 1
 	int32 GetDepth(FNodeId InNodeId) const;
 
+	// Gather the path of parent node up to the root from the specified node.
+	// First element is the direct parent, second is parent of parent, etc up to the root of the tree, so the last element is always -1.
+	// THe size of the returned array is equivalent to the depth.
+	// Thus if the specified node's parent is root the return array is 1 element containing -1.
+	// If node is invalid, the returned array is empty
+	TArray<FNodeId> GetParentIds(FNodeId InNodeId) const;
+
 	// Add a new node with the specified Payload under the InPArent at the specified child index
 	// the new node id is returned
 	FNodeId AddNode(const FPayload& InPayload, FNodeId InParent, int32 InAtChildIndex = -1);
@@ -647,6 +654,10 @@ struct FMaterialLayersFunctions : public FMaterialLayersFunctionsRuntimeData
 	ENGINE_API FLayerNodePayload GetNodePayload(FLayerNodeId InNodeId) const;
 	ENGINE_API int32 GetNodeDepth(FLayerNodeId InNodeId) const;
 
+	// generate the path from the ROOT to the specified node
+	ENGINE_API FLayerNodeIdArray GetNodeParents(FLayerNodeId InNode) const;
+
+
 	ENGINE_API int32 GetLayerFuncIndex(FLayerNodeId InNodeId) const;
 	ENGINE_API int32 GetBlendFuncIndex(FLayerNodeId InNodeId) const;
 
