fix warning in hordes build

diff --git a/Engine/Source/Editor/MaterialEditor/Private/SMaterialSubstrateTree.cpp b/Engine/Source/Editor/MaterialEditor/Private/SMaterialSubstrateTree.cpp
--- a/Engine/Source/Editor/MaterialEditor/Private/SMaterialSubstrateTree.cpp
+++ b/Engine/Source/Editor/MaterialEditor/Private/SMaterialSubstrateTree.cpp
@@ -1528,9 +1528,6 @@ TSharedRef<SWidget> SMaterialSubstrateTree::CreateThumbnailWidget(EMaterialParam
 
 void SMaterialSubstrateTree::UpdateThumbnailMaterial(TEnumAsByte<EMaterialParameterAssociation> InAssociation, int32 InIndex, bool bAlterBlendIndex)
 {
-	// From the notification, we get the NodeId triggering a parameter change, or just a refresh and we require a render for that preview material
-	TArray<int32> NodeParentsIds = FunctionInstance->GetNodeParents(InIndex);
-	TArray<int32> NodeChildrenIds = FunctionInstance->GetNodeChildren(InIndex);
 
 
 	UMaterialInstanceConstant* MaterialToUpdate = nullptr;
@@ -1546,6 +1543,9 @@ void SMaterialSubstrateTree::UpdateThumbnailMaterial(TEnumAsByte<EMaterialParame
 
 	if (MaterialToUpdate != nullptr && FunctionInstance)
 	{
+		// From the notification, we get the NodeId triggering a parameter change, or just a refresh and we require a render for that preview material
+		TArray<int32> NodeParentsIds = FunctionInstance->GetNodeParents(InIndex);
+
 		// if parents grab the parent preview mat:
 		UMaterialInstanceConstant* ParentMaterialToUpdate = nullptr; 
 		if (NodeParentsIds.Num() > 1)
