Changes to add initial support for an overview view for material expressions inside the material editor. Initial support is for Texture Sample nodes, but more additions can be made further down the line

diff --git a/Engine/Source/Editor/MaterialEditor/Private/MaterialEditor.cpp b/Engine/Source/Editor/MaterialEditor/Private/MaterialEditor.cpp
--- a/Engine/Source/Editor/MaterialEditor/Private/MaterialEditor.cpp
+++ b/Engine/Source/Editor/MaterialEditor/Private/MaterialEditor.cpp
@@ -168,6 +168,7 @@
 
 #include "SMaterialParametersOverviewWidget.h"
 #include "SMaterialEditorCustomPrimitiveDataWidget.h"
+#include "SMaterialExpressionsOverviewWidget.h"
 #include "IPropertyRowGenerator.h"
 #include "LandscapeMaterialInstanceConstant.h"
 #include "Widgets/Layout/SScrollBox.h"
@@ -503,6 +504,11 @@ void FMaterialEditor::RegisterToolbarTab(const TSharedRef<class FTabManager>& In
 		.SetGroup(WorkspaceMenuCategoryRef)
 		.SetIcon(FSlateIcon(FAppStyle::GetAppStyleSetName(), "Kismet.Tabs.Palette"));	// SUBSTRATE_TODO a Substrate icon
 
+	InTabManager->RegisterTabSpawner(FMaterialEditorTabs::MaterialOverviewTabId, FOnSpawnTab::CreateSP(this, &FMaterialEditor::SpawnTab_MaterialExpressionsOverview))
+		.SetDisplayName(LOCTEXT("MaterialOverviewTab", "Overview"))
+		.SetGroup(WorkspaceMenuCategoryRef)
+		.SetIcon(FSlateIcon(FAppStyle::GetAppStyleSetName(), "LevelEditor.Tabs.Details"));
+
 	MaterialStatsManager->RegisterTabs();
 
 	OnRegisterTabSpawners().Broadcast(InTabManager);
@@ -527,6 +533,7 @@ void FMaterialEditor::UnregisterTabSpawners(const TSharedRef<class FTabManager>&
 	InTabManager->UnregisterTabSpawner(FMaterialEditorTabs::CustomPrimitiveTabId);
 	InTabManager->UnregisterTabSpawner(FMaterialEditorTabs::LayerPropertiesTabId);
 	InTabManager->UnregisterTabSpawner(FMaterialEditorTabs::SubstrateTabId);
+	InTabManager->UnregisterTabSpawner(FMaterialEditorTabs::MaterialOverviewTabId);
 
 	MaterialStatsManager->UnregisterTabs();
 
@@ -1791,6 +1798,11 @@ void FMaterialEditor::CreateInternalWidgets()
 	Generator->OnRowsRefreshed().AddSP(this, &FMaterialEditor::GeneratorRowsRefreshed);
 	MaterialCustomPrimitiveDataWidget = SNew(SMaterialCustomPrimitiveDataPanel, MaterialEditorInstance);
 
+	MaterialExpressionsOverviewWidget = SNew(SMaterialExpressionsOverviewPanel)
+		.InMaterialEditorInstance(MaterialEditorInstance)
+		.InGenerator(Generator)
+		.InMaterialDetailsView(MaterialDetailsView);
+
 #if ENABLE_MATERIAL_LAYER_PROTOTYPE
 	MaterialLayersFunctionsInstance = SNew(SMaterialLayersFunctionsInstanceWrapper)
 		.InMaterialEditorInstance(MaterialEditorInstance)
@@ -1873,9 +1885,10 @@ void FMaterialEditor::OnFinishedChangingProperties(const FPropertyChangedEvent&
 		{
 			RefreshExpressionPreviews(true);
 		}
+
 		RefreshPreviewViewport();
 		UpdatePreviewMaterial();
-		MaterialCustomPrimitiveDataWidget->UpdateEditorInstance(MaterialEditorInstance);
+		SetEditorInstanceForWidgets();
 	}
 }
 
@@ -5212,7 +5225,7 @@ void FMaterialEditor::OnPromoteToParameter(const FToolMenuContext& InMenuContext
 	}
 	if (MaterialEditorInstance != nullptr)
 	{
-		MaterialCustomPrimitiveDataWidget->UpdateEditorInstance(MaterialEditorInstance);
+		SetEditorInstanceForWidgets();
 	}
 }
 
@@ -5450,7 +5463,7 @@ void FMaterialEditor::OnCreateSubstrateNodeForPin(const FToolMenuContext& InMenu
 
 	if (MaterialEditorInstance != nullptr)
 	{
-		MaterialCustomPrimitiveDataWidget->UpdateEditorInstance(MaterialEditorInstance);
+		SetEditorInstanceForWidgets();
 	}
 }
 
@@ -5803,6 +5816,22 @@ TSharedRef<SDockTab> FMaterialEditor::SpawnTab_Substrate(const FSpawnTabArgs& Ar
 	return SubstrateTab;
 }
 
+TSharedRef<SDockTab> FMaterialEditor::SpawnTab_MaterialExpressionsOverview(const FSpawnTabArgs& Args)
+{
+	check(Args.GetTabId() == FMaterialEditorTabs::MaterialOverviewTabId);
+
+	TSharedRef<SDockTab> SpawnedTab = SNew(SDockTab)
+		.Label(LOCTEXT("MaterialOverviewTabTitle", "Nodes Overview"))
+		[
+			SNew(SBox)
+			[
+				MaterialExpressionsOverviewWidget.ToSharedRef()
+			]
+		];
+
+	return SpawnedTab;
+}
+
 void FMaterialEditor::SetPreviewExpression(UMaterialExpression* NewPreviewExpression)
 {
 	UMaterialExpressionFunctionOutput* FunctionOutput = Cast<UMaterialExpressionFunctionOutput>(NewPreviewExpression);
@@ -6373,6 +6402,8 @@ void FMaterialEditor::DeleteNodesInternal(const TArray<class UEdGraphNode*>& Nod
 			FBlueprintEditorUtils::RemoveNode(NULL, NodesToDelete[Index], true);
 		}
 	}
+
+	SetEditorInstanceForWidgets();
 }
 
 void FMaterialEditor::CopySelectedNodes()
@@ -7289,6 +7320,8 @@ void FMaterialEditor::PostUndo(bool bSuccess)
 			FocusedGraphEdPtr.Pin()->NotifyGraphChanged();
 		}
 
+		SetEditorInstanceForWidgets();
+
 		SetMaterialDirty();
 
 		UpdateGenerator();
@@ -7310,7 +7343,7 @@ void FMaterialEditor::NotifyPostChange( const FPropertyChangedEvent& PropertyCha
 
 	if ( PropertyThatChanged )
 	{
-		MaterialCustomPrimitiveDataWidget->UpdateEditorInstance(MaterialEditorInstance);
+		SetEditorInstanceForWidgets();
 
 		const FName NameOfPropertyThatChanged( *PropertyThatChanged->GetName() );
 		if ((NameOfPropertyThatChanged == GET_MEMBER_NAME_CHECKED(UMaterialInterface, PreviewMesh)) ||
@@ -7601,7 +7634,7 @@ void FMaterialEditor::OnColorPickerCommitted(FLinearColor LinearColor, TWeakObje
 		}
 
 		RefreshExpressionPreviews();
-		MaterialCustomPrimitiveDataWidget->UpdateEditorInstance(MaterialEditorInstance);
+		SetEditorInstanceForWidgets();
 	}
 }
 
@@ -8079,6 +8112,12 @@ void FMaterialEditor::AddSelectedExpressionToFavorites()
 	}
 }
 
+void FMaterialEditor::SetEditorInstanceForWidgets() const
+{
+	MaterialCustomPrimitiveDataWidget->SetEditorInstance(MaterialEditorInstance);
+	MaterialExpressionsOverviewWidget->SetEditorInstance(MaterialEditorInstance);
+}
+
 void FMaterialEditor::UpdateDetailView()
 {
 	GetDetailView()->InvalidateCachedState();
@@ -8282,7 +8321,7 @@ void FMaterialEditor::OnNodeTitleCommitted(const FText& NewText, ETextCommit::Ty
 		NodeBeingChanged->Modify();
 		NodeBeingChanged->OnRenameNode(NewName);
 		UpdateGenerator();
-		MaterialCustomPrimitiveDataWidget->UpdateEditorInstance(MaterialEditorInstance);
+		SetEditorInstanceForWidgets();
 	}
 }
 
diff --git a/Engine/Source/Editor/MaterialEditor/Private/MaterialEditor.h b/Engine/Source/Editor/MaterialEditor/Private/MaterialEditor.h
--- a/Engine/Source/Editor/MaterialEditor/Private/MaterialEditor.h
+++ b/Engine/Source/Editor/MaterialEditor/Private/MaterialEditor.h
@@ -951,6 +951,9 @@ private:
 	/** Adds the selected expression to the favorites list. */
 	void AddSelectedExpressionToFavorites();
 
+	/* Updates internal material editor widgets */
+	void SetEditorInstanceForWidgets() const;
+
 private:
 	TSharedRef<SDockTab> SpawnTab_Preview(const FSpawnTabArgs& Args);
 	TSharedRef<SDockTab> SpawnTab_MaterialProperties(const FSpawnTabArgs& Args);
@@ -961,6 +964,7 @@ private:
 	TSharedRef<SDockTab> SpawnTab_CustomPrimitiveData(const FSpawnTabArgs& Args);
 	TSharedRef<SDockTab> SpawnTab_LayerProperties(const FSpawnTabArgs& Args);
 	TSharedRef<SDockTab> SpawnTab_Substrate(const FSpawnTabArgs& Args);
+	TSharedRef<SDockTab> SpawnTab_MaterialExpressionsOverview(const FSpawnTabArgs& Args);
 
 	void OnFinishedChangingProperties(const FPropertyChangedEvent& PropertyChangedEvent);
 	void OnFinishedChangingParametersFromOverview(const FPropertyChangedEvent& PropertyChangedEvent);
@@ -1000,6 +1004,10 @@ private:
 	/** Parameter overview list View */
 	TSharedPtr<class SMaterialParametersOverviewPanel> MaterialParametersOverviewWidget;
 
+	/** Shared texture samplers list view **/
+	TSharedPtr<class SMaterialExpressionsOverviewPanel> MaterialExpressionsOverviewWidget;
+
+	/** Custom primitive data widget **/
 	TSharedPtr<class SMaterialCustomPrimitiveDataPanel> MaterialCustomPrimitiveDataWidget;
 
 	/** Layer Properties View */
diff --git a/Engine/Source/Editor/MaterialEditor/Private/MaterialEditorModes.cpp b/Engine/Source/Editor/MaterialEditor/Private/MaterialEditorModes.cpp
--- a/Engine/Source/Editor/MaterialEditor/Private/MaterialEditorModes.cpp
+++ b/Engine/Source/Editor/MaterialEditor/Private/MaterialEditorModes.cpp
@@ -19,7 +19,7 @@ const FName FMaterialEditorApplicationModes::StandardMaterialEditorMode(TEXT("Gr
 
 TSharedPtr<FTabManager::FLayout> FMaterialEditorApplicationModes::GetDefaultEditorLayout(TSharedPtr<FMaterialEditor> InMaterialEditor)
 {
-	const TSharedRef<FTabManager::FLayout> StandaloneDefaultLayout = FTabManager::NewLayout("Standalone_MaterialEditor_Layout_v14")
+	const TSharedRef<FTabManager::FLayout> StandaloneDefaultLayout = FTabManager::NewLayout("Standalone_MaterialEditor_Layout_v16")
 	->AddArea
 	(
 		FTabManager::NewPrimaryArea() ->SetOrientation(Orient_Vertical)
@@ -43,6 +43,7 @@ TSharedPtr<FTabManager::FLayout> FMaterialEditorApplicationModes::GetDefaultEdit
 					->AddTab( FMaterialEditorTabs::ParameterDefaultsTabId, ETabState::OpenedTab )
 					->AddTab( FMaterialEditorTabs::CustomPrimitiveTabId, ETabState::ClosedTab )
 					->AddTab( FMaterialEditorTabs::LayerPropertiesTabId, ETabState::ClosedTab )
+					->AddTab( FMaterialEditorTabs::MaterialOverviewTabId, ETabState::OpenedTab )
 					->SetForegroundTab( FMaterialEditorTabs::PropertiesTabId )
 				)
 			)
diff --git a/Engine/Source/Editor/MaterialEditor/Private/MaterialEditorTabs.cpp b/Engine/Source/Editor/MaterialEditor/Private/MaterialEditorTabs.cpp
--- a/Engine/Source/Editor/MaterialEditor/Private/MaterialEditorTabs.cpp
+++ b/Engine/Source/Editor/MaterialEditor/Private/MaterialEditorTabs.cpp
@@ -14,4 +14,5 @@ const FName FMaterialEditorTabs::PreviewSettingsTabId(TEXT("MaterialEditor_Previ
 const FName FMaterialEditorTabs::ParameterDefaultsTabId(TEXT("MaterialEditor_ParameterDefaults"));
 const FName FMaterialEditorTabs::CustomPrimitiveTabId(TEXT("MaterialEditor_CustomPrimitiveData"));
 const FName FMaterialEditorTabs::LayerPropertiesTabId(TEXT("MaterialInstanceEditor_MaterialLayerProperties"));
-const FName FMaterialEditorTabs::SubstrateTabId(TEXT("MaterialEditor_Substrate"));
\ No newline at end of file
+const FName FMaterialEditorTabs::SubstrateTabId(TEXT("MaterialEditor_Substrate"));
+const FName FMaterialEditorTabs::MaterialOverviewTabId(TEXT("MaterialEditor_MaterialExpressionOverview"));
\ No newline at end of file
diff --git a/Engine/Source/Editor/MaterialEditor/Private/MaterialEditorTabs.h b/Engine/Source/Editor/MaterialEditor/Private/MaterialEditorTabs.h
--- a/Engine/Source/Editor/MaterialEditor/Private/MaterialEditorTabs.h
+++ b/Engine/Source/Editor/MaterialEditor/Private/MaterialEditorTabs.h
@@ -18,6 +18,7 @@ struct FMaterialEditorTabs
 	static const FName CustomPrimitiveTabId;
 	static const FName LayerPropertiesTabId;
 	static const FName SubstrateTabId;
+	static const FName MaterialOverviewTabId;
 
 	// Disable default constructor
 	FMaterialEditorTabs() = delete;
diff --git a/Engine/Source/Editor/MaterialEditor/Private/SMaterialEditorCustomPrimitiveDataWidget.h b/Engine/Source/Editor/MaterialEditor/Private/SMaterialEditorCustomPrimitiveDataWidget.h
--- a/Engine/Source/Editor/MaterialEditor/Private/SMaterialEditorCustomPrimitiveDataWidget.h
+++ b/Engine/Source/Editor/MaterialEditor/Private/SMaterialEditorCustomPrimitiveDataWidget.h
@@ -34,7 +34,7 @@ public:
 
 	void Construct(const FArguments& InArgs, UMaterialEditorPreviewParameters* InMaterialEditorInstance);
 
-	void UpdateEditorInstance(UMaterialEditorPreviewParameters* InMaterialEditorInstance) { MaterialEditorInstance = InMaterialEditorInstance; Refresh(); }
+	void SetEditorInstance(UMaterialEditorPreviewParameters* InMaterialEditorInstance) { MaterialEditorInstance = InMaterialEditorInstance; Refresh(); }
 
 private:
 
diff --git a/Engine/Source/Editor/MaterialEditor/Private/SMaterialExpressionsOverviewWidget.cpp b/Engine/Source/Editor/MaterialEditor/Private/SMaterialExpressionsOverviewWidget.cpp
new file mode 100644
--- /dev/null
+++ b/Engine/Source/Editor/MaterialEditor/Private/SMaterialExpressionsOverviewWidget.cpp
@@ -0,0 +1,189 @@
+// Copyright Epic Games, Inc. All Rights Reserved.
+
+#include "SMaterialExpressionsOverviewWidget.h"
+#include "MaterialEditor/MaterialEditorPreviewParameters.h"
+#include "Materials/MaterialExpressionTextureSample.h"
+#include "Widgets/SToolTip.h"
+
+class SCustomSharedSamplerRow : public SMultiColumnTableRow<TSharedPtr<FSharedSamplerDataRowData>>
+{
+public:
+	SLATE_BEGIN_ARGS(SCustomSharedSamplerRow) {}
+		SLATE_ARGUMENT(TSharedPtr<FSharedSamplerDataRowData>, Entry)
+	SLATE_END_ARGS()
+
+	void Construct(const FArguments& Args, const TSharedRef<STableViewBase>& OwnerTableView)
+	{
+		RowData = Args._Entry;
+
+		FSuperRowType::Construct(
+			FSuperRowType::FArguments()
+			.Padding(1.0f),
+			OwnerTableView
+		);
+	}
+
+	/** Overridden from SMultiColumnTableRow.  Generates a widget for this column of the list view. */
+	virtual TSharedRef<SWidget> GenerateWidgetForColumn(const FName& ColumnName) override
+	{
+		TSharedRef<STextBlock> TextBlock = SNew(STextBlock);
+
+		if (RowData->bIsDuplicate)
+		{
+			TextBlock->SetToolTip(SNew(SToolTip).Text(FText::FromString("This slot is potentially incorrectly overlapping")).BorderImage(FCoreStyle::Get().GetBrush("ToolTip.BrightBackground")));
+			TextBlock->SetColorAndOpacity(FSlateColor(FLinearColor::Red));
+		}
+
+		if (ColumnName == "Reference Count")
+		{
+			TextBlock->SetText(FText::FromString(FString::FormatAsNumber(RowData->Count)));
+			return TextBlock;
+		}
+		else if (ColumnName == "Name")
+		{
+			TextBlock->SetText(FText::FromString(RowData->Name));
+			return TextBlock;
+		}
+
+		return SNullWidget::NullWidget;
+	}
+
+private:
+	TSharedPtr<FSharedSamplerDataRowData> RowData;
+};
+
+void SMaterialExpressionsOverviewPanel::Construct(const FArguments& InArgs)
+{
+	this->ChildSlot
+		[
+			SNew(SVerticalBox)
+			+ SVerticalBox::Slot()
+			[
+				SNew(SBorder)
+				.Padding(FMargin(4.0f))
+				[
+					SNew(SVerticalBox)
+					+ SVerticalBox::Slot()
+					.AutoHeight()
+					[
+						SNew(SHorizontalBox)
+						+ SHorizontalBox::Slot()
+						.Padding(FMargin(3.0f, 4.0f))
+						.HAlign(HAlign_Left)
+						.AutoWidth()
+						.VAlign(VAlign_Center)
+						[
+							SNew(STextBlock)
+							.Text(FText::FromString("Textures"))
+						]
+					]
+					+ SVerticalBox::Slot()
+					.Padding(FMargin(3.0f, 2.0f, 3.0f, 3.0f))
+					[
+						SNew(SBorder)
+						[
+							SNew(SScrollBox)
+							+ SScrollBox::Slot()
+							[
+								SAssignNew(ListViewWidget, SListView<TSharedPtr<FSharedSamplerDataRowData>>)
+								.ListItemsSource(&Items)
+								.OnGenerateRow(this, &SMaterialExpressionsOverviewPanel::OnGenerateRowForList)
+								.OnMouseButtonClick(this, &SMaterialExpressionsOverviewPanel::OnMouseButtonClick)
+								.SelectionMode(ESelectionMode::Single)
+								.HeaderRow
+								(
+									SNew(SHeaderRow)
+									+ SHeaderRow::Column("Name").DefaultLabel(FText::FromString("Name"))
+									+ SHeaderRow::Column("Reference Count").DefaultLabel(FText::FromString("Reference Count"))
+								)
+							]
+						]
+					]
+				]
+			]
+		];
+
+	MaterialEditorInstance = InArgs._InMaterialEditorInstance;
+	MaterialDetailsView = InArgs._InMaterialDetailsView;
+	Refresh();
+}
+
+void SMaterialExpressionsOverviewPanel::SetEditorInstance(UMaterialEditorPreviewParameters* InMaterialEditorInstance)
+{
+	MaterialEditorInstance = InMaterialEditorInstance;
+	Refresh();
+}
+
+void SMaterialExpressionsOverviewPanel::Refresh()
+{
+	if (!MaterialEditorInstance)
+	{
+		// Early out in case we do not have a material editor instance available
+		return;
+	}
+
+	TMap<FString, int32> SamplerPairs;
+	TArray<UMaterialExpressionTextureSample*> TextureSamples;
+	MaterialEditorInstance->PreviewMaterial->GetAllExpressionsInMaterialAndFunctionsOfType<UMaterialExpressionTextureSample>(TextureSamples);
+
+	for (const UMaterialExpressionTextureSample* Expr : TextureSamples)
+	{
+		const FString TextureName = Expr->GetReferencedTexture() != nullptr ? Expr->GetReferencedTexture()->GetName() : Expr->GetName();
+		int32 Sampler = SamplerPairs.FindOrAdd(TextureName);
+		SamplerPairs.Add(TextureName, ++Sampler);
+	}
+
+	Items.Empty();
+
+	// Go through and mark slots as duplicate while copying over to the shared ptr array
+	for (const auto& Params : SamplerPairs)
+	{
+		Items.Add(MakeShareable(new FSharedSamplerDataRowData(Params.Key, Params.Value)));
+	}
+
+	// Sort the items in descending order based on the amount of times a given sampler has been referenced
+	Items.StableSort([](const TSharedPtr<FSharedSamplerDataRowData>& A, const TSharedPtr<FSharedSamplerDataRowData>& B) {
+		int32 CountA = A->Count;
+		int32 CountB = B->Count;
+		if (CountA == CountB)
+		{
+			return A->Name < B->Name;
+		}
+		return CountA < CountB;
+	});
+
+	ListViewWidget->RequestListRefresh();
+}
+
+TSharedRef<ITableRow> SMaterialExpressionsOverviewPanel::OnGenerateRowForList(TSharedPtr<FSharedSamplerDataRowData> Item, const TSharedRef<STableViewBase>& OwnerTable)
+{
+	return SNew(SCustomSharedSamplerRow, OwnerTable).Entry(Item);
+}
+
+void SMaterialExpressionsOverviewPanel::OnMouseButtonClick(const TSharedPtr<FSharedSamplerDataRowData> RowData)
+{
+	// Update the details panel to display info about the selected texture node
+	if (!MaterialDetailsView)
+	{
+		return;
+	}
+
+	TArray<UMaterialExpressionTextureSample*> TextureSamples;
+	MaterialEditorInstance->PreviewMaterial->GetAllExpressionsInMaterialAndFunctionsOfType<UMaterialExpressionTextureSample>(TextureSamples);
+
+	UObject* SelectedObject = nullptr;
+	for (UMaterialExpressionTextureSample* Expr : TextureSamples)
+	{
+		const FString TextureName = Expr->GetReferencedTexture() != nullptr ? Expr->GetReferencedTexture()->GetName() : Expr->GetName();
+		if (RowData->Name == TextureName)
+		{
+			SelectedObject = CastChecked<UObject>(Expr);
+		}
+	}
+
+	if (SelectedObject)
+	{
+		MaterialDetailsView->SetObject(SelectedObject, true);
+	}
+}
+
diff --git a/Engine/Source/Editor/MaterialEditor/Private/SMaterialExpressionsOverviewWidget.h b/Engine/Source/Editor/MaterialEditor/Private/SMaterialExpressionsOverviewWidget.h
new file mode 100644
--- /dev/null
+++ b/Engine/Source/Editor/MaterialEditor/Private/SMaterialExpressionsOverviewWidget.h
@@ -0,0 +1,70 @@
+// Copyright Epic Games, Inc. All Rights Reserved.
+
+#pragma once
+
+#include "CoreMinimal.h"
+#include "MaterialEditor.h"
+#include "IDetailTreeNode.h"
+#include "IDetailPropertyRow.h"
+#include "IDetailsView.h"
+#include "Widgets/DeclarativeSyntaxSupport.h"
+#include "Widgets/Layout/SScrollBox.h"
+#include "Widgets/SCompoundWidget.h"
+#include "Widgets/Views/SListView.h"
+
+class UMaterialEditorPreviewParameters;
+
+/** Data for a row in the list of shared sampler entries*/
+struct FSharedSamplerDataRowData
+{
+	FSharedSamplerDataRowData(FString InName, int32 InCount, bool bInIsDuplicate = false) : Name(InName), Count(InCount), bIsDuplicate(bInIsDuplicate) {}
+
+	FString Name = "";
+	int32 Count = 0;
+	bool bIsDuplicate = false;
+};
+
+// ********* SMaterialExpressionsOverviewPanel *******
+class SMaterialExpressionsOverviewPanel : public SCompoundWidget
+{
+public:
+	SLATE_BEGIN_ARGS(SMaterialExpressionsOverviewPanel)
+		: _InMaterialEditorInstance(nullptr)
+		, _InGenerator()
+		, _InMaterialDetailsView()
+	{}
+
+	SLATE_ARGUMENT(UMaterialEditorPreviewParameters*, InMaterialEditorInstance)
+	SLATE_ARGUMENT(TSharedPtr<class IPropertyRowGenerator>, InGenerator)
+	SLATE_ARGUMENT(TSharedPtr<class IDetailsView>, InMaterialDetailsView)
+	SLATE_END_ARGS()
+
+	void Construct(const FArguments& InArgs);
+
+	/* Sets the Material Editor Preview Parameters to be displayed and refreshes the UI */
+	void SetEditorInstance(UMaterialEditorPreviewParameters* InMaterialEditorInstance);
+
+	/* The material details widget is initialized after this widget is created so we cannot pass the view into the constructor */
+	void SetMaterialDetailsView(TSharedPtr<class IDetailsView> InDetailsView) { MaterialDetailsView = InDetailsView; }
+
+private:
+	void Refresh();
+
+	/** Adds a new text box with the string to the list */
+	TSharedRef<ITableRow> OnGenerateRowForList(TSharedPtr<FSharedSamplerDataRowData> Item, const TSharedRef<STableViewBase>& OwnerTable);
+	TWeakPtr<class IPropertyRowGenerator> Generator;
+	
+	void OnMouseButtonClick(const TSharedPtr<FSharedSamplerDataRowData> RowData);
+
+	/** The list of strings */
+	TArray<TSharedPtr<FSharedSamplerDataRowData>> Items;
+
+	/** The actual UI list */
+	TSharedPtr<SListView<TSharedPtr<FSharedSamplerDataRowData>>> ListViewWidget;
+
+	/** The set of material parameters this is associated with */
+	UMaterialEditorPreviewParameters* MaterialEditorInstance;
+
+	/** Pointer to the Material Editor instance this widget was created for **/
+	TSharedPtr<class IDetailsView> MaterialDetailsView;
+};
