Fixed bug "Slabs with lots of attributes hides the outer gray UI when scrolled past the parent item".

diff --git a/Engine/Source/Editor/MaterialEditor/Private/SMaterialSubstrateTree.cpp b/Engine/Source/Editor/MaterialEditor/Private/SMaterialSubstrateTree.cpp
--- a/Engine/Source/Editor/MaterialEditor/Private/SMaterialSubstrateTree.cpp
+++ b/Engine/Source/Editor/MaterialEditor/Private/SMaterialSubstrateTree.cpp
@@ -264,7 +264,7 @@ TOptional<EItemDropZone> SMaterialSubstrateTreeItem::CanAcceptDrop(const FDragDr
 
 		case EItemDropZone::BelowItem:
 			if (!HasLayerFuncAsset  // Can only add under if drop a new layer function
-				|| (SiblingIdx == 0)) // Can't add BEFORE the first layer or sub layer
+				|| (SiblingIdx > 0)) // Disable adding below layers for all except the first child
 			{
 				return TOptional<EItemDropZone>();
 			}
@@ -492,6 +492,7 @@ FReply SMaterialSubstrateTreeItem::OnAddMaterialAttributeClicked(int32 index)
 
 void SMaterialSubstrateTreeItem::Construct(const FArguments& InArgs, const TSharedRef<STableViewBase>& InOwnerTableView)
 {
+	HalfRoundBrush = new FSlateRoundedBoxBrush(FStyleColors::Recessed);
 	FOnTableRowDragEnter LayerDragDelegate = FOnTableRowDragEnter::CreateSP(this, &SMaterialSubstrateTreeItem::OnLayerDragEnter);
 	FOnTableRowDragLeave LayerDragLeaveDelegate = FOnTableRowDragLeave::CreateSP(this, &SMaterialSubstrateTreeItem::OnLayerDragLeave);
 
@@ -767,38 +768,22 @@ int32 SMaterialSubstrateTreeItem::OnPaint(const FPaintArgs& Args, const FGeometr
 	
 	const FSlateBrush* BackgroundBrushResource = Tree->BackgroundBrush.Get();
 	const int32 IndentLevel = GetIndentLevel();
-	// calculate children + self size
-	int NumStackChildren = 0;
-	for(TSharedPtr<FSortedParamData> Child : this->StackParameterData->Children)
-	{
-		if (Child->StackDataType == Stack)
-			NumStackChildren++;
-	}
-	float HeightAdjFactor = 1;
-
-	LayerId += IndentLevel;
-	if (IsItemExpanded())
-	{
-		HeightAdjFactor = NumStackChildren + 1;
-	}
 	
-	FVector2d OuterBorderSize = AllottedGeometry.GetLocalSize() * FVector2d(1, HeightAdjFactor);
+	FVector2d OuterBorderSize = AllottedGeometry.GetLocalSize();
 	float OffsetX = 20.0f;
-	FVector2D Offset = FVector2D(OffsetX * (IndentLevel+1), 0.0f);
+	FVector2D Offset = FVector2D(OffsetX * (IndentLevel+1), bFirstChild ? 5.0f : 0.0f);
 
 	const FSlateBrush* RoundedBoxBrush = FSubstrateMaterialEditorStyle::GetBrush("LayerView.Row.OuterRoundBrush");
 	if (IndentLevel > 0)
 	{
 		FSlateDrawElement::MakeBox(
 				OutDrawElements,
-				LayerId,
+				LayerId+1,
 				AllottedGeometry.ToPaintGeometry(OuterBorderSize - Offset, FSlateLayoutTransform(Offset/2)),
 				RoundedBoxBrush,
 				ESlateDrawEffect::None,
 				FStyleColors::Recessed.GetSpecifiedColor() * InWidgetStyle.GetColorAndOpacityTint() 
 		);
-		// TODO: This causes culling issues when the parent item drawing the outer gray for the children goes out of the scroll rect scope
-		
 		// we also need to draw to fill in the rounded edges of the items in between (other than last or first) 
 		// unless they're the first or last items in the list which would have a rounded top and a rounded bottom respectively
 		if (!bFirstChild)
@@ -807,7 +792,7 @@ int32 SMaterialSubstrateTreeItem::OnPaint(const FPaintArgs& Args, const FGeometr
 			FVector2D TranslationOffset = Offset/2 + FVector2D(0, OuterBorderSize.Y - 15.0f);
 			FSlateDrawElement::MakeBox(
 				OutDrawElements,
-				LayerId,
+				LayerId+1,
 				AllottedGeometry.ToPaintGeometry(NonRoundBorderOffset, FSlateLayoutTransform(TranslationOffset)),
 				BackgroundBrushResource,
 				ESlateDrawEffect::None,
@@ -815,18 +800,36 @@ int32 SMaterialSubstrateTreeItem::OnPaint(const FPaintArgs& Args, const FGeometr
 			);
 		}
 	}
-	// only need to draw gray part for first level to cover itself and it's children
-	if (IndentLevel <= 0)
+
+	const FSlateBrush* BrushToUse = BackgroundBrushResource;
+	FVector2D ExpandedOffset = FVector2d(0.0f, 0.0f);
+	FLinearColor ColorToUse = FStyleColors::Dropdown.GetSpecifiedColor();
+	Offset = FVector2D(16.0f, 0.0f);
+	if (IndentLevel == 0)
+	{
+		Offset = FVector2D(16.0f, 20.0f);
+		ExpandedOffset = FVector2d(0.0f, 10.0f);
+		BrushToUse = BackgroundBrushResource;
+	}
+	else if (bFirstChild)
+	{
+		Offset = FVector2D(16.0f, 0.0f);
+		ExpandedOffset = FVector2d(0.0f, 0.0f);
+		HalfRoundBrush->OutlineSettings.CornerRadii = FVector4(0.0f, 0.0f, 15.0f, 15.0f);
+		HalfRoundBrush->OutlineSettings.RoundingType = ESlateBrushRoundingType::Type::FixedRadius;
+		HalfRoundBrush->DrawAs = ESlateBrushDrawType::RoundedBox;
+		BrushToUse = HalfRoundBrush;
+	}
+	// only need to draw gray part for children or if the parent (Evaluation layer) is expanded
+	if (IndentLevel > 0 || IsItemExpanded())
 	{
-    	FVector2D ExpandedOffset = IsItemExpanded() ? FVector2d(0.0f, 10.0f) : FVector2d();
-		Offset = FVector2D(15.0f, IsItemExpanded() ? 10.0f : 20.0f);
 		FSlateDrawElement::MakeBox(
 				OutDrawElements,
 				LayerId,
 				AllottedGeometry.ToPaintGeometry(OuterBorderSize - Offset, FSlateLayoutTransform(Offset/2 + ExpandedOffset)),
-				RoundedBoxBrush,
+				BrushToUse,
 				ESlateDrawEffect::None,
-				FStyleColors::Dropdown.GetSpecifiedColor() * InWidgetStyle.GetColorAndOpacityTint() 
+				ColorToUse * InWidgetStyle.GetColorAndOpacityTint() 
 		);
 	}
 		
@@ -835,7 +838,7 @@ int32 SMaterialSubstrateTreeItem::OnPaint(const FPaintArgs& Args, const FGeometr
 	const float ReductionFactorY = 10.0f;
 	Offset = FVector2d((OffsetX * IndentLevel) + ReductionFactorX, ReductionFactorY * (2 - IndentLevel));
 	FGeometry BorderGeom = AllottedGeometry.MakeChild(AllottedGeometry.GetLocalSize() - Offset, FSlateLayoutTransform(Offset/2));
-	return STableRow<TSharedPtr<FSortedParamData>>::OnPaint(Args, BorderGeom, MyCullingRect, OutDrawElements, LayerId+1, InWidgetStyle, bParentEnabled);
+	return STableRow<TSharedPtr<FSortedParamData>>::OnPaint(Args, BorderGeom, MyCullingRect, OutDrawElements, LayerId+2, InWidgetStyle, bParentEnabled);
 }
 
 void SMaterialSubstrateTreeItem::Rename()
diff --git a/Engine/Source/Editor/MaterialEditor/Private/SMaterialSubstrateTree.h b/Engine/Source/Editor/MaterialEditor/Private/SMaterialSubstrateTree.h
--- a/Engine/Source/Editor/MaterialEditor/Private/SMaterialSubstrateTree.h
+++ b/Engine/Source/Editor/MaterialEditor/Private/SMaterialSubstrateTree.h
@@ -220,6 +220,6 @@ public:
 	SMaterialSubstrateTree* Tree;
 
 	UMaterialEditorParameters* MaterialEditorInstance;
-
+	FSlateBrush* HalfRoundBrush = new FSlateBrush();
 	FString GetInstancePath(SMaterialSubstrateTree* InTree) const;
 };
