Interactive updates (sliders) to material parameters should also override the parameters of the preview material.

- When a the default value of a numeric property changed we would call FMaterialEditor::OnNumericParameterDefaultChanged.
- This in turn calls FMaterialEditor::SetNumericParameterDefaultOnDependentMaterials, which attemps to update the parameter on other materials.
- This code called MaterialsToOverride.Add(OriginalMaterial), which would only update the _original_ material (not the one the material editor is previewing).  This would result in the material live updating if you had in the level editor viewport.
- The fix is to simply also apply this same logic to the Material owned by FMaterialEditor since this is the preview material.
- Finally, when the user lifts the mouse button the value is applied to the actual and results in a change.

[RN] rn major Rendering - Materials and Shaders

diff --git a/Engine/Source/Editor/MaterialEditor/Private/MaterialEditor.cpp b/Engine/Source/Editor/MaterialEditor/Private/MaterialEditor.cpp
--- a/Engine/Source/Editor/MaterialEditor/Private/MaterialEditor.cpp
+++ b/Engine/Source/Editor/MaterialEditor/Private/MaterialEditor.cpp
@@ -5605,6 +5605,10 @@ void FMaterialEditor::SetNumericParameterDefaultOnDependentMaterials(EMaterialPa
 	}
 	else
 	{
+		// When a parameter changes...
+		// Update the preview material.
+		MaterialsToOverride.Add(Material);
+		// The original Material we are editing.
 		MaterialsToOverride.Add(OriginalMaterial);
 	}
 
