Fix for regression in MaterialEditor::DeleteNodes that caused the operation to show the "Apply Changes" dialog.

diff --git a/Engine/Source/Editor/MaterialEditor/Private/MaterialEditor.cpp b/Engine/Source/Editor/MaterialEditor/Private/MaterialEditor.cpp
--- a/Engine/Source/Editor/MaterialEditor/Private/MaterialEditor.cpp
+++ b/Engine/Source/Editor/MaterialEditor/Private/MaterialEditor.cpp
@@ -6273,7 +6273,6 @@ void FMaterialEditor::DeleteNodes(const TArray<UEdGraphNode*>& NodesToDelete, bo
 		UpdatePreviewMaterial();
 		Material->MarkPackageDirty();
 		SetMaterialDirty();
-		UpdateOriginalMaterial();
 
 		if ( bHaveExpressionsToDelete )
 		{
