Substrate - added select node to the "create node from pin" process.

#rnx
#tests PC editor

diff --git a/Engine/Source/Editor/MaterialEditor/Private/MaterialEditor.cpp b/Engine/Source/Editor/MaterialEditor/Private/MaterialEditor.cpp
--- a/Engine/Source/Editor/MaterialEditor/Private/MaterialEditor.cpp
+++ b/Engine/Source/Editor/MaterialEditor/Private/MaterialEditor.cpp
@@ -2392,6 +2392,7 @@ void FMaterialEditor::AddGraphEditorPinActionsToContextMenu(FToolMenuSection& In
 		AddSubstrateContextualMenu(FMaterialEditorCommands::Get().CreateHorizontalMixNode, ESubstrateNodeForPin::HorizontalMix);
 		AddSubstrateContextualMenu(FMaterialEditorCommands::Get().CreateVerticalLayerNode, ESubstrateNodeForPin::VerticalLayer);
 		AddSubstrateContextualMenu(FMaterialEditorCommands::Get().CreateWeightNode, ESubstrateNodeForPin::Weight);
+		AddSubstrateContextualMenu(FMaterialEditorCommands::Get().CreateSelectNode, ESubstrateNodeForPin::Select);
 	}
 }
 
@@ -5427,6 +5428,10 @@ void FMaterialEditor::OnCreateSubstrateNodeForPin(const FToolMenuContext& InMenu
 	{
 		Action.MaterialExpressionClass = UMaterialExpressionSubstrateWeight::StaticClass();
 	}
+	else if (NodeForPin == ESubstrateNodeForPin::Select)
+	{
+		Action.MaterialExpressionClass = UMaterialExpressionSubstrateSelect::StaticClass();
+	}
 
 	check(PinNode);
 	UEdGraph* GraphObj = PinNode->GetGraph();
diff --git a/Engine/Source/Editor/MaterialEditor/Private/MaterialEditor.h b/Engine/Source/Editor/MaterialEditor/Private/MaterialEditor.h
--- a/Engine/Source/Editor/MaterialEditor/Private/MaterialEditor.h
+++ b/Engine/Source/Editor/MaterialEditor/Private/MaterialEditor.h
@@ -847,7 +847,8 @@ private:
 		Slab,
 		HorizontalMix,
 		VerticalLayer,
-		Weight
+		Weight,
+		Select
 	};
 	/** Will create a Substrate node as input to the pin */
 	void OnCreateSubstrateNodeForPin(const FToolMenuContext& InMenuContext, ESubstrateNodeForPin NodeForPin) const;
diff --git a/Engine/Source/Editor/MaterialEditor/Private/MaterialEditorActions.cpp b/Engine/Source/Editor/MaterialEditor/Private/MaterialEditorActions.cpp
--- a/Engine/Source/Editor/MaterialEditor/Private/MaterialEditorActions.cpp
+++ b/Engine/Source/Editor/MaterialEditor/Private/MaterialEditorActions.cpp
@@ -80,6 +80,7 @@ void FMaterialEditorCommands::RegisterCommands()
 	UI_COMMAND( CreateHorizontalMixNode, "Create Horizontal Mix Node", "Create a Horizontal Mix Node linked to the pin", EUserInterfaceActionType::Button, FInputChord());
 	UI_COMMAND( CreateVerticalLayerNode, "Create Vertical Layer Node", "Create a Vertical Layer Node linked to the pin", EUserInterfaceActionType::Button, FInputChord());
 	UI_COMMAND( CreateWeightNode, "Create Weight Node", "Create a Weight Node linked to the pin", EUserInterfaceActionType::Button, FInputChord());
+	UI_COMMAND( CreateSelectNode, "Create Select Node", "Create a Select Node linked to the pin", EUserInterfaceActionType::Button, FInputChord());
 	UI_COMMAND( ResetToDefault, "Reset Pin to Default Value", "Reset pin value to default for that type", EUserInterfaceActionType::Button, FInputChord());
 
 	UI_COMMAND(QualityLevel_All, "All", "Sets node preview to show all quality levels.)", EUserInterfaceActionType::RadioButton, FInputChord());
diff --git a/Engine/Source/Editor/MaterialEditor/Public/MaterialEditorActions.h b/Engine/Source/Editor/MaterialEditor/Public/MaterialEditorActions.h
--- a/Engine/Source/Editor/MaterialEditor/Public/MaterialEditorActions.h
+++ b/Engine/Source/Editor/MaterialEditor/Public/MaterialEditorActions.h
@@ -205,6 +205,9 @@ public:
 	/** Create weight node */
 	TSharedPtr< FUICommandInfo > CreateWeightNode;
 
+	/** Create select node */
+	TSharedPtr< FUICommandInfo > CreateSelectNode;
+
 	TSharedPtr< FUICommandInfo > QualityLevel_All;
 	TSharedPtr< FUICommandInfo > QualityLevel_Epic;
 	TSharedPtr< FUICommandInfo > QualityLevel_High;
