[Viewport Toolbar] Fix offset of "Previewing" Material Function editor message in viewport

- Fix short offset for "Previewing" message, which was overlapping old toolbar
- Set offset to 0.0f if only new toolbar is visible

diff --git a/Engine/Source/Editor/MaterialEditor/Private/MaterialEditor.cpp b/Engine/Source/Editor/MaterialEditor/Private/MaterialEditor.cpp
--- a/Engine/Source/Editor/MaterialEditor/Private/MaterialEditor.cpp
+++ b/Engine/Source/Editor/MaterialEditor/Private/MaterialEditor.cpp
@@ -179,8 +179,7 @@
 #include "MaterialEditorContext.h"
 #include "UObject/MetaData.h"
 #include "ToolMenus.h"
-
-
+#include "ViewportToolbar/UnrealEdViewportToolbar.h"
 
 #define LOCTEXT_NAMESPACE "MaterialEditor"
 
@@ -2570,7 +2569,8 @@ void FMaterialEditor::DrawMessages( FViewport* InViewport, FCanvas* Canvas )
 {
 	if( PreviewExpression != NULL )
 	{
-		Canvas->PushAbsoluteTransform( FTranslationMatrix(FVector(0.0f, 30.0f,0.0f) ) );
+		float VerticalOffset = UE::UnrealEd::ShowOldViewportToolbars() ? 36.0f : 0.0f;
+		Canvas->PushAbsoluteTransform( FTranslationMatrix(FVector(0.0f, VerticalOffset,0.0f) ) );
 
 		// The message to display in the viewport.
 		FString Name = FString::Printf( TEXT("Previewing: %s"), *PreviewExpression->GetName() );
