Fix leaking SMaterialEditor3DPreviewViewport widget being kept alive by SMaterialEditorViewportToolBar
Because of this, Material assets could not be unloaded even after the material editor was closed, which puts them in a bad state if the plugin gets unmounted and remounted.
#rnx

diff --git a/Engine/Source/Editor/MaterialEditor/Private/SMaterialEditorViewportToolBar.cpp b/Engine/Source/Editor/MaterialEditor/Private/SMaterialEditorViewportToolBar.cpp
--- a/Engine/Source/Editor/MaterialEditor/Private/SMaterialEditorViewportToolBar.cpp
+++ b/Engine/Source/Editor/MaterialEditor/Private/SMaterialEditorViewportToolBar.cpp
@@ -60,10 +60,10 @@ void SMaterialEditorViewportToolBar::Construct(const FArguments& InArgs, TShared
 
 TSharedRef<SWidget> SMaterialEditorViewportToolBar::GenerateShowMenu() const
 {
-	if (MaterialEditorViewportPtr)
+	if (TSharedPtr<SMaterialEditor3DPreviewViewport> MaterialEditorViewport = MaterialEditorViewportPtr.Pin())
 	{
 		constexpr bool bShowViewportStats = false;
-		return UE::MaterialEditor::CreateShowMenuWidget(MaterialEditorViewportPtr.ToSharedRef(), bShowViewportStats);
+		return UE::MaterialEditor::CreateShowMenuWidget(MaterialEditorViewport.ToSharedRef(), bShowViewportStats);
 	}
 
 	return SNullWidget::NullWidget;
diff --git a/Engine/Source/Editor/MaterialEditor/Private/SMaterialEditorViewportToolBar.h b/Engine/Source/Editor/MaterialEditor/Private/SMaterialEditorViewportToolBar.h
--- a/Engine/Source/Editor/MaterialEditor/Private/SMaterialEditorViewportToolBar.h
+++ b/Engine/Source/Editor/MaterialEditor/Private/SMaterialEditorViewportToolBar.h
@@ -28,7 +28,7 @@ public:
 	virtual bool IsViewModeSupported(EViewModeIndex ViewModeIndex) const override;
 
 private:
-	TSharedPtr<SMaterialEditor3DPreviewViewport> MaterialEditorViewportPtr;
+	TWeakPtr<SMaterialEditor3DPreviewViewport> MaterialEditorViewportPtr;
 };
 
 ///////////////////////////////////////////////////////////
