Material Layer Project:
- By default a new Slab has no Material Attributes, display an inviting widget to "Add Attributes"
- Fix the update of the thumbnails for Slab level

diff --git a/Engine/Source/Editor/MaterialEditor/Private/MaterialEditorInstanceDetailCustomization.cpp b/Engine/Source/Editor/MaterialEditor/Private/MaterialEditorInstanceDetailCustomization.cpp
--- a/Engine/Source/Editor/MaterialEditor/Private/MaterialEditorInstanceDetailCustomization.cpp
+++ b/Engine/Source/Editor/MaterialEditor/Private/MaterialEditorInstanceDetailCustomization.cpp
@@ -118,11 +118,14 @@ void FMaterialInstanceParameterDetails::CustomizeDetails(IDetailLayoutBuilder& D
 
 		// we should now gather all sub-stack items to loop through all together
 		TArray<TSharedPtr<FSortedParamData>> AssetStacksCollection;
-		MaterialLayersFunctionsInstance->NestedTree->CollectAssetStackItemsRecursively(SelectedItem, AssetStacksCollection);
+		TArray<uint32> AssetParentNodeCollection;
+		MaterialLayersFunctionsInstance->NestedTree->CollectAssetStackItemsRecursively(SelectedItem, AssetStacksCollection, AssetParentNodeCollection);
 
 		// we go through list of assets
-		for (TSharedPtr<FSortedParamData> AssetItem : AssetStacksCollection)
+		for (int32 Index = 0; Index < AssetStacksCollection.Num(); ++Index)
 		{
+			TSharedPtr<FSortedParamData> AssetItem = AssetStacksCollection[Index];
+			uint32 NodeID = AssetParentNodeCollection[Index];
 
 			for (TSharedPtr<FSortedParamData> GroupParamData : AssetItem->Children)
 			{
@@ -143,7 +146,7 @@ void FMaterialInstanceParameterDetails::CustomizeDetails(IDetailLayoutBuilder& D
 
 						FSimpleDelegate UpdateThumbnails = FSimpleDelegate::CreateLambda([=, this]()
 							{
-								this->MaterialLayersFunctionsInstance->NestedTree->UpdateThumbnailMaterial(AssetItem->ParameterInfo.Association, AssetItem->ParameterInfo.Index);
+								this->MaterialLayersFunctionsInstance->NestedTree->UpdateThumbnailMaterial(AssetItem->ParameterInfo.Association, NodeID);
 							});
 						GroupPropertyHandle->SetOnPropertyValueChanged(UpdateThumbnails);
 						GroupPropertyHandle->SetOnChildPropertyValueChanged(UpdateThumbnails);
diff --git a/Engine/Source/Editor/MaterialEditor/Private/MaterialEditorStyle.cpp b/Engine/Source/Editor/MaterialEditor/Private/MaterialEditorStyle.cpp
--- a/Engine/Source/Editor/MaterialEditor/Private/MaterialEditorStyle.cpp
+++ b/Engine/Source/Editor/MaterialEditor/Private/MaterialEditorStyle.cpp
@@ -361,6 +361,10 @@ void FSubstrateMaterialEditorStyle::SetupTextStyles(const TSharedRef<FSlateStyle
 
 	Style->Set("Font.Stage.Details.Small.Bold", FTextBlockStyle(NormalTextStyle)
 		.SetFont(IDetailLayoutBuilder::GetDetailFontBold()));
+
+	Style->Set("LayerView.Row.SlabWithoutAttributes",
+		FTextBlockStyle(ComponentNameBigTextStyle)
+		.SetColorAndOpacity(FStyleColors::SelectInactive));
 }
 
 #undef IMAGE_BRUSH
diff --git a/Engine/Source/Editor/MaterialEditor/Private/SMaterialLayersFunctionsTree.cpp b/Engine/Source/Editor/MaterialEditor/Private/SMaterialLayersFunctionsTree.cpp
--- a/Engine/Source/Editor/MaterialEditor/Private/SMaterialLayersFunctionsTree.cpp
+++ b/Engine/Source/Editor/MaterialEditor/Private/SMaterialLayersFunctionsTree.cpp
@@ -1877,8 +1877,8 @@ void SMaterialLayersFunctionsInstanceWrapper::Refresh()
 #if ENABLE_MATERIAL_LAYER_PROTOTYPE
 			FMenuBuilder DetailViewOptions( true, nullptr);
 			DetailViewOptions.AddMenuEntry(
-			LOCTEXT("AddMaterialEvaluation", "Add Material Evaluation"),
-			LOCTEXT("AddMaterialEvaluation_ToolTip", "Add new Material Evaluation to the Substrate tree"),
+			LOCTEXT("AddMaterialSlab", "Add Material Slab"),
+			LOCTEXT("AddMaterialSlab_ToolTip", "Add new Material Slab to the Substrate tree"),
 			FSlateIcon(),
 			FUIAction(FExecuteAction::CreateSP(NestedTree.Get(), &SMaterialSubstrateTree::AddRootNodeLayer)));
 
diff --git a/Engine/Source/Editor/MaterialEditor/Private/SMaterialLayersFunctionsTree.h b/Engine/Source/Editor/MaterialEditor/Private/SMaterialLayersFunctionsTree.h
--- a/Engine/Source/Editor/MaterialEditor/Private/SMaterialLayersFunctionsTree.h
+++ b/Engine/Source/Editor/MaterialEditor/Private/SMaterialLayersFunctionsTree.h
@@ -173,7 +173,7 @@ public:
 	void CreateGroupsWidget();
 
 	// Adding this noop method for compatibility with the Substrate enabled Material Layer. 
-	void CollectAssetStackItemsRecursively(TSharedPtr<FSortedParamData> Item, TArray<TSharedPtr<FSortedParamData>>& OutGroupsContainer) {}
+	void CollectAssetStackItemsRecursively(TSharedPtr<FSortedParamData> Item, TArray<TSharedPtr<FSortedParamData>>& OutGroupsContainer, TArray<uint32>& OutNodeIdsContainer) {}
 
 	bool IsLayerVisible(int32 Index) const;
 
diff --git a/Engine/Source/Editor/MaterialEditor/Private/SMaterialSubstrateTree.cpp b/Engine/Source/Editor/MaterialEditor/Private/SMaterialSubstrateTree.cpp
--- a/Engine/Source/Editor/MaterialEditor/Private/SMaterialSubstrateTree.cpp
+++ b/Engine/Source/Editor/MaterialEditor/Private/SMaterialSubstrateTree.cpp
@@ -59,6 +59,7 @@
 #include "DetailTreeNode.h"
 #include "Framework/Commands/GenericCommands.h"
 #include "Materials/MaterialFunctionMaterialLayerBlend.h"
+#include "MaterialShared.h"
 
 #define LOCTEXT_NAMESPACE "MaterialSubstrateTree"
 
@@ -271,7 +272,7 @@ TOptional<EItemDropZone> SMaterialSubstrateTreeItem::CanAcceptDrop(const FDragDr
 			break;
 		case EItemDropZone::OntoItem:
 			if (!(HasLayerFuncAsset || HasBlendFuncAsset) // Can only drop valid assets
-				|| (!HasLayerFuncAsset && HasBlendFuncAsset && (SiblingIdx == 0))) // Or just a blendfunc on the first layer
+				|| (!HasLayerFuncAsset && HasBlendFuncAsset)) 
 			{
 				return TOptional<EItemDropZone>();
 			}
@@ -402,12 +403,24 @@ FReply SMaterialSubstrateTreeItem::OnLayerDrop(const FDragDropEvent& DragDropEve
 			// NOTE: The drop cases Above and Below take into account the fact that the list is displayed bottom up!
 			case EItemDropZone::AboveItem:
 				check(HasLayerFuncAsset); // Only add with a valid new LayerFunc asset
+				// if new layer node's parent is root then add a L1 group layer and THEN the L2 first node
+				if (ParentNodeId == FMaterialLayersFunctionsTree::InvalidId)
+				{
+					ParentNodeId = Tree->FunctionInstance->AppendLayerNode(ParentNodeId, SiblingIdx + 1);
+					SiblingIdx = -1; // the attributes will be aded first
+				}
 				TargetNodeId = Tree->FunctionInstance->AppendLayerNode(ParentNodeId, SiblingIdx + 1); // Above means insert after
 				DidModifyTree = true;
 				break;
 
 			case EItemDropZone::BelowItem:
 				check(HasLayerFuncAsset); // Only add with a valid new LayerFunc asset
+				// if new layer node's parent is root then add a L1 group layer and THEN the L2 first node
+				if (ParentNodeId == FMaterialLayersFunctionsTree::InvalidId)
+				{
+					ParentNodeId = Tree->FunctionInstance->AppendLayerNode(ParentNodeId, SiblingIdx);
+					SiblingIdx = -1; // the attributes will be aded first
+				}
 				TargetNodeId = Tree->FunctionInstance->AppendLayerNode(ParentNodeId, SiblingIdx); // Under means insert at
 				DidModifyTree = true;
 				break;
@@ -532,7 +545,8 @@ void SMaterialSubstrateTreeItem::Construct(const FArguments& InArgs, const TShar
 	EHorizontalAlignment ValueAlignment = HAlign_Left;
 
 	bool bCanAppendSubLayer = Tree->FunctionInstance->CanAppendLayerNode(StackParameterData->ParameterInfo.Index);
-		
+	bool bIsSlabWithNoAttributes = bCanAppendSubLayer && Tree->FunctionInstance->GetNodeChildren(StackParameterData->ParameterInfo.Index).IsEmpty();
+
 // STACK --------------------------------------------------
 	if (StackParameterData->StackDataType == EStackDataType::Stack)
 	{
@@ -705,24 +719,24 @@ void SMaterialSubstrateTreeItem::Construct(const FArguments& InArgs, const TShar
 		if (bCanAppendSubLayer)
 		{
 			MainStack->AddSlot()
-					.VAlign(VAlign_Center)
-					.Padding(2,0,0,0)
-					.AutoWidth()
-					[
-						SNew(SButton)
+				.VAlign(VAlign_Center)
+				.Padding(2, 0, 0, 0)
+				.AutoWidth()
+				[
+					SNew(SButton)
 						.ContentPadding(FMargin(1, 0))
 						.ButtonStyle(FAppStyle::Get(), "SimpleButton")
-							.ToolTipText(LOCTEXT("AddMaterialAttribute", "Add Material Attribute"))
-							.OnClicked(this, &SMaterialSubstrateTreeItem::OnAddMaterialAttributeClicked, StackParameterData->ParameterInfo.Index)
-							[
-								SNew(SImage)
+						.ToolTipText(LOCTEXT("AddMaterialAttribute", "Add Material Attribute"))
+						.OnClicked(this, &SMaterialSubstrateTreeItem::OnAddMaterialAttributeClicked, StackParameterData->ParameterInfo.Index)
+						[
+							SNew(SImage)
 								.Image(FAppStyle::Get().GetBrush("Icons.Plus"))
 								.ColorAndOpacity(FStyleColors::AccentGreen)
-							]
-					
-					];
+						]
+
+				];
 		}
-		
+
 		MainStack->AddSlot()
 			.HAlign(HAlign_Center)
 			.VAlign(VAlign_Center)
@@ -731,9 +745,9 @@ void SMaterialSubstrateTreeItem::Construct(const FArguments& InArgs, const TShar
 			[
 				FMaterialPropertyHelpers::MakeStackReorderHandle(SharedThis(this))
 			];
-		
 
-		if (bCanAppendSubLayer)
+
+		if (bCanAppendSubLayer && !bIsSlabWithNoAttributes)
 		{
 			MainStack->AddSlot()
 				.AutoWidth()
@@ -743,6 +757,23 @@ void SMaterialSubstrateTreeItem::Construct(const FArguments& InArgs, const TShar
 					SNew(SExpanderArrow, SharedThis(this))
 				];
 		}
+
+		if (bIsSlabWithNoAttributes)
+		{
+			WrapperWidget->AddSlot()
+				.AutoHeight()
+				.VAlign(VAlign_Center)
+				.Padding(8)
+				[
+					SNew(SBorder)
+						.Padding(8)
+						[
+							SNew(STextBlock)
+								.Text(LOCTEXT("AddAttributes", "Add Attributes"))
+								.TextStyle(FSubstrateMaterialEditorStyle::Get(), "LayerView.Row.SlabWithoutAttributes")
+						]
+				];
+		}
 	}
 
 	this->ChildSlot
@@ -761,11 +792,11 @@ int32 SMaterialSubstrateTreeItem::OnPaint(const FPaintArgs& Args, const FGeometr
 {
 	int32 TargetNodeId = StackParameterData->ParameterInfo.Index;
 	int32 ParentNodeId = Tree->FunctionInstance->GetNodeParent(TargetNodeId);
-	auto ChildrenNodeId = Tree->FunctionInstance->GetNodeChildren(ParentNodeId);
+	auto ChildrenNodeIds = Tree->FunctionInstance->GetNodeChildren(ParentNodeId);
 	int32 SiblingIdx = -1;
-	ChildrenNodeId.Find(TargetNodeId, SiblingIdx);
+	ChildrenNodeIds.Find(TargetNodeId, SiblingIdx);
 	bool bFirstChild = SiblingIdx == 0; 
-	
+
 	const FSlateBrush* BackgroundBrushResource = Tree->BackgroundBrush.Get();
 	const int32 IndentLevel = GetIndentLevel();
 	
@@ -833,10 +864,18 @@ int32 SMaterialSubstrateTreeItem::OnPaint(const FPaintArgs& Args, const FGeometr
 		);
 	}
 		
+	const float AddAttributeBoxHeight = 36;
+	const float AddAttributeBoxPadding = 8;
+	float NoChildrenBoxHeight = 0.0f;
+	if (IndentLevel == 0 && ChildrenNodeIds.IsEmpty())
+	{
+		NoChildrenBoxHeight = AddAttributeBoxHeight + AddAttributeBoxPadding * 2;
+	}
 	OffsetX = 40.0f;
 	const float ReductionFactorX = 15.0f;
 	const float ReductionFactorY = 10.0f;
-	Offset = FVector2d((OffsetX * IndentLevel) + ReductionFactorX, ReductionFactorY * (2 - IndentLevel));
+	Offset = FVector2d((OffsetX * IndentLevel) + ReductionFactorX, ReductionFactorY * (2 - IndentLevel) - NoChildrenBoxHeight);
+ 
 	FGeometry BorderGeom = AllottedGeometry.MakeChild(AllottedGeometry.GetLocalSize() - Offset, FSlateLayoutTransform(Offset/2));
 	return STableRow<TSharedPtr<FSortedParamData>>::OnPaint(Args, BorderGeom, MyCullingRect, OutDrawElements, LayerId+2, InWidgetStyle, bParentEnabled);
 }
@@ -1287,7 +1326,7 @@ void SMaterialSubstrateTree::RecursiveCreateWidgets(FRecursiveCreateWidgetsConte
 				TSharedPtr<FMaterialLayersFunctionsRuntimeGraphCache> LayerTreeCache = LayersFunctions->RuntimeGraphCache;
 				if (LayerTreeCache)
 				{
-					UMaterialExpressionMaterialFunctionCall* Call = LayerTreeCache->NodeMaterialGraphExpressions[Payload.Layer];
+					UMaterialExpressionMaterialFunctionCall* Call = LayerTreeCache->NodeMaterialGraphExpressions[InNodeId];
 					if (Call && Call->MaterialFunction)
 					{
 						if (MaterialEditorInstance->StoredLayerPreviews[Payload.Layer] == nullptr)
@@ -1295,7 +1334,7 @@ void SMaterialSubstrateTree::RecursiveCreateWidgets(FRecursiveCreateWidgetsConte
 							MaterialEditorInstance->StoredLayerPreviews[Payload.Layer] = (NewObject<UMaterialInstanceConstant>(MaterialEditorInstance, NAME_None));
 						}
 
-						UMaterial* PreviewMaterial = LayersFunctions->GetRuntimeNodePreviewMaterial(Payload.Layer);
+						UMaterial* PreviewMaterial = LayersFunctions->GetRuntimeNodePreviewMaterial(InNodeId);
 
 						if (MaterialEditorInstance->StoredLayerPreviews[Payload.Layer] && MaterialEditorInstance->StoredLayerPreviews[Payload.Layer]->Parent != PreviewMaterial)
 						{
@@ -1531,17 +1570,16 @@ TSharedRef<SWidget> SMaterialSubstrateTree::CreateThumbnailWidget(EMaterialParam
 
 void SMaterialSubstrateTree::UpdateThumbnailMaterial(TEnumAsByte<EMaterialParameterAssociation> InAssociation, int32 InIndex, bool bAlterBlendIndex)
 {
-
+	int32 AssetIndex = FunctionInstance->GetLayerFuncIndex(InIndex);
 
 	UMaterialInstanceConstant* MaterialToUpdate = nullptr;
-	int32 ParameterIndex = InIndex;
 	if (InAssociation == EMaterialParameterAssociation::LayerParameter)
 	{
-		MaterialToUpdate = MaterialEditorInstance->StoredLayerPreviews[ParameterIndex];
+		MaterialToUpdate = MaterialEditorInstance->StoredLayerPreviews[AssetIndex];
 	}
 	if (InAssociation == EMaterialParameterAssociation::BlendParameter)
 	{
-		MaterialToUpdate = MaterialEditorInstance->StoredBlendPreviews[ParameterIndex];
+		MaterialToUpdate = MaterialEditorInstance->StoredBlendPreviews[AssetIndex];
 	}
 
 	if (MaterialToUpdate != nullptr && FunctionInstance)
@@ -1553,11 +1591,12 @@ void SMaterialSubstrateTree::UpdateThumbnailMaterial(TEnumAsByte<EMaterialParame
 		UMaterialInstanceConstant* ParentMaterialToUpdate = nullptr; 
 		if (NodeParentsIds.Num() > 1)
 		{
-			ParentMaterialToUpdate = MaterialEditorInstance->StoredLayerPreviews[NodeParentsIds[0]];
+			int32 ParentAssetIndex = FunctionInstance->GetLayerFuncIndex(NodeParentsIds[0]);
+			ParentMaterialToUpdate = MaterialEditorInstance->StoredLayerPreviews[ParentAssetIndex];
 		}
 
 		TArray<TSharedPtr<FSortedParamData>> AssetItemsContainer;
-		CollectStackItemsForMaterialFunctionAsset(InAssociation, InIndex, AssetItemsContainer);
+		CollectStackItemsForMaterialFunctionAsset(InAssociation, AssetIndex, AssetItemsContainer);
 
 		TArray<FEditorParameterGroup> ParameterGroups;
 		for (TSharedPtr<FSortedParamData> AssetItem : AssetItemsContainer)
@@ -1574,7 +1613,7 @@ void SMaterialSubstrateTree::UpdateThumbnailMaterial(TEnumAsByte<EMaterialParame
 						DuplicatedGroup.GroupSortPriority = Group->Group.GroupSortPriority;
 						for (UDEditorParameterValue* Parameter : Group->Group.Parameters)
 						{
-							if (Parameter->ParameterInfo.Index == ParameterIndex)
+							if (Parameter->ParameterInfo.Index == AssetIndex)
 							{
 								DuplicatedGroup.Parameters.Add(Parameter);
 							}
@@ -1588,7 +1627,7 @@ void SMaterialSubstrateTree::UpdateThumbnailMaterial(TEnumAsByte<EMaterialParame
 		FMaterialPropertyHelpers::TransitionAndCopyParameters(MaterialToUpdate, ParameterGroups, true);
 
 		if (ParentMaterialToUpdate)
-			FMaterialPropertyHelpers::TransitionAndCopyParameters(ParentMaterialToUpdate, ParameterGroups, true);
+			FMaterialPropertyHelpers::CopyMaterialToInstance(ParentMaterialToUpdate, ParameterGroups);
 	}
 }
 
@@ -1694,32 +1733,35 @@ void SMaterialSubstrateTree::CollectStackItemsRecursively(TSharedPtr<FSortedPara
 	}
 }
 
-void SMaterialSubstrateTree::CollectAssetStackItemsRecursively(TSharedPtr<FSortedParamData> Item, TArray<TSharedPtr<FSortedParamData>>& OutGroupsContainer)
+void SMaterialSubstrateTree::CollectAssetStackItemsRecursively(TSharedPtr<FSortedParamData> Item, TArray<TSharedPtr<FSortedParamData>>& OutGroupsContainer, TArray<uint32>& OutNodeIdsContainer)
 {
+	uint32 NodeID = Item->ParameterInfo.Index;
 	for (TSharedPtr<FSortedParamData> Child : Item->Children)
 	{
 		if (Child->StackDataType == EStackDataType::Asset)
 		{
 			OutGroupsContainer.Add(Child);
+			OutNodeIdsContainer.Add(NodeID);
 		}
 
-		CollectAssetStackItemsRecursively(Child, OutGroupsContainer);
+		CollectAssetStackItemsRecursively(Child, OutGroupsContainer, OutNodeIdsContainer);
 	}
 }
 
-void SMaterialSubstrateTree::CollectStackItemsForMaterialFunctionAsset(EMaterialParameterAssociation InAssociation, int32 InIndex, TArray<TSharedPtr<FSortedParamData>>& OutGroupsContainer)
+void SMaterialSubstrateTree::CollectStackItemsForMaterialFunctionAsset(EMaterialParameterAssociation InAssociation, int32 InAssetIndex, TArray<TSharedPtr<FSortedParamData>>& OutGroupsContainer)
 {
 	TArray<TSharedPtr<FSortedParamData>> CollectedAssetItems;
+	TArray<uint32> CollectedNodeIds;
 	for (TSharedPtr<FSortedParamData> Child : LayerProperties)
 	{
-		CollectAssetStackItemsRecursively(Child, CollectedAssetItems);
+		CollectAssetStackItemsRecursively(Child, CollectedAssetItems, CollectedNodeIds);
 	}
 
 	for (TSharedPtr<FSortedParamData> AssetItem : CollectedAssetItems)
 	{
 		if (AssetItem->ParameterInfo.Association == InAssociation)
 		{
-			if (AssetItem->ParameterInfo.Index == InIndex)
+			if (AssetItem->ParameterInfo.Index == InAssetIndex)
 			{
 				OutGroupsContainer.Add(AssetItem);
 			}
diff --git a/Engine/Source/Editor/MaterialEditor/Private/SMaterialSubstrateTree.h b/Engine/Source/Editor/MaterialEditor/Private/SMaterialSubstrateTree.h
--- a/Engine/Source/Editor/MaterialEditor/Private/SMaterialSubstrateTree.h
+++ b/Engine/Source/Editor/MaterialEditor/Private/SMaterialSubstrateTree.h
@@ -97,9 +97,9 @@ public:
 
 	void CollectStackItemsRecursively(TSharedPtr<FSortedParamData> Item, TArray<TSharedPtr<FSortedParamData>>& OutGroupsContainer);
 
-	void CollectAssetStackItemsRecursively(TSharedPtr<FSortedParamData> Item, TArray<TSharedPtr<FSortedParamData>>& OutGroupsContainer);
+	void CollectAssetStackItemsRecursively(TSharedPtr<FSortedParamData> Item, TArray<TSharedPtr<FSortedParamData>>& OutGroupsContainer, TArray<uint32>& OutNodeIdsContainer);
 
-	void CollectStackItemsForMaterialFunctionAsset(EMaterialParameterAssociation InAssociation, int32 InIndex, TArray<TSharedPtr<FSortedParamData>>& OutGroupsContainer);
+	void CollectStackItemsForMaterialFunctionAsset(EMaterialParameterAssociation InAssociation, int32 InAssetIndex, TArray<TSharedPtr<FSortedParamData>>& OutGroupsContainer);
 
 protected:
 	using FNodeId = int32;
diff --git a/Engine/Source/Runtime/Engine/Private/Materials/MaterialExpressions.cpp b/Engine/Source/Runtime/Engine/Private/Materials/MaterialExpressions.cpp
--- a/Engine/Source/Runtime/Engine/Private/Materials/MaterialExpressions.cpp
+++ b/Engine/Source/Runtime/Engine/Private/Materials/MaterialExpressions.cpp
@@ -19744,17 +19744,10 @@ FLayerNodeId FMaterialLayersFunctions::AppendLayerNode(FLayerNodeId InParent, in
 	if (!CanAppendLayerNode(InParent))
 		return FMaterialLayersFunctionsTree::InvalidId;
 
-	// if new layer node's parent is root then add a L1 group layer and THEN the L2 first node
-	if (InParent == FMaterialLayersFunctionsTree::InvalidId)
-	{
-		const int32 LayerIndex = AppendBlendedLayer();
-		InParent = Tree.AddNode({ LayerIndex, LayerIndex }, InParent, InSiblingIndex);
-	}
-
-	const int32 LayerIndex = AppendBlendedLayer();
+	const int32 PayloadIndex = AppendBlendedLayer();
 
 	// InPayload
-	FLayerNodeId NewNodeId = Tree.AddNode({ LayerIndex, LayerIndex }, InParent, InSiblingIndex);
+	FLayerNodeId NewNodeId = Tree.AddNode({ PayloadIndex, PayloadIndex }, InParent, InSiblingIndex);
 
 	MLFT::Log(this, __FUNCTION__);
 
@@ -19769,12 +19762,8 @@ bool FMaterialLayersFunctions::CanRemoveLayerNode(FLayerNodeId InNodeId) const
 	// Only if really valid
 	if (Tree.IsValidId(InNodeId))
 	{
-		// NOTE: Current design guarantee first sub layer in any layer, and first top layer
-		FLayerNodeId SiblingHead = Tree.GetSiblingHeadId(InNodeId);
-		if (InNodeId != SiblingHead)
-		{
-			return true;
-		}
+		// Must keep the very first Slab / node
+		return Tree.Root != InNodeId;
 	}
 	return false;
 }
