Address warning because a memeber is tested AFTER being used

diff --git a/Engine/Source/Editor/MaterialEditor/Private/SMaterialSubstrateTree.cpp b/Engine/Source/Editor/MaterialEditor/Private/SMaterialSubstrateTree.cpp
--- a/Engine/Source/Editor/MaterialEditor/Private/SMaterialSubstrateTree.cpp
+++ b/Engine/Source/Editor/MaterialEditor/Private/SMaterialSubstrateTree.cpp
@@ -1582,7 +1582,7 @@ void SMaterialSubstrateTree::UpdateThumbnailMaterial(TEnumAsByte<EMaterialParame
 		MaterialToUpdate = MaterialEditorInstance->StoredBlendPreviews[AssetIndex];
 	}
 
-	if (MaterialToUpdate != nullptr && FunctionInstance)
+	if (MaterialToUpdate != nullptr)
 	{
 		// From the notification, we get the NodeId triggering a parameter change, or just a refresh and we require a render for that preview material
 		TArray<int32> NodeParentsIds = FunctionInstance->GetNodeParents(InIndex);
