Material Layer -  Add support for the instance version of the ML and MLB

diff --git a/Engine/Source/Editor/MaterialEditor/Private/SMaterialSubstrateTree.cpp b/Engine/Source/Editor/MaterialEditor/Private/SMaterialSubstrateTree.cpp
--- a/Engine/Source/Editor/MaterialEditor/Private/SMaterialSubstrateTree.cpp
+++ b/Engine/Source/Editor/MaterialEditor/Private/SMaterialSubstrateTree.cpp
@@ -58,11 +58,24 @@
 #include "MaterialEditorStyle.h"
 #include "DetailTreeNode.h"
 #include "Framework/Commands/GenericCommands.h"
+#include "Materials/MaterialFunctionMaterialLayer.h"
 #include "Materials/MaterialFunctionMaterialLayerBlend.h"
 #include "MaterialShared.h"
 
 #define LOCTEXT_NAMESPACE "MaterialSubstrateTree"
 
+// Check if the assetdata is avalid MaterialLayer Function or FunctionInstance
+bool IsAssetDataAMaterialLayerFunction(const FAssetData& AssetData)
+{
+	return AssetData.IsInstanceOf(UMaterialFunctionMaterialLayer::StaticClass()) || AssetData.IsInstanceOf(UMaterialFunctionMaterialLayerInstance::StaticClass());
+}
+
+// Check if the assetdata is avalid MaterialLayerBlend Function or FunctionInstance
+bool IsAssetDataAMaterialLayerBlendFunction(const FAssetData& AssetData)
+{
+	return AssetData.IsInstanceOf(UMaterialFunctionMaterialLayerBlend::StaticClass()) || AssetData.IsInstanceOf(UMaterialFunctionMaterialLayerBlendInstance::StaticClass());
+}
+
 /// ===========================================================================================================
 /// SMaterialSubstrateItem Methods
 /// ===========================================================================================================
@@ -244,14 +257,8 @@ TOptional<EItemDropZone> SMaterialSubstrateTreeItem::CanAcceptDrop(const FDragDr
 		bool HasBlendFuncAsset = false;
 		for (const FAssetData& AssetData : AssetDragDropOperation->GetAssets())
 		{
-			if (AssetData.AssetClassPath.GetAssetName() == TEXT("MaterialFunctionMaterialLayer"))
-			{
-				HasLayerFuncAsset = true;
-			}
-			else if (AssetData.AssetClassPath.GetAssetName() == TEXT("MaterialFunctionMaterialLayerBlend"))
-			{
-				HasBlendFuncAsset = true;
-			}
+			HasLayerFuncAsset |= IsAssetDataAMaterialLayerFunction(AssetData);
+			HasBlendFuncAsset |= IsAssetDataAMaterialLayerBlendFunction(AssetData);
 		}
 
 		switch (DropZone)
@@ -378,14 +385,8 @@ FReply SMaterialSubstrateTreeItem::OnLayerDrop(const FDragDropEvent& DragDropEve
 			bool HasBlendFuncAsset = false;
 			for (const FAssetData& AssetData : AssetDropOp->GetAssets())
 			{
-				if (AssetData.AssetClassPath.GetAssetName() == TEXT("MaterialFunctionMaterialLayer"))
-				{
-					HasLayerFuncAsset = true;
-				}
-				else if (AssetData.AssetClassPath.GetAssetName() == TEXT("MaterialFunctionMaterialLayerBlend"))
-				{
-					HasBlendFuncAsset = true;
-				}
+				HasLayerFuncAsset |= IsAssetDataAMaterialLayerFunction(AssetData);
+				HasBlendFuncAsset |= IsAssetDataAMaterialLayerBlendFunction(AssetData);
 			}
 
 			// Drop above or below could CREATE a new layer node:
@@ -447,14 +448,14 @@ FReply SMaterialSubstrateTreeItem::OnLayerDrop(const FDragDropEvent& DragDropEve
 			{
 				EMaterialParameterAssociation InAssociation = EMaterialParameterAssociation::GlobalParameter;
 			
-				if (AssetData.AssetClassPath.GetAssetName() == TEXT("MaterialFunctionMaterialLayer"))
+				if (IsAssetDataAMaterialLayerFunction(AssetData))
 				{
 					InAssociation = EMaterialParameterAssociation::LayerParameter;
 
 					Tree->RefreshOnAssetChange(AssetData, TargetNodeId, InAssociation);
 					DidModifyTree = true;
 				}
-				else if (AssetData.AssetClassPath.GetAssetName() == TEXT("MaterialFunctionMaterialLayerBlend"))
+				else if (IsAssetDataAMaterialLayerBlendFunction(AssetData))
 				{
 					InAssociation = EMaterialParameterAssociation::BlendParameter;
 					Tree->RefreshOnAssetChange(AssetData, TargetNodeId, InAssociation);
@@ -633,14 +634,19 @@ void SMaterialSubstrateTreeItem::Construct(const FArguments& InArgs, const TShar
                 		})
                 		.OnObjectChanged_Lambda([=, this](const FAssetData& InAssetData)
                 		{
-                			FSoftObjectPath ObjPath = InAssetData.GetSoftObjectPath();
-                			AssetChild->ParameterHandle->SetValue(ObjPath.TryLoad());
+							if (IsAssetDataAMaterialLayerBlendFunction(InAssetData))
+							{
+								FSoftObjectPath ObjPath = InAssetData.GetSoftObjectPath();
+								AssetChild->ParameterHandle->SetValue(ObjPath.TryLoad());
 
-                			Tree->FunctionInstanceHandle->NotifyPostChange(EPropertyChangeType::ValueSet);
-							Tree->CreateGroupsWidget();
-							Tree->RequestTreeRefresh();
+								Tree->FunctionInstanceHandle->NotifyPostChange(EPropertyChangeType::ValueSet);
+								Tree->CreateGroupsWidget();
+								Tree->RequestTreeRefresh();
+							}
                 		})
-                		.AllowedClass(UMaterialFunctionMaterialLayerBlend::StaticClass())
+                		// TODO: We need to authorize several UCLass in this asset picker UMaterialFunctionMaterialLayerBlend AND UMaterialFunctionMaterialLayerBlendInstance
+						// the widget only support one, we need to find a solution
+						//.AllowedClass(UMaterialFunctionMaterialLayerBlend::StaticClass())
                 		.AllowClear(true)
                 		.DisplayUseSelected(false)
                 		.DisplayBrowse(false);
diff --git a/Engine/Source/Runtime/Engine/Private/Materials/MaterialExpressions.cpp b/Engine/Source/Runtime/Engine/Private/Materials/MaterialExpressions.cpp
--- a/Engine/Source/Runtime/Engine/Private/Materials/MaterialExpressions.cpp
+++ b/Engine/Source/Runtime/Engine/Private/Materials/MaterialExpressions.cpp
@@ -18861,7 +18861,46 @@ namespace MLFT {
 		}
 		// Tree
 		LogMessage += InTab + TEXT("Tree:") + NewLine;
-		LogMessage += LFRD->Tree.Log(InTab);
+		//LogMessage += LFRD->Tree.Log(InTab);
+		struct TraverseNode { int32 NodeID; int32 Depth; FMaterialLayersFunctionsTree::FPayload Payload; };
+		TArray<TraverseNode> TraverseNodes;
+
+		auto Visitor = [&](FMaterialLayersFunctionsTree::FNodeId InNodeId, int32 InDepth, int32 InSiblingNum, FMaterialLayersFunctionsTree::FPayload InPayload)
+			{
+				TraverseNodes.Add({ InNodeId,  InDepth, InPayload });
+			};
+		LFRD->Tree.Traverse(Visitor);
+
+		for (auto ItNode = TraverseNodes.begin(); ItNode != TraverseNodes.end(); ++ItNode)
+		{
+			TraverseNode* N = &(*ItNode);
+			FString Tab = InTab;
+			for (int d = 0; d < N->Depth; d++)
+				Tab += "    ";
+
+			FString OffsetTab = "";
+			for (int d = N->Depth; d < 2; d++)
+				OffsetTab += "----";
+
+			FString LayerName = TEXT("               ");
+			bool HasLayerFunc = false;
+			if (N->Payload.Layer != FMaterialLayersFunctionsTree::InvalidId && LFRD->Layers[N->Payload.Layer])
+			{
+				LayerName = FString::Printf(TEXT("L%-2i %-12s"), N->Payload.Layer, *LFRD->Layers[N->Payload.Layer]->GetName());
+				HasLayerFunc = true;
+			}
+			FString BlendName = TEXT("               ");
+			bool HasBlendFunc = false;
+			if (N->Payload.Blend != FMaterialLayersFunctionsTree::InvalidId && LFRD->Blends[N->Payload.Blend])
+			{
+				BlendName = FString::Printf(TEXT("B%-2i %-12s"), N->Payload.Blend, *LFRD->Blends[N->Payload.Blend]->GetName());
+				HasBlendFunc = true;
+			}
+
+			LogMessage += Tab + FString::Printf(TEXT("%i -%s> %s %s"), N->NodeID, *OffsetTab, *BlendName, *LayerName) + NewLine;
+		}
+
+
 		return LogMessage;
 	}
 
