Fix issue where the Material editor always showed both the new and old viewport toolbars, regardless of the value of the CVAR ToolMenusViewportToolbars.

diff --git a/Engine/Source/Editor/MaterialEditor/Private/SMaterialEditorViewport.cpp b/Engine/Source/Editor/MaterialEditor/Private/SMaterialEditorViewport.cpp
--- a/Engine/Source/Editor/MaterialEditor/Private/SMaterialEditorViewport.cpp
+++ b/Engine/Source/Editor/MaterialEditor/Private/SMaterialEditorViewport.cpp
@@ -1011,7 +1011,25 @@ TSharedPtr<SWidget> SMaterialEditor3DPreviewViewport::BuildViewportToolbar()
 		}
 	}
 
-	return UToolMenus::Get()->GenerateWidget(MaterialEditorViewportToolbarName, ViewportToolbarContext);
+	// clang-format off
+	const TSharedRef<SWidget> NewViewportToolbar =
+		SNew(SBox)
+		// We need this visibility override because the material editor doesn't return its old viewport toolbar
+		// via the MakeViewportToolbar override, but adds it manually to the viewport overlay.
+		// Therefore, the central code that controls visibility of the old and new viewport toolbars get confused,
+		// and we must set the visibility override here.
+		.Visibility_Lambda(
+		[this]() -> EVisibility
+		{
+			return UE::UnrealEd::ShowNewViewportToolbars() ? EVisibility::Visible : EVisibility::Collapsed;
+
+		})
+		[
+			UToolMenus::Get()->GenerateWidget(MaterialEditorViewportToolbarName, ViewportToolbarContext)
+		];
+	// clang-format on
+
+	return NewViewportToolbar;
 }
 
 TSharedPtr<IPreviewProfileController> SMaterialEditor3DPreviewViewport::CreatePreviewProfileController()
