[Viewport Toolbar] Material Editor - Update preview scene settings menu

diff --git a/Engine/Source/Editor/MaterialEditor/Private/MaterialEditorViewportToolbarSections.cpp b/Engine/Source/Editor/MaterialEditor/Private/MaterialEditorViewportToolbarSections.cpp
--- a/Engine/Source/Editor/MaterialEditor/Private/MaterialEditorViewportToolbarSections.cpp
+++ b/Engine/Source/Editor/MaterialEditor/Private/MaterialEditorViewportToolbarSections.cpp
@@ -7,6 +7,7 @@
 #include "SMaterialEditorViewport.h"
 #include "ToolMenu.h"
 #include "ToolMenus.h"
+#include "Styling/SlateIconFinder.h"
 #include "ViewportToolbar/UnrealEdViewportToolbar.h"
 #include "ViewportToolbar/UnrealEdViewportToolbarContext.h"
 
@@ -98,12 +99,48 @@ void UE::MaterialEditor::FillShowSubmenu(UToolMenu* InMenu, bool bInShowViewport
 			}
 
 			UnnamedSection.AddMenuEntry(FMaterialEditorCommands::Get().ToggleMaterialStats);
-
-			UnnamedSection.AddSeparator(NAME_None);
-
-			UnnamedSection.AddMenuEntry(FMaterialEditorCommands::Get().TogglePreviewBackground);
 		}
 	}
 }
 
+void UE::MaterialEditor::ExtendPreviewSceneSettingsSubmenu(FName InSubmenuName)
+{
+	UToolMenu* const Submenu = UToolMenus::Get()->ExtendMenu(InSubmenuName);
+	if (!Submenu)
+	{
+		return;
+	}
+
+	FToolMenuSection& PreviewMeshSection = Submenu->FindOrAddSection(
+		"AssetViewerPreviewMeshSection", LOCTEXT("AssetViewerPreviewMeshSectionLabel", "Preview Mesh Options")
+	);
+	FToolMenuEntry& SphereEntry =
+		PreviewMeshSection.AddMenuEntry(FMaterialEditorCommands::Get().SetSpherePreview);
+	SphereEntry.SetShowInToolbarTopLevel(true);
+
+	FToolMenuEntry& CylinderEntry =
+		PreviewMeshSection.AddMenuEntry(FMaterialEditorCommands::Get().SetCylinderPreview);
+	CylinderEntry.SetShowInToolbarTopLevel(true);
+
+	FToolMenuEntry& PlaneEntry =
+		PreviewMeshSection.AddMenuEntry(FMaterialEditorCommands::Get().SetPlanePreview);
+	PlaneEntry.SetShowInToolbarTopLevel(true);
+
+	FToolMenuEntry& CubeEntry = PreviewMeshSection.AddMenuEntry(FMaterialEditorCommands::Get().SetCubePreview);
+	CubeEntry.SetShowInToolbarTopLevel(true);
+
+	FToolMenuEntry& MeshFromSelectionEntry =
+		PreviewMeshSection.AddMenuEntry(FMaterialEditorCommands::Get().SetPreviewMeshFromSelection);
+	MeshFromSelectionEntry.Label =
+		LOCTEXT("SetPreviewMeshFromSelectionLabel", "Static Mesh in Content Browser");
+	MeshFromSelectionEntry.SetShowInToolbarTopLevel(true);
+	MeshFromSelectionEntry.Icon = FSlateIconFinder::FindIconForClass(UStaticMesh::StaticClass());
+
+	FToolMenuSection& ProfileOptionsSection = Submenu->FindOrAddSection(
+		"AssetViewerProfileOptionsSection",
+		LOCTEXT("AssetViewerProfileOptionsSectionLabel", "Preview Scene Options")
+	);
+	ProfileOptionsSection.AddMenuEntry(FMaterialEditorCommands::Get().TogglePreviewBackground);
+}
+
 #undef LOCTEXT_NAMESPACE
diff --git a/Engine/Source/Editor/MaterialEditor/Private/MaterialEditorViewportToolbarSections.h b/Engine/Source/Editor/MaterialEditor/Private/MaterialEditorViewportToolbarSections.h
--- a/Engine/Source/Editor/MaterialEditor/Private/MaterialEditorViewportToolbarSections.h
+++ b/Engine/Source/Editor/MaterialEditor/Private/MaterialEditorViewportToolbarSections.h
@@ -12,4 +12,5 @@ namespace UE::MaterialEditor
 TSharedRef<SWidget> CreateShowMenuWidget(const TSharedRef<SMaterialEditor3DPreviewViewport>& InMaterialEditorViewport, bool bInShowViewportStatsToggle = true);
 FToolMenuEntry CreateShowSubmenu();
 void FillShowSubmenu(UToolMenu* InMenu, bool bInShowViewportStatsToggle = true);
+void ExtendPreviewSceneSettingsSubmenu(FName InSubmenuName);
 }
diff --git a/Engine/Source/Editor/MaterialEditor/Private/SMaterialEditorViewport.cpp b/Engine/Source/Editor/MaterialEditor/Private/SMaterialEditorViewport.cpp
--- a/Engine/Source/Editor/MaterialEditor/Private/SMaterialEditorViewport.cpp
+++ b/Engine/Source/Editor/MaterialEditor/Private/SMaterialEditorViewport.cpp
@@ -994,7 +994,15 @@ TSharedPtr<SWidget> SMaterialEditor3DPreviewViewport::BuildViewportToolbar()
 			RightSection.AddEntry(UE::UnrealEd::CreateViewModesSubmenu());
 			RightSection.AddEntry(UE::MaterialEditor::CreateShowSubmenu());
 			RightSection.AddEntry(UE::UnrealEd::CreatePerformanceAndScalabilitySubmenu());
-			RightSection.AddEntry(UE::UnrealEd::CreateAssetViewerProfileSubmenu());
+
+			// Add Preview Scene Submenu
+			{
+				UE::MaterialEditor::ExtendPreviewSceneSettingsSubmenu(
+					"MaterialEditor.ViewportToolbar.AssetViewerProfile"
+				);
+
+				RightSection.AddEntry(UE::UnrealEd::CreateAssetViewerProfileSubmenu());
+			}
 		}
 	}
 
