[Viewport Toolbar] Preview Scene Settings Tab opening Entry (Static Mesh, Material, Niagara, Persona)

Adds API to create the required entry in UnrealEdViewportToolbar.

Entry added to:
- Static Mesh Editor
- Material + Material Instance Editor
- Niagara Editor

diff --git a/Engine/Plugins/FX/Niagara/Source/NiagaraEditor/Private/Toolkits/SystemToolkitModes/NiagaraSystemToolkitModeBase.cpp b/Engine/Plugins/FX/Niagara/Source/NiagaraEditor/Private/Toolkits/SystemToolkitModes/NiagaraSystemToolkitModeBase.cpp
--- a/Engine/Plugins/FX/Niagara/Source/NiagaraEditor/Private/Toolkits/SystemToolkitModes/NiagaraSystemToolkitModeBase.cpp
+++ b/Engine/Plugins/FX/Niagara/Source/NiagaraEditor/Private/Toolkits/SystemToolkitModes/NiagaraSystemToolkitModeBase.cpp
@@ -475,7 +475,8 @@ TSharedRef<SDockTab> FNiagaraSystemToolkitModeBase::SpawnTab_Viewport(const FSpa
 	{
 		Toolkit->Viewport = SNew(SNiagaraSystemViewport, Toolkit->GetSystemViewModel().ToSharedRef())
 			.OnThumbnailCaptured(Toolkit.Get(), &FNiagaraSystemToolkit::OnThumbnailCaptured)
-			.Sequencer(Toolkit->GetSystemViewModel()->GetSequencer());
+			.Sequencer(Toolkit->GetSystemViewModel()->GetSequencer())
+			.AssetEditorToolkit(SystemToolkit);
 	}
 
 	TSharedRef<SDockTab> SpawnedTab =
diff --git a/Engine/Plugins/FX/Niagara/Source/NiagaraEditor/Private/Widgets/SNiagaraSystemViewport.cpp b/Engine/Plugins/FX/Niagara/Source/NiagaraEditor/Private/Widgets/SNiagaraSystemViewport.cpp
--- a/Engine/Plugins/FX/Niagara/Source/NiagaraEditor/Private/Widgets/SNiagaraSystemViewport.cpp
+++ b/Engine/Plugins/FX/Niagara/Source/NiagaraEditor/Private/Widgets/SNiagaraSystemViewport.cpp
@@ -36,6 +36,7 @@
 #include "Slate/SceneViewport.h"
 #include "ThumbnailRendering/ThumbnailManager.h"
 #include "ToolMenus.h"
+#include "Toolkits/SystemToolkitModes/NiagaraSystemToolkitModeBase.h"
 #include "ViewportToolbar/UnrealEdViewportToolbar.h"
 #include "ViewModels/NiagaraEmitterHandleViewModel.h"
 #include "ViewModels/NiagaraSystemViewModel.h"
@@ -833,7 +834,8 @@ void SNiagaraSystemViewport::Construct(const FArguments& InArgs, TSharedRef<FNia
 
 	OnThumbnailCapturedDelegate = InArgs._OnThumbnailCaptured;
 	Sequencer = InArgs._Sequencer;
-	
+	AssetEditorToolkit = InArgs._AssetEditorToolkit;
+
 	SEditorViewport::Construct( SEditorViewport::FArguments() );
 
 	Client->EngineShowFlags.SetGrid(Settings->IsShowGridInViewport());
@@ -1279,10 +1281,9 @@ TSharedPtr<SWidget> SNiagaraSystemViewport::BuildViewportToolbar()
 
 			// Add Preview Scene Submenu
 			{
-				UE::NiagaraSystemEditor::ExtendPreviewSceneSettingsSubmenu(
-					"NiagaraSystemEditor.ViewportToolbar.AssetViewerProfile"
-				);
-
+				const FName AssetViewerProfileMenuName = "NiagaraSystemEditor.ViewportToolbar.AssetViewerProfile";
+				UE::NiagaraSystemEditor::ExtendPreviewSceneSettingsSubmenu(AssetViewerProfileMenuName);
+				UE::UnrealEd::ExtendPreviewSceneSettingsWithTabEntry(AssetViewerProfileMenuName);
 				RightSection.AddEntry(UE::UnrealEd::CreateAssetViewerProfileSubmenu());
 			}
 		}
@@ -1296,6 +1297,8 @@ TSharedPtr<SWidget> SNiagaraSystemViewport::BuildViewportToolbar()
 		{
 			UUnrealEdViewportToolbarContext* const ContextObject = NewObject<UUnrealEdViewportToolbarContext>();
 			ContextObject->Viewport = SharedThis(this);
+			ContextObject->AssetEditorToolkit = AssetEditorToolkit;
+			ContextObject->PreviewSettingsTabId = FNiagaraSystemToolkitModeBase::PreviewSettingsTabId;
 
 			// Hook up our toolbar's filter for supported view modes.
 			ContextObject->IsViewModeSupported = UE::UnrealEd::IsViewModeSupportedDelegate::CreateLambda(
diff --git a/Engine/Plugins/FX/Niagara/Source/NiagaraEditor/Private/Widgets/SNiagaraSystemViewport.h b/Engine/Plugins/FX/Niagara/Source/NiagaraEditor/Private/Widgets/SNiagaraSystemViewport.h
--- a/Engine/Plugins/FX/Niagara/Source/NiagaraEditor/Private/Widgets/SNiagaraSystemViewport.h
+++ b/Engine/Plugins/FX/Niagara/Source/NiagaraEditor/Private/Widgets/SNiagaraSystemViewport.h
@@ -17,6 +17,7 @@
 
 class FNiagaraSystemViewModel;
 class UNiagaraComponent;
+class FAssetEditorToolkit;
 class FNiagaraSystemEditorViewportClient;
 class FNiagaraSystemInstance;
 class UNiagaraEffectType;
@@ -34,6 +35,7 @@ public:
 		SLATE_EVENT(FOnThumbnailCaptured, OnThumbnailCaptured)
 		/** So we can retrieve different data from sequencer to display in the viewport */
 		SLATE_ARGUMENT(TWeakPtr<ISequencer>, Sequencer)
+		SLATE_ARGUMENT(TWeakPtr<FAssetEditorToolkit>, AssetEditorToolkit)
 	SLATE_END_ARGS()
 	
 	void Construct(const FArguments& InArgs, TSharedRef<FNiagaraSystemViewModel> InSystemViewModel);
@@ -126,6 +128,8 @@ private:
 
 	TWeakPtr<ISequencer> Sequencer = nullptr;
 
+	TWeakPtr<FAssetEditorToolkit> AssetEditorToolkit = nullptr;
+
 	TWeakPtr<FNiagaraSystemViewModel> SystemViewModel;
 	
 	/** Preview Scene - uses advanced preview settings */
diff --git a/Engine/Source/Editor/MaterialEditor/Private/MaterialEditorTabs.cpp b/Engine/Source/Editor/MaterialEditor/Private/MaterialEditorTabs.cpp
--- a/Engine/Source/Editor/MaterialEditor/Private/MaterialEditorTabs.cpp
+++ b/Engine/Source/Editor/MaterialEditor/Private/MaterialEditorTabs.cpp
@@ -15,4 +15,11 @@ const FName FMaterialEditorTabs::ParameterDefaultsTabId(TEXT("MaterialEditor_Par
 const FName FMaterialEditorTabs::CustomPrimitiveTabId(TEXT("MaterialEditor_CustomPrimitiveData"));
 const FName FMaterialEditorTabs::LayerPropertiesTabId(TEXT("MaterialInstanceEditor_MaterialLayerProperties"));
 const FName FMaterialEditorTabs::SubstrateTabId(TEXT("MaterialEditor_Substrate"));
-const FName FMaterialEditorTabs::MaterialOverviewTabId(TEXT("MaterialEditor_MaterialExpressionOverview"));
\ No newline at end of file
+const FName FMaterialEditorTabs::MaterialOverviewTabId(TEXT("MaterialEditor_MaterialExpressionOverview"));
+
+
+const FName FMaterialInstanceEditorTabs::PreviewTabId( TEXT( "MaterialInstanceEditor_Preview" ) );
+const FName FMaterialInstanceEditorTabs::PropertiesTabId( TEXT( "MaterialInstanceEditor_MaterialProperties" ) );
+const FName FMaterialInstanceEditorTabs::LayerPropertiesTabId(TEXT("MaterialInstanceEditor_MaterialLayerProperties"));
+const FName FMaterialInstanceEditorTabs::PreviewSettingsTabId(TEXT("MaterialInstanceEditor_PreviewSettings"));
+const FName FMaterialInstanceEditorTabs::AssetBrowserTabId(TEXT("MaterialInstanceEditor_AssetBrowser"));
diff --git a/Engine/Source/Editor/MaterialEditor/Private/MaterialEditorTabs.h b/Engine/Source/Editor/MaterialEditor/Private/MaterialEditorTabs.h
--- a/Engine/Source/Editor/MaterialEditor/Private/MaterialEditorTabs.h
+++ b/Engine/Source/Editor/MaterialEditor/Private/MaterialEditorTabs.h
@@ -24,3 +24,16 @@ struct FMaterialEditorTabs
 	FMaterialEditorTabs() = delete;
 };
 
+struct FMaterialInstanceEditorTabs
+{
+	/**	The ids for the tabs spawned by this toolkit */
+	static const FName PreviewTabId;
+	static const FName PropertiesTabId;
+	static const FName LayerPropertiesTabId;
+	static const FName PreviewSettingsTabId;
+	static const FName AssetBrowserTabId;
+
+	// Disable default constructor
+	FMaterialInstanceEditorTabs() = delete;
+};
+
diff --git a/Engine/Source/Editor/MaterialEditor/Private/MaterialInstanceEditor.cpp b/Engine/Source/Editor/MaterialEditor/Private/MaterialInstanceEditor.cpp
--- a/Engine/Source/Editor/MaterialEditor/Private/MaterialInstanceEditor.cpp
+++ b/Engine/Source/Editor/MaterialEditor/Private/MaterialInstanceEditor.cpp
@@ -57,23 +57,20 @@
 #include "Widgets/Layout/SScrollBox.h"
 #include "DebugViewModeHelpers.h"
 #include "IContentBrowserSingleton.h"
+#include "MaterialEditorTabs.h"
+#include "MaterialEditorViewportToolbarSections.h"
 #include "Materials/MaterialFunctionMaterialLayer.h"
 #include "Materials/MaterialFunctionMaterialLayerBlend.h"
 #include "VT/RuntimeVirtualTexture.h"
 #include "SparseVolumeTexture/SparseVolumeTexture.h"
 #include "Widgets/Input/SButton.h"
 #include "Subsystems/AssetEditorSubsystem.h"
+#include "ViewportToolbar/UnrealEdViewportToolbar.h"
 
 #define LOCTEXT_NAMESPACE "MaterialInstanceEditor"
 
 DEFINE_LOG_CATEGORY_STATIC(LogMaterialInstanceEditor, Log, All);
 
-const FName FMaterialInstanceEditor::PreviewTabId( TEXT( "MaterialInstanceEditor_Preview" ) );
-const FName FMaterialInstanceEditor::PropertiesTabId( TEXT( "MaterialInstanceEditor_MaterialProperties" ) );
-const FName FMaterialInstanceEditor::LayerPropertiesTabId(TEXT("MaterialInstanceEditor_MaterialLayerProperties"));
-const FName FMaterialInstanceEditor::PreviewSettingsTabId(TEXT("MaterialInstanceEditor_PreviewSettings"));
-const FName FMaterialInstanceEditor::AssetBrowserTabId(TEXT("MaterialInstanceEditor_AssetBrowser"));
-
 //////////////////////////////////////////////////////////////////////////
 // SMaterialTreeWidgetItem
 class SMaterialTreeWidgetItem : public SMultiColumnTableRow< TWeakObjectPtr<UMaterialInterface> >
@@ -240,19 +237,19 @@ void FMaterialInstanceEditor::RegisterTabSpawners(const TSharedRef<class FTabMan
 
 	FAssetEditorToolkit::RegisterTabSpawners(InTabManager);
 	
-	InTabManager->RegisterTabSpawner( PreviewTabId, FOnSpawnTab::CreateSP( this, &FMaterialInstanceEditor::SpawnTab_Preview ) )
+	InTabManager->RegisterTabSpawner( FMaterialInstanceEditorTabs::PreviewTabId, FOnSpawnTab::CreateSP( this, &FMaterialInstanceEditor::SpawnTab_Preview ) )
 		.SetDisplayName( LOCTEXT( "ViewportTab", "Viewport" ) )
 		.SetGroup( WorkspaceMenuCategoryRef )
 		.SetIcon( FSlateIcon( FAppStyle::GetAppStyleSetName(), "LevelEditor.Tabs.Viewports" ) );
 	
-	InTabManager->RegisterTabSpawner( PropertiesTabId, FOnSpawnTab::CreateSP( this, &FMaterialInstanceEditor::SpawnTab_Properties ) )
+	InTabManager->RegisterTabSpawner( FMaterialInstanceEditorTabs::PropertiesTabId, FOnSpawnTab::CreateSP( this, &FMaterialInstanceEditor::SpawnTab_Properties ) )
 		.SetDisplayName( LOCTEXT( "PropertiesTab", "Details" ) )
 		.SetGroup( WorkspaceMenuCategoryRef )
 		.SetIcon( FSlateIcon( FAppStyle::GetAppStyleSetName(), "LevelEditor.Tabs.Details" ) );
 
 	if (!bIsFunctionPreviewMaterial)
 	{
-		InTabManager->RegisterTabSpawner(LayerPropertiesTabId, FOnSpawnTab::CreateSP(this, &FMaterialInstanceEditor::SpawnTab_LayerProperties))
+		InTabManager->RegisterTabSpawner(FMaterialInstanceEditorTabs::LayerPropertiesTabId, FOnSpawnTab::CreateSP(this, &FMaterialInstanceEditor::SpawnTab_LayerProperties))
 #if ENABLE_MATERIAL_LAYER_PROTOTYPE		
 			.SetDisplayName(LOCTEXT("MaterialLayersTab", "Material Layers"))
 #else
@@ -263,7 +260,7 @@ void FMaterialInstanceEditor::RegisterTabSpawners(const TSharedRef<class FTabMan
 			.SetIcon(FSlateIcon(FAppStyle::GetAppStyleSetName(), "LevelEditor.Tabs.Layers"));
 	}
 	
-	InTabManager->RegisterTabSpawner(PreviewSettingsTabId, FOnSpawnTab::CreateSP(this, &FMaterialInstanceEditor::SpawnTab_PreviewSettings))
+	InTabManager->RegisterTabSpawner(FMaterialInstanceEditorTabs::PreviewSettingsTabId, FOnSpawnTab::CreateSP(this, &FMaterialInstanceEditor::SpawnTab_PreviewSettings))
 		.SetDisplayName(LOCTEXT("PreviewSceneSettingsTab", "Preview Scene Settings"))
 		.SetGroup(WorkspaceMenuCategoryRef)
 		.SetIcon(FSlateIcon(FAppStyle::GetAppStyleSetName(), "LevelEditor.Tabs.Details"));
@@ -283,13 +280,13 @@ void FMaterialInstanceEditor::UnregisterTabSpawners(const TSharedRef<class FTabM
 {
 	FAssetEditorToolkit::UnregisterTabSpawners(InTabManager);
 
-	InTabManager->UnregisterTabSpawner( PreviewTabId );		
-	InTabManager->UnregisterTabSpawner( PropertiesTabId );	
+	InTabManager->UnregisterTabSpawner( FMaterialInstanceEditorTabs::PreviewTabId );
+	InTabManager->UnregisterTabSpawner(FMaterialInstanceEditorTabs:: PropertiesTabId );
 	if (!bIsFunctionPreviewMaterial)
 	{
-		InTabManager->UnregisterTabSpawner(LayerPropertiesTabId);
+		InTabManager->UnregisterTabSpawner(FMaterialInstanceEditorTabs::LayerPropertiesTabId);
 	}
-	InTabManager->UnregisterTabSpawner( PreviewSettingsTabId );
+	InTabManager->UnregisterTabSpawner( FMaterialInstanceEditorTabs::PreviewSettingsTabId );
 #if ENABLE_MATERIAL_LAYER_PROTOTYPE
 	InTabManager->UnregisterTabSpawner( AssetBrowserTabId );
 #endif
@@ -456,13 +453,13 @@ void FMaterialInstanceEditor::InitMaterialInstanceEditor( const EToolkitMode::Ty
 				->Split
 				(
 					FTabManager::NewStack()->SetSizeCoefficient(0.70f)->SetHideTabWell(true)
-					->AddTab(PreviewTabId, ETabState::OpenedTab)
-					->AddTab(PreviewSettingsTabId, ETabState::ClosedTab)
+					->AddTab(FMaterialInstanceEditorTabs::PreviewTabId, ETabState::OpenedTab)
+					->AddTab(FMaterialInstanceEditorTabs::PreviewSettingsTabId, ETabState::ClosedTab)
 				)
 				->Split
 				(
 					FTabManager::NewStack()->SetSizeCoefficient(0.30f)
-					->AddTab(PropertiesTabId, ETabState::OpenedTab)
+					->AddTab(FMaterialInstanceEditorTabs::PropertiesTabId, ETabState::OpenedTab)
 				)
 			)
 		);
@@ -520,15 +517,15 @@ void FMaterialInstanceEditor::InitMaterialInstanceEditor( const EToolkitMode::Ty
 					->Split
 					(
 						FTabManager::NewStack()->SetSizeCoefficient(0.70f)->SetHideTabWell(true)
-						->AddTab(PreviewTabId, ETabState::OpenedTab)
-						->AddTab(PreviewSettingsTabId, ETabState::ClosedTab)
+						->AddTab(FMaterialInstanceEditorTabs::PreviewTabId, ETabState::OpenedTab)
+						->AddTab(FMaterialInstanceEditorTabs::PreviewSettingsTabId, ETabState::ClosedTab)
 					)
 					->Split
 					(
 						FTabManager::NewStack()->SetSizeCoefficient(0.30f)
-						->AddTab(PropertiesTabId, ETabState::OpenedTab)
-						->AddTab(LayerPropertiesTabId, ETabState::OpenedTab)
-						->SetForegroundTab(PropertiesTabId)
+						->AddTab(FMaterialInstanceEditorTabs::PropertiesTabId, ETabState::OpenedTab)
+						->AddTab(FMaterialInstanceEditorTabs::LayerPropertiesTabId, ETabState::OpenedTab)
+						->SetForegroundTab(FMaterialInstanceEditorTabs::PropertiesTabId)
 					)
 				)
 			);
@@ -1054,7 +1051,7 @@ void FMaterialInstanceEditor::RefreshPreviewViewport()
 
 TSharedRef<SDockTab> FMaterialInstanceEditor::SpawnTab_Preview( const FSpawnTabArgs& Args )
 {	
-	check( Args.GetTabId().TabType == PreviewTabId );
+	check( Args.GetTabId().TabType == FMaterialInstanceEditorTabs::PreviewTabId );
 
 	TSharedRef<SDockTab> SpawnedTab = SNew(SDockTab)
 		.Label(LOCTEXT("ViewportTabTitle", "Viewport"))
@@ -1079,7 +1076,7 @@ TSharedRef<SDockTab> FMaterialInstanceEditor::SpawnTab_Preview( const FSpawnTabA
 
 TSharedRef<SDockTab> FMaterialInstanceEditor::SpawnTab_Properties( const FSpawnTabArgs& Args )
 {	
-	check( Args.GetTabId().TabType == PropertiesTabId );
+	check( Args.GetTabId().TabType == FMaterialInstanceEditorTabs::PropertiesTabId );
 
 	TSharedRef<SDockTab> SpawnedTab = SNew(SDockTab)
 		.Label(LOCTEXT("MaterialPropertiesTitle", "Details"))
@@ -1099,7 +1096,7 @@ TSharedRef<SDockTab> FMaterialInstanceEditor::SpawnTab_Properties( const FSpawnT
 
 TSharedRef<SDockTab> FMaterialInstanceEditor::SpawnTab_LayerProperties(const FSpawnTabArgs& Args)
 {
-	check(Args.GetTabId().TabType == LayerPropertiesTabId);
+	check(Args.GetTabId().TabType == FMaterialInstanceEditorTabs::LayerPropertiesTabId);
 
 	TSharedRef<SDockTab> SpawnedTab = SNew(SDockTab)
 		.Label(LOCTEXT("MaterialLayerPropertiesTitle", "Layer Parameters"))
@@ -1121,7 +1118,7 @@ TSharedRef<SDockTab> FMaterialInstanceEditor::SpawnTab_LayerProperties(const FSp
 
 TSharedRef<SDockTab> FMaterialInstanceEditor::SpawnTab_PreviewSettings(const FSpawnTabArgs& Args)
 {
-	check(Args.GetTabId() == PreviewSettingsTabId);
+	check(Args.GetTabId() == FMaterialInstanceEditorTabs::PreviewSettingsTabId);
 
 	TSharedRef<SWidget> InWidget = SNullWidget::NullWidget;
 	if (PreviewVC.IsValid())
@@ -1144,7 +1141,7 @@ TSharedRef<SDockTab> FMaterialInstanceEditor::SpawnTab_PreviewSettings(const FSp
 
 TSharedRef<SDockTab> FMaterialInstanceEditor::SpawnTab_AssetBrowser(const FSpawnTabArgs& Args)
 {
-	check(Args.GetTabId() == AssetBrowserTabId);
+	check(Args.GetTabId() == FMaterialInstanceEditorTabs::AssetBrowserTabId);
 
 	// TSharedRef<SWidget> InWidget = SNullWidget::NullWidget;
 	FContentBrowserModule& ContentBrowserModule = FModuleManager::Get().LoadModuleChecked<FContentBrowserModule>(TEXT("ContentBrowser"));
diff --git a/Engine/Source/Editor/MaterialEditor/Private/MaterialInstanceEditor.h b/Engine/Source/Editor/MaterialEditor/Private/MaterialInstanceEditor.h
--- a/Engine/Source/Editor/MaterialEditor/Private/MaterialInstanceEditor.h
+++ b/Engine/Source/Editor/MaterialEditor/Private/MaterialInstanceEditor.h
@@ -280,13 +280,6 @@ private:
 	UMaterial*					FunctionMaterialProxy; 
 	UMaterialInstanceConstant*	FunctionInstanceProxy; 
 
-	/**	The ids for the tabs spawned by this toolkit */
-	static const FName PreviewTabId;		
-	static const FName PropertiesTabId;	
-	static const FName LayerPropertiesTabId;
-	static const FName PreviewSettingsTabId;
-	static const FName AssetBrowserTabId;
-
 	/** Object used as material statistics manager */
 	TSharedPtr<class FMaterialStats> MaterialStatsManager;
 };
diff --git a/Engine/Source/Editor/MaterialEditor/Private/SMaterialEditorViewport.cpp b/Engine/Source/Editor/MaterialEditor/Private/SMaterialEditorViewport.cpp
--- a/Engine/Source/Editor/MaterialEditor/Private/SMaterialEditorViewport.cpp
+++ b/Engine/Source/Editor/MaterialEditor/Private/SMaterialEditorViewport.cpp
@@ -49,6 +49,7 @@
 #include "ViewportToolbar/UnrealEdViewportToolbar.h"
 
 #define LOCTEXT_NAMESPACE "MaterialEditor"
+#include "MaterialEditorTabs.h"
 #include "PreviewProfileController.h"
 #include "UnrealWidget.h"
 
@@ -997,9 +998,9 @@ TSharedPtr<SWidget> SMaterialEditor3DPreviewViewport::BuildViewportToolbar()
 
 			// Add Preview Scene Submenu
 			{
-				UE::MaterialEditor::ExtendPreviewSceneSettingsSubmenu(
-					"MaterialEditor.ViewportToolbar.AssetViewerProfile"
-				);
+				const FName AssetViewerProfileMenuName = "MaterialEditor.ViewportToolbar.AssetViewerProfile";
+				UE::MaterialEditor::ExtendPreviewSceneSettingsSubmenu(AssetViewerProfileMenuName);
+				UE::UnrealEd::ExtendPreviewSceneSettingsWithTabEntry(AssetViewerProfileMenuName);
 
 				RightSection.AddEntry(UE::UnrealEd::CreateAssetViewerProfileSubmenu());
 			}
@@ -1015,6 +1016,20 @@ TSharedPtr<SWidget> SMaterialEditor3DPreviewViewport::BuildViewportToolbar()
 			UUnrealEdViewportToolbarContext* const ContextObject =
 				UE::UnrealEd::CreateViewportToolbarDefaultContext(SharedThis(this));
 
+			if (TSharedPtr<IMaterialEditor> MaterialEditorPinned = MaterialEditorPtr.Pin())
+			{
+				if (Cast<UMaterial>(MaterialEditorPinned->GetMaterialInterface()))
+				{
+					ContextObject->PreviewSettingsTabId = FMaterialEditorTabs::PreviewSettingsTabId;
+				}
+				else if (Cast<UMaterialInstance>(MaterialEditorPinned->GetMaterialInterface()))
+				{
+					ContextObject->PreviewSettingsTabId = FMaterialInstanceEditorTabs::PreviewSettingsTabId;
+				}
+
+				ContextObject->AssetEditorToolkit = MaterialEditorPinned;
+			}
+
 			ViewportToolbarContext.AddObject(ContextObject);
 		}
 	}
diff --git a/Engine/Source/Editor/Persona/Private/SAnimationEditorViewport.cpp b/Engine/Source/Editor/Persona/Private/SAnimationEditorViewport.cpp
--- a/Engine/Source/Editor/Persona/Private/SAnimationEditorViewport.cpp
+++ b/Engine/Source/Editor/Persona/Private/SAnimationEditorViewport.cpp
@@ -294,9 +294,9 @@ TSharedPtr<SWidget> SAnimationEditorViewport::BuildViewportToolbar()
 
 			// Add Preview Scene Submenu
 			{
-				UE::AnimationEditor::ExtendPreviewSceneSettingsSubmenu(
-					"AnimationEditor.ViewportToolbar.AssetViewerProfile"
-				);
+				const FName AssetViewerProfileMenuName = "AnimationEditor.ViewportToolbar.AssetViewerProfile";
+				UE::AnimationEditor::ExtendPreviewSceneSettingsSubmenu(AssetViewerProfileMenuName);
+				UE::UnrealEd::ExtendPreviewSceneSettingsWithTabEntry(AssetViewerProfileMenuName);
 
 				RightSection.AddEntry(UE::UnrealEd::CreateAssetViewerProfileSubmenu());
 			}
@@ -328,6 +328,9 @@ TSharedPtr<SWidget> SAnimationEditorViewport::BuildViewportToolbar()
 				UE::UnrealEd::CreateViewportToolbarDefaultContext(SharedThis(this));
 
 			ContextObject->bShowGridToggleInPreviewSceneSettings = true;
+			ContextObject->AssetEditorToolkit = AssetEditorToolkitPtr;
+			ContextObject->PreviewSettingsTabId = FPersonaTabs::AdvancedPreviewSceneSettingsID;
+
 			ViewportToolbarContext.AddObject(ContextObject);
 		}
 
diff --git a/Engine/Source/Editor/StaticMeshEditor/Private/SStaticMeshEditorViewport.cpp b/Engine/Source/Editor/StaticMeshEditor/Private/SStaticMeshEditorViewport.cpp
--- a/Engine/Source/Editor/StaticMeshEditor/Private/SStaticMeshEditorViewport.cpp
+++ b/Engine/Source/Editor/StaticMeshEditor/Private/SStaticMeshEditorViewport.cpp
@@ -21,6 +21,7 @@
 #include "SEditorViewportToolBarMenu.h"
 #include "Editor.h"
 #include "PreviewProfileController.h"
+#include "StaticMeshEditor.h"
 #include "StaticMeshEditorViewportToolbarSections.h"
 #include "ToolMenus.h"
 #include "ViewportToolbar/UnrealEdViewportToolbar.h"
@@ -1014,7 +1015,12 @@ TSharedPtr<SWidget> SStaticMeshEditorViewport::BuildViewportToolbar()
 			RightSection.AddEntry(UE::StaticMeshEditor::CreateShowSubmenu());
 			RightSection.AddEntry(UE::StaticMeshEditor::CreateLODSubmenu());
 			RightSection.AddEntry(UE::UnrealEd::CreatePerformanceAndScalabilitySubmenu());
-			RightSection.AddEntry(UE::UnrealEd::CreateAssetViewerProfileSubmenu());
+
+			// Add Preview Scene Submenu
+			{
+				UE::UnrealEd::ExtendPreviewSceneSettingsWithTabEntry("StaticMeshEditor.ViewportToolbar.AssetViewerProfile");
+				RightSection.AddEntry(UE::UnrealEd::CreateAssetViewerProfileSubmenu());
+			}
 		}
 	}
 
@@ -1031,6 +1037,9 @@ TSharedPtr<SWidget> SStaticMeshEditorViewport::BuildViewportToolbar()
 			ContextObject->bShowGridToggleInPreviewSceneSettings = true;
 			ContextObject->ExcludedShowMenuFlags.AddUnique(FEngineShowFlags::SF_Grid);
 
+			ContextObject->AssetEditorToolkit = StaticMeshEditorPtr;
+			ContextObject->PreviewSettingsTabId = FStaticMeshEditor::PreviewSceneSettingsTabId;
+
 			ViewportToolbarContext.AddObject(ContextObject);
 		}
 	}
diff --git a/Engine/Source/Editor/UnrealEd/Private/ViewportToolbar/UnrealEdViewportToolbar.cpp b/Engine/Source/Editor/UnrealEd/Private/ViewportToolbar/UnrealEdViewportToolbar.cpp
--- a/Engine/Source/Editor/UnrealEd/Private/ViewportToolbar/UnrealEdViewportToolbar.cpp
+++ b/Engine/Source/Editor/UnrealEd/Private/ViewportToolbar/UnrealEdViewportToolbar.cpp
@@ -2587,6 +2587,45 @@ FToolMenuEntry CreateAssetViewerProfileSubmenu()
 	);
 }
 
+void ExtendPreviewSceneSettingsWithTabEntry(FName InAssetViewerProfileSubmenuName)
+{
+	UToolMenu* const Submenu = UToolMenus::Get()->ExtendMenu(InAssetViewerProfileSubmenuName);
+	if (!Submenu)
+	{
+		return;
+	}
+
+	FToolMenuSection& PreviewMeshSection = Submenu->FindOrAddSection("PreviewSceneTabOpeningSection");
+
+	FToolUIAction UIAction;
+	UIAction.ExecuteAction = FToolMenuExecuteAction::CreateLambda(
+		[](const FToolMenuContext& InContext)
+		{
+			UUnrealEdViewportToolbarContext* Context = Cast<UUnrealEdViewportToolbarContext>(InContext.FindByClass(UUnrealEdViewportToolbarContext::StaticClass()));
+			if (!Context)
+			{
+				return;
+			}
+			if (TSharedPtr<FAssetEditorToolkit> ToolkitPinned = Context->AssetEditorToolkit.Pin())
+			{
+				if (TSharedPtr<FTabManager> TabManager = ToolkitPinned->GetTabManager())
+				{
+					TabManager->TryInvokeTab(Context->PreviewSettingsTabId);
+				}
+			}
+		}
+	);
+
+	PreviewMeshSection.AddEntry(FToolMenuEntry::InitMenuEntry(
+		"OpenPreviewSceneSettingsTab",
+		LOCTEXT("OpenPreviewSceneSettingsTabLabel", "Preview Scene Settings..."),
+		LOCTEXT("OpenPreviewSceneSettingsTabTooltip", "Opens a details tab with Preview Scene Settings."),
+		FSlateIcon(FAppStyle::GetAppStyleSetName(), "ViewportToolbar.PreviewSceneSettings"),
+		UIAction,
+		EUserInterfaceActionType::Button
+	));
+}
+
 void AddExposureSection(UToolMenu* InMenu, const TSharedPtr<SEditorViewport>& EditorViewport)
 {
 	const FEditorViewportCommands& BaseViewportCommands = FEditorViewportCommands::Get();
diff --git a/Engine/Source/Editor/UnrealEd/Public/ViewportToolbar/UnrealEdViewportToolbar.h b/Engine/Source/Editor/UnrealEd/Public/ViewportToolbar/UnrealEdViewportToolbar.h
--- a/Engine/Source/Editor/UnrealEd/Public/ViewportToolbar/UnrealEdViewportToolbar.h
+++ b/Engine/Source/Editor/UnrealEd/Public/ViewportToolbar/UnrealEdViewportToolbar.h
@@ -10,6 +10,7 @@
 #include "UnrealWidgetFwd.h"
 
 class FEditorViewportClient;
+class FTabManager;
 class FText;
 class IPreviewProfileController;
 class IPreviewProfileController;
@@ -18,6 +19,7 @@ class UToolMenu;
 enum ELevelViewportType : int;
 enum ERotationGridMode : int;
 struct FNewToolMenuChoice;
+struct FTabId;
 struct FToolMenuEntry;
 
 namespace UE::UnrealEd
@@ -153,6 +155,7 @@ UNREALED_API FToolMenuEntry CreateCameraSubmenu();
 UE_DEPRECATED(5.6, "Use CreateAssetViewerProfileSubmenu() instead.")
 UNREALED_API FToolMenuEntry CreateViewportToolbarAssetViewerProfileSubmenu();
 UNREALED_API FToolMenuEntry CreateAssetViewerProfileSubmenu();
+UNREALED_API void ExtendPreviewSceneSettingsWithTabEntry(FName InAssetViewerProfileSubmenuName);
 
 UNREALED_API void AddExposureSection(UToolMenu* InMenu, const TSharedPtr<SEditorViewport>& EditorViewport);
 UNREALED_API void PopulateCameraMenu(UToolMenu* InMenu);
diff --git a/Engine/Source/Editor/UnrealEd/Public/ViewportToolbar/UnrealEdViewportToolbarContext.h b/Engine/Source/Editor/UnrealEd/Public/ViewportToolbar/UnrealEdViewportToolbarContext.h
--- a/Engine/Source/Editor/UnrealEd/Public/ViewportToolbar/UnrealEdViewportToolbarContext.h
+++ b/Engine/Source/Editor/UnrealEd/Public/ViewportToolbar/UnrealEdViewportToolbarContext.h
@@ -3,12 +3,15 @@
 #pragma once
 
 #include "Engine/EngineBaseTypes.h"
+#include "Framework/Docking/TabManager.h"
 #include "ShowFlags.h"
 #include "Templates/SharedPointer.h"
 #include "UObject/Object.h"
 
 #include "UnrealEdViewportToolbarContext.generated.h"
 
+class FAssetEditorToolkit;
+
 namespace UE::UnrealEd
 {
 
@@ -51,4 +54,14 @@ public:
 	 * Add flags to this array to remove them from the Show Menu
 	 */
 	TArray<FEngineShowFlags::EShowFlag> ExcludedShowMenuFlags;
+
+	/**
+	 * Can be used to retrieve data e.g. TabManager
+	 */
+	TWeakPtr<FAssetEditorToolkit> AssetEditorToolkit;
+
+	/**
+	 * Identifier for the Preview Settings Tab
+	 */
+	FTabId PreviewSettingsTabId;
 };
