Fix broken TabID references used under ENABLE_MATERIAL_LAYER_PROTOTYPE

diff --git a/Engine/Source/Editor/MaterialEditor/Private/MaterialInstanceEditor.cpp b/Engine/Source/Editor/MaterialEditor/Private/MaterialInstanceEditor.cpp
--- a/Engine/Source/Editor/MaterialEditor/Private/MaterialInstanceEditor.cpp
+++ b/Engine/Source/Editor/MaterialEditor/Private/MaterialInstanceEditor.cpp
@@ -266,7 +266,7 @@ void FMaterialInstanceEditor::RegisterTabSpawners(const TSharedRef<class FTabMan
 		.SetIcon(FSlateIcon(FAppStyle::GetAppStyleSetName(), "LevelEditor.Tabs.Details"));
 
 #if ENABLE_MATERIAL_LAYER_PROTOTYPE
-	InTabManager->RegisterTabSpawner(AssetBrowserTabId, FOnSpawnTab::CreateSP(this, &FMaterialInstanceEditor::SpawnTab_AssetBrowser))
+	InTabManager->RegisterTabSpawner(FMaterialInstanceEditorTabs::AssetBrowserTabId, FOnSpawnTab::CreateSP(this, &FMaterialInstanceEditor::SpawnTab_AssetBrowser))
 		.SetDisplayName(LOCTEXT("AssetBrowserTab", "Asset Browser"))
 		.SetGroup(WorkspaceMenuCategoryRef)
 		.SetIcon(FSlateIcon(FAppStyle::GetAppStyleSetName(), "LevelEditor.Tabs.Details"));
@@ -288,7 +288,7 @@ void FMaterialInstanceEditor::UnregisterTabSpawners(const TSharedRef<class FTabM
 	}
 	InTabManager->UnregisterTabSpawner( FMaterialInstanceEditorTabs::PreviewSettingsTabId );
 #if ENABLE_MATERIAL_LAYER_PROTOTYPE
-	InTabManager->UnregisterTabSpawner( AssetBrowserTabId );
+	InTabManager->UnregisterTabSpawner( FMaterialInstanceEditorTabs::AssetBrowserTabId );
 #endif
 	MaterialStatsManager->UnregisterTabs();
 
@@ -478,8 +478,8 @@ void FMaterialInstanceEditor::InitMaterialInstanceEditor( const EToolkitMode::Ty
 					->Split
 					(
 						FTabManager::NewStack()->SetSizeCoefficient(0.70f)->SetHideTabWell(true)
-						->AddTab(PreviewTabId, ETabState::OpenedTab)
-						->AddTab(PreviewSettingsTabId, ETabState::ClosedTab)
+						->AddTab(FMaterialInstanceEditorTabs::PreviewTabId, ETabState::OpenedTab)
+						->AddTab(FMaterialInstanceEditorTabs::PreviewSettingsTabId, ETabState::ClosedTab)
 					)
 					->Split
 					(
@@ -490,18 +490,18 @@ void FMaterialInstanceEditor::InitMaterialInstanceEditor( const EToolkitMode::Ty
 							->Split
 							(
 								FTabManager::NewStack()/*->SetSizeCoefficient(0.30f)*/
-								->AddTab(LayerPropertiesTabId, ETabState::OpenedTab)
+								->AddTab(FMaterialInstanceEditorTabs::LayerPropertiesTabId, ETabState::OpenedTab)
 							)
 							->Split
 							(
 								FTabManager::NewStack()/*->SetSizeCoefficient(0.30f)*/
-								->AddTab(AssetBrowserTabId, ETabState::OpenedTab)
+								->AddTab(FMaterialInstanceEditorTabs::AssetBrowserTabId, ETabState::OpenedTab)
 							)
 						)
 						->Split
 						(
 						FTabManager::NewStack()->SetSizeCoefficient(0.30f)
-								->AddTab(PropertiesTabId, ETabState::OpenedTab)
+								->AddTab(FMaterialInstanceEditorTabs::PropertiesTabId, ETabState::OpenedTab)
 						)
 					)
 				)
