[Viewport Toolbar] Use radio buttons instead of checkboxes for preview mesh in new toolbar

diff --git a/Engine/Source/Editor/MaterialEditor/Private/MaterialEditorActions.cpp b/Engine/Source/Editor/MaterialEditor/Private/MaterialEditorActions.cpp
--- a/Engine/Source/Editor/MaterialEditor/Private/MaterialEditorActions.cpp
+++ b/Engine/Source/Editor/MaterialEditor/Private/MaterialEditorActions.cpp
@@ -31,11 +31,11 @@ void FMaterialEditorCommands::RegisterCommands()
 
 	UI_COMMAND( ShowAllMaterialParameters, "Show Inactive", "Show the material instance parameters currently hidden behind static switches.", EUserInterfaceActionType::ToggleButton, FInputChord() );
 
-	UI_COMMAND( SetCylinderPreview, "Cylinder", "Sets the preview mesh to a cylinder primitive.", EUserInterfaceActionType::ToggleButton, FInputChord() );
-	UI_COMMAND( SetSpherePreview, "Sphere", "Sets the preview mesh to a sphere primitive.", EUserInterfaceActionType::ToggleButton, FInputChord() );
-	UI_COMMAND( SetPlanePreview, "Plane", "Sets the preview mesh to a plane primitive.", EUserInterfaceActionType::ToggleButton, FInputChord() );
-	UI_COMMAND( SetCubePreview, "Cube", "Sets the preview mesh to a cube primitive.", EUserInterfaceActionType::ToggleButton, FInputChord() );
-	UI_COMMAND( SetPreviewMeshFromSelection, "Mesh", "Sets the preview mesh based on the current content browser selection.", EUserInterfaceActionType::ToggleButton, FInputChord() );
+	UI_COMMAND( SetCylinderPreview, "Cylinder", "Sets the preview mesh to a cylinder primitive.", EUserInterfaceActionType::RadioButton, FInputChord() );
+	UI_COMMAND( SetSpherePreview, "Sphere", "Sets the preview mesh to a sphere primitive.", EUserInterfaceActionType::RadioButton, FInputChord() );
+	UI_COMMAND( SetPlanePreview, "Plane", "Sets the preview mesh to a plane primitive.", EUserInterfaceActionType::RadioButton, FInputChord() );
+	UI_COMMAND( SetCubePreview, "Cube", "Sets the preview mesh to a cube primitive.", EUserInterfaceActionType::RadioButton, FInputChord() );
+	UI_COMMAND( SetPreviewMeshFromSelection, "Mesh", "Sets the preview mesh based on the current content browser selection.", EUserInterfaceActionType::RadioButton, FInputChord() );
 	UI_COMMAND( TogglePreviewBackground, "Background", "Toggles the preview pane's background.", EUserInterfaceActionType::ToggleButton, FInputChord() );
 
 	UI_COMMAND( CameraHome, "Home", "Goes home on the canvas.", EUserInterfaceActionType::Button, FInputChord() );
