[Viewport Toolbar] Material Editor: hide preview mesh buttons from bottom right of viewport when only new toolbar is visible

diff --git a/Engine/Source/Editor/MaterialEditor/Private/SMaterialEditorViewport.cpp b/Engine/Source/Editor/MaterialEditor/Private/SMaterialEditorViewport.cpp
--- a/Engine/Source/Editor/MaterialEditor/Private/SMaterialEditorViewport.cpp
+++ b/Engine/Source/Editor/MaterialEditor/Private/SMaterialEditorViewport.cpp
@@ -1092,6 +1092,12 @@ void SMaterialEditor3DPreviewViewport::PopulateViewportOverlays(TSharedRef<class
 		.VAlign(VAlign_Bottom)
 		[
 			SNew(SMaterialEditorViewportPreviewShapeToolBar, SharedThis(this))
+			.Visibility_Lambda(
+				[this]() -> EVisibility
+				{
+					return UE::UnrealEd::ShowOldViewportToolbars() ? EVisibility::Visible : EVisibility::Collapsed;
+				}
+			)
 		];
 		// clang-format on
 
