Fix material editor keeping reference to derived material instances during level unloading

diff --git a/Engine/Source/Editor/MaterialEditor/Private/MaterialEditor.cpp b/Engine/Source/Editor/MaterialEditor/Private/MaterialEditor.cpp
--- a/Engine/Source/Editor/MaterialEditor/Private/MaterialEditor.cpp
+++ b/Engine/Source/Editor/MaterialEditor/Private/MaterialEditor.cpp
@@ -994,9 +994,12 @@ void FMaterialEditor::InitMaterialEditor( const EToolkitMode::Type Mode, const T
 	if (Generator.IsValid() && MaterialEditorInstance)
 	{
 		UpdateGenerator();
-
 	}
 
+	// Register for pre/post world cleanup
+	OnWorldCleanupDelegateHandle = FWorldDelegates::OnWorldCleanup.AddSP(this, &FMaterialEditor::OnWorldCleanup);
+	OnPostWorldCleanupDelegateHandle = FWorldDelegates::OnPostWorldCleanup.AddSP(this, &FMaterialEditor::OnPostWorldCleanup);
+
 	if (OriginalMaterial->bUsedAsSpecialEngineMaterial)
 	{
 		FSuppressableWarningDialog::FSetupInfo Info(
@@ -1071,6 +1074,18 @@ FMaterialEditor::~FMaterialEditor()
 {
 	bDestructing = true;
 
+	if (OnWorldCleanupDelegateHandle.IsValid())
+	{
+		FWorldDelegates::OnWorldCleanup.Remove(OnWorldCleanupDelegateHandle);
+		OnWorldCleanupDelegateHandle.Reset();
+	}
+
+	if (OnPostWorldCleanupDelegateHandle.IsValid())
+	{
+		FWorldDelegates::OnPostWorldCleanup.Remove(OnPostWorldCleanupDelegateHandle);
+		OnPostWorldCleanupDelegateHandle.Reset();
+	}
+
 	// Notify other editors if this material editor has a post process named output, which may affect their preview
 	NotifyUserSceneTextureLoadOrUnload();
 
@@ -6738,9 +6753,9 @@ void FMaterialEditor::CreateDerivedMaterialInstancesPreviews()
 	TArray<FString> DisplayNames;
 	DisplayNames.Push(OriginalMaterial->GetName());
 
-	if (MaterialStatsManager->GetProvideDerivedMIFlag())
+	const int32 MaxCount = CVarMaterialEdMaxDerivedMaterialInstances.GetValueOnGameThread();
+	if (MaterialStatsManager->GetProvideDerivedMIFlag() && MaxCount != 0)
 	{
-		const int32 MaxCount = CVarMaterialEdMaxDerivedMaterialInstances.GetValueOnGameThread();
 		// TODO consider a mode where we load all MaterialChildList also
 
 		for (TObjectIterator<UMaterialInstance> It(/*AdditionalExclusionFlags = */RF_ClassDefaultObject, /*bIncludeDerivedClasses = */true, /*InInternalExclusionFlags = */EInternalObjectFlags::Garbage); It; ++It)
@@ -8039,6 +8054,25 @@ void FMaterialEditor::CleanUnusedExpressions()
 	}
 }
 
+void FMaterialEditor::OnWorldCleanup(UWorld* World, bool bSessionEnded, bool bCleanupResources)
+{
+	// no-op if there are no derived materials used
+	if (DerivedMaterialInstances.Num() == 0 && OriginalDerivedMaterialInstances.Num() == 0)
+	{
+		return;
+	}
+
+	DerivedMaterialInstances.Empty();
+	OriginalDerivedMaterialInstances.Empty();
+	// will free all referenced resources via FMaterial::DeleteMaterialsOnRenderThread
+	MaterialStatsManager->SetMaterial(bStatsFromPreviewMaterial ? Material : OriginalMaterial, {});
+}
+
+void FMaterialEditor::OnPostWorldCleanup(UWorld* World, bool bSessionEnded, bool bCleanupResources)
+{
+	MaterialStatsManager->SignalMaterialChanged();
+}
+
 bool FMaterialEditor::CheckExpressionRemovalWarnings(const TArray<UEdGraphNode*>& NodesToRemove)
 {
 	FString FunctionWarningString;
diff --git a/Engine/Source/Editor/MaterialEditor/Private/MaterialEditor.h b/Engine/Source/Editor/MaterialEditor/Private/MaterialEditor.h
--- a/Engine/Source/Editor/MaterialEditor/Private/MaterialEditor.h
+++ b/Engine/Source/Editor/MaterialEditor/Private/MaterialEditor.h
@@ -923,6 +923,16 @@ private:
 	 */
 	void CleanUnusedExpressions();
 
+	/**
+	 * Cleans up any referenced resources.
+	 */
+	void OnWorldCleanup(UWorld* World, bool bSessionEnded, bool bCleanupResources);
+
+	/**
+	 * Refreshes editor after world cleanup.
+	 */
+	void OnPostWorldCleanup(UWorld* World, bool bSessionEnded, bool bCleanupResources);
+
 	/**
 	 * Perform a deep copy of all expressions within the given graph.
 	 *
@@ -1084,4 +1094,8 @@ private:
 	TArray< FAssetData > FunctionChildList;
 
 	TSharedPtr<class IPropertyRowGenerator> Generator;
+
+	/** Delegate handles for FWorldDelegates registrations. */
+	FDelegateHandle OnWorldCleanupDelegateHandle;
+	FDelegateHandle OnPostWorldCleanupDelegateHandle;
 };
