Increase the Material Expression "Title" weight.

- When searching for nodes in the right click context menu we weight direct matches with the "Title" of the material expression more.
- This way searching for "object posi" finds the "ObjectPositionWS" node, instead of finding the "Object Scale" material function.
- The search heuristic would also search and weight the node's category, description, and tooltip.  This meant sometimes we would find the non-intuitive node.
- Users expected to search and find nodes based on their name and not their category or tooltip

[RN] minor Rendering - Materials and Shaders

diff --git a/Engine/Source/Editor/MaterialEditor/Public/MaterialEditorSettings.h b/Engine/Source/Editor/MaterialEditor/Public/MaterialEditorSettings.h
--- a/Engine/Source/Editor/MaterialEditor/Public/MaterialEditorSettings.h
+++ b/Engine/Source/Editor/MaterialEditor/Public/MaterialEditorSettings.h
@@ -114,7 +114,7 @@ public:
 
 	/** The amount of weight placed on the search items title". */
 	UPROPERTY(config, EditAnywhere, Category = "Context Menu", meta = (DisplayName = "Node Title Weight"))
-	float ContextMenuNodeTitleWeight = 20.0f;
+	float ContextMenuNodeTitleWeight = 25.0f;
 
 	/** The amount of weight placed on search items keyword. */
 	UPROPERTY(config, EditAnywhere, Category = "Context Menu", meta = (DisplayName = "Keyword Weight"))
