Material editor - make sure preview meshes are always visible in the viewport toolbar

diff --git a/Engine/Source/Editor/MaterialEditor/Private/MaterialEditorViewportToolbarSections.cpp b/Engine/Source/Editor/MaterialEditor/Private/MaterialEditorViewportToolbarSections.cpp
--- a/Engine/Source/Editor/MaterialEditor/Private/MaterialEditorViewportToolbarSections.cpp
+++ b/Engine/Source/Editor/MaterialEditor/Private/MaterialEditorViewportToolbarSections.cpp
@@ -111,36 +111,37 @@ void UE::MaterialEditor::ExtendPreviewSceneSettingsSubmenu(FName InSubmenuName)
 		return;
 	}
 
-	FToolMenuSection& PreviewMeshSection = Submenu->FindOrAddSection(
-		"AssetViewerPreviewMeshSection", LOCTEXT("AssetViewerPreviewMeshSectionLabel", "Preview Mesh Options")
-	);
-	FToolMenuEntry& SphereEntry =
-		PreviewMeshSection.AddMenuEntry(FMaterialEditorCommands::Get().SetSpherePreview);
-	SphereEntry.SetShowInToolbarTopLevel(true);
-
-	FToolMenuEntry& CylinderEntry =
-		PreviewMeshSection.AddMenuEntry(FMaterialEditorCommands::Get().SetCylinderPreview);
-	CylinderEntry.SetShowInToolbarTopLevel(true);
-
-	FToolMenuEntry& PlaneEntry =
-		PreviewMeshSection.AddMenuEntry(FMaterialEditorCommands::Get().SetPlanePreview);
-	PlaneEntry.SetShowInToolbarTopLevel(true);
-
-	FToolMenuEntry& CubeEntry = PreviewMeshSection.AddMenuEntry(FMaterialEditorCommands::Get().SetCubePreview);
-	CubeEntry.SetShowInToolbarTopLevel(true);
-
-	FToolMenuEntry& MeshFromSelectionEntry =
-		PreviewMeshSection.AddMenuEntry(FMaterialEditorCommands::Get().SetPreviewMeshFromSelection);
-	MeshFromSelectionEntry.Label =
-		LOCTEXT("SetPreviewMeshFromSelectionLabel", "Static Mesh in Content Browser");
-	MeshFromSelectionEntry.SetShowInToolbarTopLevel(true);
-	MeshFromSelectionEntry.Icon = FSlateIconFinder::FindIconForClass(UStaticMesh::StaticClass());
-
-	FToolMenuSection& ProfileOptionsSection = Submenu->FindOrAddSection(
-		"AssetViewerProfileOptionsSection",
-		LOCTEXT("AssetViewerProfileOptionsSectionLabel", "Preview Scene Options")
-	);
-	ProfileOptionsSection.AddMenuEntry(FMaterialEditorCommands::Get().TogglePreviewBackground);
+	{
+		FToolMenuSection& PreviewMeshSection = Submenu->FindOrAddSection(
+			"AssetViewerPreviewMeshSection", LOCTEXT("AssetViewerPreviewMeshSectionLabel", "Preview Mesh Options")
+		);
+
+		for (const TSharedPtr<FUICommandInfo>& Command : {
+				 FMaterialEditorCommands::Get().SetSpherePreview,
+				 FMaterialEditorCommands::Get().SetCylinderPreview,
+				 FMaterialEditorCommands::Get().SetPlanePreview,
+				 FMaterialEditorCommands::Get().SetCubePreview,
+			 })
+		{
+			FToolMenuEntry& Entry = PreviewMeshSection.AddMenuEntry(Command);
+			Entry.SetShowInToolbarTopLevel(true);
+			Entry.ToolBarData.ResizeParams.ClippingPriority = 2000;
+		}
+
+		FToolMenuEntry& MeshFromSelectionEntry =
+			PreviewMeshSection.AddMenuEntry(FMaterialEditorCommands::Get().SetPreviewMeshFromSelection);
+		MeshFromSelectionEntry.Label = LOCTEXT("SetPreviewMeshFromSelectionLabel", "Static Mesh in Content Browser");
+		MeshFromSelectionEntry.Icon = FSlateIconFinder::FindIconForClass(UStaticMesh::StaticClass());
+		MeshFromSelectionEntry.SetShowInToolbarTopLevel(true);
+		MeshFromSelectionEntry.ToolBarData.ResizeParams.ClippingPriority = 2000;
+	}
+
+	{
+		FToolMenuSection& ProfileOptionsSection = Submenu->FindOrAddSection(
+			"AssetViewerProfileOptionsSection", LOCTEXT("AssetViewerProfileOptionsSectionLabel", "Preview Scene Options")
+		);
+		ProfileOptionsSection.AddMenuEntry(FMaterialEditorCommands::Get().TogglePreviewBackground);
+	}
 }
 
 #undef LOCTEXT_NAMESPACE
