Fix for occasional crash when undoing and redoing material function inputs deletion.

[RN] Fix for occasional crash when undoing and redoing material function inputs deletion.

diff --git a/Engine/Source/Editor/MaterialEditor/Private/MaterialEditor.cpp b/Engine/Source/Editor/MaterialEditor/Private/MaterialEditor.cpp
--- a/Engine/Source/Editor/MaterialEditor/Private/MaterialEditor.cpp
+++ b/Engine/Source/Editor/MaterialEditor/Private/MaterialEditor.cpp
@@ -7300,16 +7300,23 @@ void FMaterialEditor::PostUndo(bool bSuccess)
 {
 	if (bSuccess)
 	{	
+		// Due to the order of invocation of callbacks when undo/redo is performed, it may occur that some expression pointers in
+		// the material expression collection are null. Sanitize the list here for safety.
 		if (MaterialFunction)
 		{
+			MaterialFunction->GetExpressionCollection().RemoveExpression(nullptr);
 			Material->AssignExpressionCollection(MaterialFunction->GetExpressionCollection());
 		}
+		else
+		{
+			Material->GetExpressionCollection().RemoveExpression(nullptr);
+		}
 
 		if (TSharedPtr<SGraphEditor> FocusedGraphEd = FocusedGraphEdPtr.Pin())
 		{
 			FocusedGraphEd->ClearSelectionSet();
 		}
-		
+
 		// After an undo operation, objects material expressions reference have their back-references to the material expressions
 		// cleared (the serializer sets it to null). Loop through all material expressions in the material and fix-up these secondary references.
 		for (UMaterialExpression* MaterialExpression : Material->GetExpressions())
diff --git a/Engine/Source/Runtime/Engine/Private/Materials/HLSLMaterialTranslator.cpp b/Engine/Source/Runtime/Engine/Private/Materials/HLSLMaterialTranslator.cpp
--- a/Engine/Source/Runtime/Engine/Private/Materials/HLSLMaterialTranslator.cpp
+++ b/Engine/Source/Runtime/Engine/Private/Materials/HLSLMaterialTranslator.cpp
@@ -1465,7 +1465,7 @@ void FHLSLMaterialTranslator::TranslateMaterial()
 		for (UMaterialExpression* Expression : UMaterial->GetExpressions())
 		{
 			CHECK_DDC_QUERY_FINISHED_ELSE_RETURN();
-			if (Expression->ContainsInputLoop(VisitedExpressions))
+			if (Expression && Expression->ContainsInputLoop(VisitedExpressions))
 			{
 				AppendExpressionError(Expression, TEXT("Expression is part of a cycle. Please make sure the material graph is acyclic."));
 				TranslationResult = EHLSLMaterialTranslatorResult::Failure;
