Fixed compile errors happening when changing forceinline to inline

diff --git a/Engine/Plugins/Editor/AssetManagerEditor/Source/AssetManagerEditor/Private/ReferenceViewer/SReferenceViewer.cpp b/Engine/Plugins/Editor/AssetManagerEditor/Source/AssetManagerEditor/Private/ReferenceViewer/SReferenceViewer.cpp
--- a/Engine/Plugins/Editor/AssetManagerEditor/Source/AssetManagerEditor/Private/ReferenceViewer/SReferenceViewer.cpp
+++ b/Engine/Plugins/Editor/AssetManagerEditor/Source/AssetManagerEditor/Private/ReferenceViewer/SReferenceViewer.cpp
@@ -2340,24 +2340,22 @@ bool SReferenceViewer::HasExactlyOnePackageNodeSelected() const
 {
 	if (GraphEditorPtr.IsValid())
 	{
-		if (GraphEditorPtr->GetSelectedNodes().Num() != 1)
+		TSet<UObject*> SelectedNodes = GraphEditorPtr->GetSelectedNodes();
+		if (SelectedNodes.Num() != 1)
 		{
 			return false;
 		}
 
-		TSet<UObject*> SelectedNodes = GraphEditorPtr->GetSelectedNodes();
-		for (UObject* Node : SelectedNodes)
+		UObject* Node = *SelectedNodes.begin();
+		UEdGraphNode_Reference* ReferenceNode = Cast<UEdGraphNode_Reference>(Node);
+		if (ReferenceNode)
 		{
-			UEdGraphNode_Reference* ReferenceNode = Cast<UEdGraphNode_Reference>(Node);
-			if (ReferenceNode)
+			if (ReferenceNode->IsPackage())
 			{
-				if (ReferenceNode->IsPackage())
-				{
-					return true;
-				}
+				return true;
 			}
-			return false;
 		}
+		return false;
 	}
 
 	return false;
diff --git a/Engine/Plugins/Editor/ContentBrowser/ContentBrowserFileDataSource/Source/ContentBrowserFileDataSource/Private/ContentBrowserFileDataPayload.cpp b/Engine/Plugins/Editor/ContentBrowser/ContentBrowserFileDataSource/Source/ContentBrowserFileDataSource/Private/ContentBrowserFileDataPayload.cpp
--- a/Engine/Plugins/Editor/ContentBrowser/ContentBrowserFileDataSource/Source/ContentBrowserFileDataSource/Private/ContentBrowserFileDataPayload.cpp
+++ b/Engine/Plugins/Editor/ContentBrowser/ContentBrowserFileDataSource/Source/ContentBrowserFileDataSource/Private/ContentBrowserFileDataPayload.cpp
@@ -51,16 +51,14 @@ FText FFileConfigData::GetDiscoveryDescription() const
 {
 	if (FileActionsMap.Num() > 0)
 	{
-		for (const auto& FileActionsPair : FileActionsMap)
+		const auto& FileActionsPair = *FileActionsMap.begin();
+		if (FileActionsMap.Num() == 1)
 		{
-			if (FileActionsMap.Num() == 1)
-			{
-				return FText::Format(LOCTEXT("DiscoveringFileType", "Discovering {0} files..."), FileActionsPair.Value->TypeDisplayName);
-			}
-			else
-			{
-				return FText::Format(LOCTEXT("DiscoveringFileTypes", "Discovering {0} (and {1} other) files..."), FileActionsPair.Value->TypeDisplayName, FileActionsMap.Num() - 1);
-			}
+			return FText::Format(LOCTEXT("DiscoveringFileType", "Discovering {0} files..."), FileActionsPair.Value->TypeDisplayName);
+		}
+		else
+		{
+			return FText::Format(LOCTEXT("DiscoveringFileTypes", "Discovering {0} (and {1} other) files..."), FileActionsPair.Value->TypeDisplayName, FileActionsMap.Num() - 1);
 		}
 	}
 
diff --git a/Engine/Source/Editor/MaterialEditor/Private/MaterialEditor.cpp b/Engine/Source/Editor/MaterialEditor/Private/MaterialEditor.cpp
--- a/Engine/Source/Editor/MaterialEditor/Private/MaterialEditor.cpp
+++ b/Engine/Source/Editor/MaterialEditor/Private/MaterialEditor.cpp
@@ -5526,17 +5526,14 @@ FString FMaterialEditor::GetDocLinkForSelectedNode()
 	FGraphPanelSelectionSet SelectedNodes = GetSelectedNodes();
 	if (SelectedNodes.Num() == 1)
 	{
-		for (UObject* ObjectInSelection : SelectedNodes)
+		UObject* ObjectInSelection = *SelectedNodes.begin();
+		if (ObjectInSelection != NULL)
 		{
-			if (ObjectInSelection != NULL)
-			{
-				UMaterialGraphNode* SelectedGraphNode = Cast<UMaterialGraphNode>(ObjectInSelection);
-				FString DocLink = SelectedGraphNode->GetDocumentationLink();
-				FString DocExcerpt = SelectedGraphNode->GetDocumentationExcerptName();
+			UMaterialGraphNode* SelectedGraphNode = Cast<UMaterialGraphNode>(ObjectInSelection);
+			FString DocLink = SelectedGraphNode->GetDocumentationLink();
+			FString DocExcerpt = SelectedGraphNode->GetDocumentationExcerptName();
 
-				DocumentationLink = FEditorClassUtils::GetDocumentationLinkFromExcerpt(DocLink, DocExcerpt);
-			}
-			break;
+			DocumentationLink = FEditorClassUtils::GetDocumentationLinkFromExcerpt(DocLink, DocExcerpt);
 		}
 	}
 
@@ -5550,17 +5547,14 @@ FString FMaterialEditor::GetDocLinkBaseUrlForSelectedNode()
 	FGraphPanelSelectionSet SelectedNodes = GetSelectedNodes();
 	if (SelectedNodes.Num() == 1)
 	{
-		for (UObject* ObjectInSelection : SelectedNodes)
+		UObject* ObjectInSelection = *SelectedNodes.begin();
+		if (ObjectInSelection != NULL)
 		{
-			if (ObjectInSelection != NULL)
-			{
-				UMaterialGraphNode* SelectedGraphNode = Cast<UMaterialGraphNode>(ObjectInSelection);
-				FString DocLink = SelectedGraphNode->GetDocumentationLink();
-				FString DocExcerpt = SelectedGraphNode->GetDocumentationExcerptName();
+			UMaterialGraphNode* SelectedGraphNode = Cast<UMaterialGraphNode>(ObjectInSelection);
+			FString DocLink = SelectedGraphNode->GetDocumentationLink();
+			FString DocExcerpt = SelectedGraphNode->GetDocumentationExcerptName();
 
-				DocumentationLinkBaseUrl = FEditorClassUtils::GetDocumentationLinkBaseUrlFromExcerpt(DocLink, DocExcerpt);
-			}
-			break;
+			DocumentationLinkBaseUrl = FEditorClassUtils::GetDocumentationLinkBaseUrlFromExcerpt(DocLink, DocExcerpt);
 		}
 	}
 
