Removed now redundant "r.MaterialEnableNewHLSLGenerator" CVar. Turning on the new translator is now performed with a single cvar "r.Material.Translator.EnableNew=1".

#rnx

diff --git a/Engine/Source/Editor/MaterialEditor/Private/MaterialEditorDetailCustomization.cpp b/Engine/Source/Editor/MaterialEditor/Private/MaterialEditorDetailCustomization.cpp
--- a/Engine/Source/Editor/MaterialEditor/Private/MaterialEditorDetailCustomization.cpp
+++ b/Engine/Source/Editor/MaterialEditor/Private/MaterialEditorDetailCustomization.cpp
@@ -783,7 +783,7 @@ TSharedRef<class IDetailCustomization> FMaterialDetailCustomization::MakeInstanc
 
 void FMaterialDetailCustomization::CustomizeDetails( IDetailLayoutBuilder& DetailLayout )
 {
-	static const auto CVarMaterialEnableNewHLSLGenerator = IConsoleManager::Get().FindTConsoleVariableDataInt(TEXT("r.MaterialEnableNewHLSLGenerator"));
+	static const auto CVarMaterialTranslatorEnableNew = IConsoleManager::Get().FindTConsoleVariableDataBool(TEXT("r.Material.Translator.EnableNew"));
 
 	TArray<TWeakObjectPtr<UObject> > Objects;
 	DetailLayout.GetObjectsBeingCustomized( Objects );
@@ -855,7 +855,7 @@ void FMaterialDetailCustomization::CustomizeDetails( IDetailLayoutBuilder& Detai
 					}));
 			}
 
-			if (PropertyName == GET_MEMBER_NAME_CHECKED(UMaterial, bEnableNewHLSLGenerator) && !CVarMaterialEnableNewHLSLGenerator->GetValueOnAnyThread())
+			if (PropertyName == GET_MEMBER_NAME_CHECKED(UMaterial, bEnableNewHLSLGenerator) && !CVarMaterialTranslatorEnableNew->GetValueOnAnyThread())
 			{
 				DetailLayout.HideProperty(PropertyHandle);
 			}
@@ -940,7 +940,7 @@ TSharedRef<class IDetailCustomization> FMaterialFunctionDetailCustomization::Mak
 
 void FMaterialFunctionDetailCustomization::CustomizeDetails(IDetailLayoutBuilder& DetailLayout)
 {
-	static const auto CVarMaterialEnableNewHLSLGenerator = IConsoleManager::Get().FindTConsoleVariableDataInt(TEXT("r.MaterialEnableNewHLSLGenerator"));
+	static const auto CVarMaterialTranslatorEnableNew = IConsoleManager::Get().FindTConsoleVariableDataBool(TEXT("r.Material.Translator.EnableNew"));
 
 	// MaterialFunction category
 	{
@@ -954,7 +954,7 @@ void FMaterialFunctionDetailCustomization::CustomizeDetails(IDetailLayoutBuilder
 			FProperty* Property = PropertyHandle->GetProperty();
 			FName PropertyName = Property->GetFName();
 
-			if (PropertyName == GET_MEMBER_NAME_CHECKED(UMaterialFunction, bEnableNewHLSLGenerator) && !CVarMaterialEnableNewHLSLGenerator->GetValueOnAnyThread())
+			if (PropertyName == GET_MEMBER_NAME_CHECKED(UMaterialFunction, bEnableNewHLSLGenerator) && !CVarMaterialTranslatorEnableNew->GetValueOnAnyThread())
 			{
 				DetailLayout.HideProperty(PropertyHandle);
 			}
diff --git a/Engine/Source/Editor/MaterialEditor/Private/MaterialExpressionClasses.cpp b/Engine/Source/Editor/MaterialEditor/Private/MaterialExpressionClasses.cpp
--- a/Engine/Source/Editor/MaterialEditor/Private/MaterialExpressionClasses.cpp
+++ b/Engine/Source/Editor/MaterialEditor/Private/MaterialExpressionClasses.cpp
@@ -81,10 +81,8 @@ void MaterialExpressionClasses::InitMaterialExpressionClasses()
 {
 	if(!bInitialized)
 	{
-		static const auto CVarMaterialEnableNewHLSLGenerator = IConsoleManager::Get().FindTConsoleVariableDataInt(TEXT("r.MaterialEnableNewHLSLGenerator"));
-		
-		const bool bEnableControlFlow = AllowMaterialControlFlow();
-		const bool bEnableNewHLSLGenerator = CVarMaterialEnableNewHLSLGenerator->GetValueOnAnyThread() != 0;
+		static const auto CVarMaterialTranslatorEnableNew = IConsoleManager::Get().FindTConsoleVariableDataBool(TEXT("r.Material.Translator.EnableNew"));
+		const bool bMaterialTranslatorEnableNew = CVarMaterialTranslatorEnableNew->GetValueOnAnyThread();
 
 		UMaterialEditorOptions* TempEditorOptions = NewObject<UMaterialEditorOptions>();
 		UClass* BaseType = UMaterialExpression::StaticClass();
@@ -108,12 +106,12 @@ void MaterialExpressionClasses::InitMaterialExpressionClasses()
 						}
 
 						// Hide node types related to control flow, unless it's enabled
-						if (!bEnableControlFlow && Class->HasMetaData("MaterialControlFlow"))
+						if (Class->HasMetaData("MaterialControlFlow"))
 						{
 							continue;
 						}
 
-						if (!bEnableNewHLSLGenerator && Class->HasMetaData("MaterialNewHLSLGenerator"))
+						if (!bMaterialTranslatorEnableNew && Class->HasMetaData("MaterialNewHLSLGenerator"))
 						{
 							continue;
 						}
diff --git a/Engine/Source/Runtime/Engine/Private/Materials/MaterialExpressions.cpp b/Engine/Source/Runtime/Engine/Private/Materials/MaterialExpressions.cpp
--- a/Engine/Source/Runtime/Engine/Private/Materials/MaterialExpressions.cpp
+++ b/Engine/Source/Runtime/Engine/Private/Materials/MaterialExpressions.cpp
@@ -15054,8 +15054,8 @@ bool UMaterialFunction::IsUsingNewHLSLGenerator() const
 {
 	if (bEnableNewHLSLGenerator)
 	{
-		static const auto CVar = IConsoleManager::Get().FindTConsoleVariableDataInt(TEXT("r.MaterialEnableNewHLSLGenerator"));
-		return CVar->GetValueOnAnyThread() != 0;
+		static const auto CVar = IConsoleManager::Get().FindTConsoleVariableDataBool(TEXT("r.Material.Translator.EnableNew"));
+		return CVar->GetValueOnAnyThread();
 	}
 	return false;
 }
diff --git a/Engine/Source/Runtime/Engine/Private/Materials/MaterialInterface.cpp b/Engine/Source/Runtime/Engine/Private/Materials/MaterialInterface.cpp
--- a/Engine/Source/Runtime/Engine/Private/Materials/MaterialInterface.cpp
+++ b/Engine/Source/Runtime/Engine/Private/Materials/MaterialInterface.cpp
@@ -64,15 +64,6 @@
 
 UEnum* UMaterialInterface::SamplerTypeEnum = nullptr;
 
-static TAutoConsoleVariable<int32> CVarMaterialEnableNewHLSLGenerator(
-	TEXT("r.MaterialEnableNewHLSLGenerator"),
-	0,
-	TEXT("Enables the new (WIP) material HLSL generator.\n")
-	TEXT("0 - Don't allow\n")
-	TEXT("1 - Allow if enabled by material\n")
-	TEXT("2 - Force all materials to use new generator\n"),
-	ECVF_RenderThreadSafe | ECVF_ReadOnly);
-
 //////////////////////////////////////////////////////////////////////////
 
 bool IsHairStrandsGeometrySupported(const EShaderPlatform Platform)
@@ -290,16 +281,11 @@ const FMaterialCachedExpressionData& UMaterialInterface::GetCachedExpressionData
 
 bool UMaterialInterface::IsUsingNewHLSLGenerator() const
 {
-	const int CVarValue = CVarMaterialEnableNewHLSLGenerator.GetValueOnAnyThread();
-	if (CVarValue == 0)
+	static const auto CVar = IConsoleManager::Get().FindTConsoleVariableDataBool(TEXT("r.Material.Translator.EnableNew"));
+	if (!CVar->GetValueOnAnyThread())
 	{
 		return false;
 	}
-	else if (CVarValue == 2)
-	{
-		return true;
-	}
-
 	const UMaterial* BaseMaterial = GetMaterial_Concurrent();
 	return BaseMaterial ? BaseMaterial->bEnableNewHLSLGenerator : false;
 }
diff --git a/Engine/Source/Runtime/Engine/Public/MaterialShared.h b/Engine/Source/Runtime/Engine/Public/MaterialShared.h
--- a/Engine/Source/Runtime/Engine/Public/MaterialShared.h
+++ b/Engine/Source/Runtime/Engine/Public/MaterialShared.h
@@ -3323,10 +3323,3 @@ inline bool ShouldIncludeMaterialInDefaultOpaquePass(const FMaterial& Material)
 	return !Material.IsSky()
 		&& !Material.GetShadingModels().HasShadingModel(MSM_SingleLayerWater);
 }
-
-// r.MaterialEnableControlFlow is removed and the feature is forced disabled as how control flow should be
-// implemented in the material editor is still under discussion.
-inline bool AllowMaterialControlFlow()
-{
-	return false;
-}
