Added a Topology widget on the viewport for Material Instances. Also added a substrate panel to view details about the compiled output for our layered material.

diff --git a/Engine/Source/Editor/MaterialEditor/Private/MaterialEditorTabs.cpp b/Engine/Source/Editor/MaterialEditor/Private/MaterialEditorTabs.cpp
--- a/Engine/Source/Editor/MaterialEditor/Private/MaterialEditorTabs.cpp
+++ b/Engine/Source/Editor/MaterialEditor/Private/MaterialEditorTabs.cpp
@@ -23,3 +23,4 @@ const FName FMaterialInstanceEditorTabs::PropertiesTabId( TEXT( "MaterialInstanc
 const FName FMaterialInstanceEditorTabs::LayerPropertiesTabId(TEXT("MaterialInstanceEditor_MaterialLayerProperties"));
 const FName FMaterialInstanceEditorTabs::PreviewSettingsTabId(TEXT("MaterialInstanceEditor_PreviewSettings"));
 const FName FMaterialInstanceEditorTabs::AssetBrowserTabId(TEXT("MaterialInstanceEditor_AssetBrowser"));
+const FName FMaterialInstanceEditorTabs::SubstrateTabId(TEXT("MaterialInstanceEditor_Substrate"));
diff --git a/Engine/Source/Editor/MaterialEditor/Private/MaterialEditorTabs.h b/Engine/Source/Editor/MaterialEditor/Private/MaterialEditorTabs.h
--- a/Engine/Source/Editor/MaterialEditor/Private/MaterialEditorTabs.h
+++ b/Engine/Source/Editor/MaterialEditor/Private/MaterialEditorTabs.h
@@ -32,6 +32,7 @@ struct FMaterialInstanceEditorTabs
 	static const FName LayerPropertiesTabId;
 	static const FName PreviewSettingsTabId;
 	static const FName AssetBrowserTabId;
+	static const FName SubstrateTabId;
 
 	// Disable default constructor
 	FMaterialInstanceEditorTabs() = delete;
diff --git a/Engine/Source/Editor/MaterialEditor/Private/MaterialInstanceEditor.cpp b/Engine/Source/Editor/MaterialEditor/Private/MaterialInstanceEditor.cpp
--- a/Engine/Source/Editor/MaterialEditor/Private/MaterialInstanceEditor.cpp
+++ b/Engine/Source/Editor/MaterialEditor/Private/MaterialInstanceEditor.cpp
@@ -65,6 +65,7 @@
 #include "SparseVolumeTexture/SparseVolumeTexture.h"
 #include "Widgets/Input/SButton.h"
 #include "Subsystems/AssetEditorSubsystem.h"
+#include "SMaterialEditorSubstrateWidget.h"
 #include "ViewportToolbar/UnrealEdViewportToolbar.h"
 
 #define LOCTEXT_NAMESPACE "MaterialInstanceEditor"
@@ -270,6 +271,11 @@ void FMaterialInstanceEditor::RegisterTabSpawners(const TSharedRef<class FTabMan
 		.SetDisplayName(LOCTEXT("AssetBrowserTab", "Asset Browser"))
 		.SetGroup(WorkspaceMenuCategoryRef)
 		.SetIcon(FSlateIcon(FAppStyle::GetAppStyleSetName(), "LevelEditor.Tabs.Details"));
+
+	InTabManager->RegisterTabSpawner(FMaterialInstanceEditorTabs::SubstrateTabId, FOnSpawnTab::CreateSP(this, &FMaterialInstanceEditor::SpawnTab_Substrate))
+		.SetDisplayName(LOCTEXT("SubstrateTab", "Substrate"))
+		.SetGroup(WorkspaceMenuCategoryRef)
+		.SetIcon(FSlateIcon(FAppStyle::GetAppStyleSetName(), "Kismet.Tabs.Palette"));	// SUBSTRATE_TODO a Substrate icon
 #endif
 	MaterialStatsManager->RegisterTabs();
 
@@ -289,6 +295,8 @@ void FMaterialInstanceEditor::UnregisterTabSpawners(const TSharedRef<class FTabM
 	InTabManager->UnregisterTabSpawner( FMaterialInstanceEditorTabs::PreviewSettingsTabId );
 #if ENABLE_MATERIAL_LAYER_PROTOTYPE
 	InTabManager->UnregisterTabSpawner( FMaterialInstanceEditorTabs::AssetBrowserTabId );
+
+	InTabManager->UnregisterTabSpawner(FMaterialInstanceEditorTabs::SubstrateTabId);
 #endif
 	MaterialStatsManager->UnregisterTabs();
 
@@ -910,7 +918,9 @@ void FMaterialInstanceEditor::CreateInternalWidgets()
 		MaterialLayersFunctionsInstance = SNew(SMaterialLayersFunctionsInstanceWrapper)
 			.InMaterialEditorInstance(MaterialEditorInstance)
 			.InShowHiddenDelegate(FGetShowHiddenParameters::CreateSP(this, &FMaterialInstanceEditor::GetShowHiddenParameters));
+		MaterialLayersFunctionsInstance->OnLayerPropertyChanged.BindSP(this, &FMaterialInstanceEditor::Refresh);
 	}
+	
 	FOnGetDetailCustomizationInstance LayoutMICDetails = FOnGetDetailCustomizationInstance::CreateStatic( 
 		&FMaterialInstanceParameterDetails::MakeInstance, MaterialEditorInstance.Get(), MaterialLayersFunctionsInstance.Get(), FGetShowHiddenParameters::CreateSP(this, &FMaterialInstanceEditor::GetShowHiddenParameters) );
 	MaterialInstanceDetails->RegisterInstancedCustomPropertyLayout( UMaterialEditorInstanceConstant::StaticClass(), LayoutMICDetails );
@@ -918,7 +928,7 @@ void FMaterialInstanceEditor::CreateInternalWidgets()
 	MaterialInstanceDetails->SetCustomFilterDelegate(FSimpleDelegate::CreateSP(this, &FMaterialInstanceEditor::FilterOverriddenProperties));
 	MaterialEditorInstance->DetailsView = MaterialInstanceDetails;
 
-	
+	SubstrateWidget = SNew(SMaterialEditorSubstrateWidget, SharedThis(this));
 }
 
 void FMaterialInstanceEditor::FilterOverriddenProperties()
@@ -1173,7 +1183,21 @@ TSharedRef<SDockTab> FMaterialInstanceEditor::SpawnTab_AssetBrowser(const FSpawn
 
 	return SpawnedTab;
 }
+TSharedRef<SDockTab> FMaterialInstanceEditor::SpawnTab_Substrate(const FSpawnTabArgs& Args)
+{
+	check(Args.GetTabId() == FMaterialInstanceEditorTabs::SubstrateTabId);
+
+	TSharedRef<SDockTab> SubstrateTab = SNew(SDockTab)
+		.Label(LOCTEXT("MaterialSubstrateTabTitle", "Substrate"))
+		[
+			SNew(SBox)
+			[
+				SubstrateWidget.ToSharedRef()
+			]
+		];
 
+	return SubstrateTab;
+}
 TSharedPtr<SWidget> FMaterialInstanceEditor::OnGetAssetContextMenu(const TArray<FAssetData>& SelectedAssets) const
 {
 	if (SelectedAssets.Num() <= 0)
@@ -1782,6 +1806,11 @@ void FMaterialInstanceEditor::Refresh()
 	UpdatePropertyWindow();
 
 	RefreshOnScreenMessages();
+
+#if ENABLE_MATERIAL_LAYER_PROTOTYPE
+	if (SubstrateWidget)
+		SubstrateWidget->UpdateFromMaterial();
+#endif
 }
 
 void FMaterialInstanceEditor::PostUndo( bool bSuccess )
diff --git a/Engine/Source/Editor/MaterialEditor/Private/MaterialInstanceEditor.h b/Engine/Source/Editor/MaterialEditor/Private/MaterialInstanceEditor.h
--- a/Engine/Source/Editor/MaterialEditor/Private/MaterialInstanceEditor.h
+++ b/Engine/Source/Editor/MaterialEditor/Private/MaterialInstanceEditor.h
@@ -197,6 +197,10 @@ private:
 	
 	/** Spawns the asset browser tab */
 	TSharedRef<SDockTab> SpawnTab_AssetBrowser(const FSpawnTabArgs& Args);
+
+	/** Spawns the substrate details tab */
+	TSharedRef<SDockTab> SpawnTab_Substrate(const FSpawnTabArgs& Args);
+	
 	void OnAssetDoubleClicked(const FAssetData& AssetData);
 	TSharedPtr<SWidget> OnGetAssetContextMenu(const TArray<FAssetData>& SelectedAssets) const;
 	
@@ -271,6 +275,9 @@ private:
 	/** Set to true when editor is being destructed */
 	bool bDestructing = false;
 
+	/** The Substrate control tab */
+	TSharedPtr<class SMaterialEditorSubstrateWidget> SubstrateWidget;
+	
 	TSharedPtr<FExtensibilityManager> MenuExtensibilityManager;
 	TSharedPtr<FExtensibilityManager> ToolBarExtensibilityManager;
 
diff --git a/Engine/Source/Editor/MaterialEditor/Private/SMaterialEditorSubstrateWidget.cpp b/Engine/Source/Editor/MaterialEditor/Private/SMaterialEditorSubstrateWidget.cpp
--- a/Engine/Source/Editor/MaterialEditor/Private/SMaterialEditorSubstrateWidget.cpp
+++ b/Engine/Source/Editor/MaterialEditor/Private/SMaterialEditorSubstrateWidget.cpp
@@ -22,7 +22,7 @@
 
 DEFINE_LOG_CATEGORY_STATIC(LogMaterialEditorSubstrateWidget, Log, All);
 
-void SMaterialEditorSubstrateWidget::Construct(const FArguments& InArgs, TWeakPtr<FMaterialEditor> InMaterialEditorPtr)
+void SMaterialEditorSubstrateWidget::Construct(const FArguments& InArgs, TWeakPtr<IMaterialEditor> InMaterialEditorPtr)
 {
 	MaterialEditorPtr = InMaterialEditorPtr;
 
@@ -446,13 +446,18 @@ void SMaterialEditorSubstrateWidget::Tick(const FGeometry& AllottedGeometry, con
 	FText SubstrateMaterialDescription;
 	if (MaterialEditorPtr.IsValid())
 	{
-		auto MaterialEditor = MaterialEditorPtr.Pin();
+		TSharedPtr<IMaterialEditor> MaterialEditorPinned = MaterialEditorPtr.Pin();
+		FMaterialResource* MaterialResource = MaterialEditorPinned->GetMaterialInterface()->GetMaterialResource(GMaxRHIFeatureLevel);
 
-		UMaterial* MaterialForStats = MaterialEditor->bStatsFromPreviewMaterial ? MaterialEditor->Material : MaterialEditor->OriginalMaterial;
+		if (Cast<UMaterial>(MaterialEditorPinned->GetMaterialInterface()))
+		{
+			TSharedPtr<FMaterialEditor> MaterialEditor = StaticCastSharedPtr<FMaterialEditor>(MaterialEditorPinned);
+			UMaterial* MaterialForStats = MaterialEditor->bStatsFromPreviewMaterial ? MaterialEditor->Material : MaterialEditor->OriginalMaterial;
+			MaterialResource = MaterialForStats->GetMaterialResource(GMaxRHIFeatureLevel);
+		}
 
 		SubstrateMaterialDescription = FText::FromString(FString(TEXT("SubstrateMaterialDescription")));
 
-		const FMaterialResource* MaterialResource = MaterialForStats->GetMaterialResource(GMaxRHIFeatureLevel);
 		if (MaterialResource)
 		{
 			FString MaterialDescription;
diff --git a/Engine/Source/Editor/MaterialEditor/Private/SMaterialEditorSubstrateWidget.h b/Engine/Source/Editor/MaterialEditor/Private/SMaterialEditorSubstrateWidget.h
--- a/Engine/Source/Editor/MaterialEditor/Private/SMaterialEditorSubstrateWidget.h
+++ b/Engine/Source/Editor/MaterialEditor/Private/SMaterialEditorSubstrateWidget.h
@@ -14,7 +14,7 @@
 #include "Widgets/Text/STextBlock.h"
 #include "Widgets/Layout/SBox.h"
 
-class FMaterialEditor;
+class IMaterialEditor;
 
 class SMaterialEditorSubstrateWidget : public SCompoundWidget
 {
@@ -24,7 +24,7 @@ public:
 	}
 	SLATE_END_ARGS()
 
-	void Construct(const FArguments& InArgs, TWeakPtr<FMaterialEditor> InMaterialEditorPtr);
+	void Construct(const FArguments& InArgs, TWeakPtr<IMaterialEditor> InMaterialEditorPtr);
 
 	void UpdateFromMaterial() { bUpdateRequested = true; }
 
@@ -55,7 +55,7 @@ private:
 	FReply OnButtonApplyToPreview();
 
 	/** Pointer back to the material editor that owns this */
-	TWeakPtr<FMaterialEditor> MaterialEditorPtr;
+	TWeakPtr<IMaterialEditor> MaterialEditorPtr;
 
 	bool bUpdateRequested = true;
 
diff --git a/Engine/Source/Editor/MaterialEditor/Private/SMaterialEditorTopologyWidget.cpp b/Engine/Source/Editor/MaterialEditor/Private/SMaterialEditorTopologyWidget.cpp
new file mode 100644
--- /dev/null
+++ b/Engine/Source/Editor/MaterialEditor/Private/SMaterialEditorTopologyWidget.cpp
@@ -0,0 +1,292 @@
+// Copyright Epic Games, Inc. All Rights Reserved.
+
+#include "SMaterialEditorTopologyWidget.h"
+
+#include "EditorWidgetsModule.h"
+#include "AssetRegistry/AssetRegistryModule.h"
+#include "MaterialEditor.h"
+#include "MaterialInstanceEditor.h"
+#include "SubstrateDefinitions.h"
+#include "Rendering/SubstrateMaterialShared.h"
+#include <functional>
+
+#include "MaterialInstanceEditor.h"
+#include "Widgets/Input/SButton.h"
+#include "Widgets/Input/SCheckBox.h"
+#include "Widgets/Layout/SWrapBox.h"
+#include "Widgets/Layout/SScrollBox.h"
+#include "Widgets/Layout/SBorder.h"
+#include "Widgets/SNullWidget.h"
+#include "Widgets/Notifications/SErrorText.h"
+#include "AssetThumbnail.h"
+#include "SGraphSubstrateMaterial.h"
+#include "RHIShaderPlatform.h"
+
+#define LOCTEXT_NAMESPACE "SMaterialEditorTopologyWidget"
+
+void SMaterialEditorTopologyWidget::Construct(const FArguments& InArgs, TWeakPtr<IMaterialEditor> InMaterialEditorPtr)
+{
+	MaterialEditorPtr = InMaterialEditorPtr;
+
+
+	if (Substrate::IsSubstrateEnabled())
+	{
+		this->ChildSlot
+		[
+			SNew(SBorder)
+				.BorderImage(FAppStyle::Get().GetBrush("ActionableMessage.Border"))
+				.BorderBackgroundColor(FLinearColor(0.22f, 0.22f, 0.22f, 0.75f))
+				.Padding(FMargin(5.0f, 5.0f, 5.0f, 5.0f))
+			[
+				SNew(SVerticalBox)
+				+SVerticalBox::Slot()
+				.AutoHeight()
+				.HAlign(HAlign_Center)
+				[
+					SNew(STextBlock)
+					.ColorAndOpacity(FLinearColor::White)
+					.ShadowColorAndOpacity(FLinearColor::Black)
+					.ShadowOffset(FVector2D::UnitVector)
+					.Text(LOCTEXT("MaterialTopology", "Material Topology"))
+				]
+				+SVerticalBox::Slot()
+				.AutoHeight()
+				.Padding(0.0f, 5.0f, 0.0f, 0.0f)
+				[
+					SNew(SWrapBox)
+					.UseAllottedSize(true)
+					+ SWrapBox::Slot()
+					.Padding(10.0f)
+					.HAlign(HAlign_Center)
+					.VAlign(VAlign_Center)
+					[
+						SAssignNew(MaterialBox, SBox)
+					]
+				]
+			]
+		];
+	}
+	else
+	{
+		this->ChildSlot
+		[
+			SNew(SVerticalBox)
+			+SVerticalBox::Slot()
+			.AutoHeight()
+			.Padding(0.0f, 5.0f, 0.0f, 0.0f)
+			[
+				SNew(SWrapBox)
+				.UseAllottedSize(true)
+				+SWrapBox::Slot()
+				.Padding(5.0f)
+				.HAlign(HAlign_Center)
+				.VAlign(VAlign_Center)
+				[
+					SNew(STextBlock)
+					.ColorAndOpacity(FLinearColor::Yellow)
+					.ShadowColorAndOpacity(FLinearColor::Black)
+					.ShadowOffset(FVector2D::UnitVector)
+					.Text(LOCTEXT("SubstrateWidgetNotEnable", "Details cannot be shown: Substrate (Beta) is not enabled for this project (See the project settings window, rendering settings section)."))
+				]
+			]
+		];
+	}
+}
+
+TSharedRef<SWidget> SMaterialEditorTopologyWidget::GetContent()
+{
+	return SharedThis(this);
+}
+
+SMaterialEditorTopologyWidget::~SMaterialEditorTopologyWidget()
+{
+}
+
+
+void SMaterialEditorTopologyWidget::Tick(const FGeometry& AllottedGeometry, const double InCurrentTime, const float InDeltaTime)
+{
+	if (!bUpdateRequested || !Substrate::IsSubstrateEnabled())
+	{
+		return;
+	}
+	bUpdateRequested = false;
+
+	FText SubstrateMaterialDescription;
+	if (MaterialEditorPtr.IsValid())
+	{
+		TSharedPtr<IMaterialEditor> MaterialEditorPinned = MaterialEditorPtr.Pin();
+		FMaterialResource* MaterialResource = MaterialEditorPinned->GetMaterialInterface()->GetMaterialResource(GMaxRHIFeatureLevel);
+
+		if (Cast<UMaterial>(MaterialEditorPinned->GetMaterialInterface()))
+		{
+			TSharedPtr<FMaterialEditor> MaterialEditor = StaticCastSharedPtr<FMaterialEditor>(MaterialEditorPinned);
+			UMaterial* MaterialForStats = MaterialEditor->bStatsFromPreviewMaterial ? MaterialEditor->Material : MaterialEditor->OriginalMaterial;
+			MaterialResource = MaterialForStats->GetMaterialResource(GMaxRHIFeatureLevel);
+		}
+
+		SubstrateMaterialDescription = FText::FromString(FString(TEXT("SubstrateMaterialDescription")));
+		
+		if (MaterialResource)
+		{
+			FString MaterialDescription;
+			FString MaterialBudget;
+
+			bool bMaterialOutOfBudgetHasBeenSimplified = false;
+			FMaterialShaderMap* ShaderMap = MaterialResource->GetGameThreadShaderMap();
+			if (ShaderMap)
+			{
+				const FSubstrateMaterialCompilationOutput& CompilationOutput = ShaderMap->GetSubstrateMaterialCompilationOutput();
+				const uint32 FinalPixelByteCount = CompilationOutput.SubstrateUintPerPixel * sizeof(uint32);
+				const uint32 FinalPixelClosureCount = CompilationOutput.SubstrateClosureCount;
+				bMaterialOutOfBudgetHasBeenSimplified = CompilationOutput.bMaterialOutOfBudgetHasBeenSimplified > 0;
+
+				// Now generate a visual representation of the material from the topology tree of operators.
+				if (CompilationOutput.RootOperatorIndex >= 0)
+				{
+					const FMaterialLayersFunctions* LayersFunctions = MaterialResource->GetMaterialLayers();
+					MaterialBox->SetContent(ProcessOperatorAsThumbnails(CompilationOutput, LayersFunctions));
+				}
+				else
+				{
+					// The tree does not looks sane so generate a visual error without crashing.
+					auto TreeError = SNew(SErrorText)
+						.ErrorText(LOCTEXT("TreeError", "Tree Error"))
+						.BackgroundColor(FSlateColor(EStyleColor::AccentRed));
+					const TSharedRef<SWidget>& TreeErrorAsShared = TreeError->AsShared();
+					MaterialBox->SetContent(TreeErrorAsShared);
+				}
+			}
+			else
+			{
+				MaterialDescription = TEXT("Shader map not found.");
+				MaterialBox->SetContent(SNullWidget::NullWidget);
+			}
+
+		}
+		
+	}
+}
+
+static const TSharedRef<SWidget> InternalProcessOperatorAsThumbnails(
+	const FSubstrateMaterialCompilationOutput& CompilationOutput, 
+	const FSubstrateOperator& Op, 
+	const TArray<FGuid>& InGuid,
+	const FMaterialLayersFunctions* LayersFunctions,
+	EStyleColor OverrideColor)
+{
+	const bool bIsCurrent = false;
+	const EStyleColor Color0 = OverrideColor;
+	const EStyleColor Color1 = OverrideColor;
+	switch (Op.OperatorType)
+	{
+		case SUBSTRATE_OPERATOR_WEIGHT:
+		{
+			const EStyleColor Color = bIsCurrent ? EStyleColor::AccentGreen : OverrideColor;
+			return InternalProcessOperatorAsThumbnails(CompilationOutput, CompilationOutput.Operators[Op.LeftIndex], InGuid, LayersFunctions, Color);
+		}
+		break;
+		case SUBSTRATE_OPERATOR_VERTICAL:
+		{
+			auto VerticalOperator = SNew(SVerticalBox)
+				+SVerticalBox::Slot()
+				.AutoHeight()
+				.VAlign(VAlign_Fill)
+				.HAlign(HAlign_Fill)
+				[
+					InternalProcessOperatorAsThumbnails(CompilationOutput, CompilationOutput.Operators[Op.LeftIndex], InGuid, LayersFunctions, Color0)
+				]
+				+ SVerticalBox::Slot()
+				.AutoHeight()
+				.VAlign(VAlign_Fill)
+				.HAlign(HAlign_Fill)
+				[
+					InternalProcessOperatorAsThumbnails(CompilationOutput, CompilationOutput.Operators[Op.RightIndex], InGuid, LayersFunctions, Color1)
+				];
+			return VerticalOperator->AsShared();
+		}
+		break;
+		case SUBSTRATE_OPERATOR_HORIZONTAL:
+		case SUBSTRATE_OPERATOR_SELECT:
+		{
+			auto HorizontalOperator = SNew(SHorizontalBox)
+				+SHorizontalBox::Slot()
+				.AutoWidth()
+				.VAlign(VAlign_Fill)
+				.HAlign(HAlign_Fill)
+				[
+					InternalProcessOperatorAsThumbnails(CompilationOutput, CompilationOutput.Operators[Op.LeftIndex], InGuid, LayersFunctions, Color0)
+				]
+				+SHorizontalBox::Slot()
+				.AutoWidth()
+				.VAlign(VAlign_Fill)
+				.HAlign(HAlign_Fill)
+				[
+					InternalProcessOperatorAsThumbnails(CompilationOutput, CompilationOutput.Operators[Op.RightIndex], InGuid, LayersFunctions, Color1)
+				];
+			return HorizontalOperator->AsShared();
+		}
+		break;
+		case SUBSTRATE_OPERATOR_ADD:
+		{
+			auto HorizontalOperator = SNew(SHorizontalBox)
+				+SHorizontalBox::Slot()
+				.AutoWidth()
+				.VAlign(VAlign_Fill)
+				.HAlign(HAlign_Fill)
+				[
+					InternalProcessOperatorAsThumbnails(CompilationOutput, CompilationOutput.Operators[Op.LeftIndex], InGuid, LayersFunctions, Color0)
+				]
+				+SHorizontalBox::Slot()
+				.AutoWidth()
+				.VAlign(VAlign_Fill)
+				.HAlign(HAlign_Fill)
+				[
+					InternalProcessOperatorAsThumbnails(CompilationOutput, CompilationOutput.Operators[Op.RightIndex], InGuid, LayersFunctions, Color1)
+				];
+			return HorizontalOperator->AsShared();
+		}
+		break;
+		case SUBSTRATE_OPERATOR_BSDF_LEGACY:	// legacy BSDF should have been converted to BSDF already.
+		case SUBSTRATE_OPERATOR_BSDF:
+			{
+				int32 ExpressionIndex = INDEX_NONE;
+				const float ThumbnailSize = 40.0f;
+				UObject* ThumbnailObject = nullptr;
+#if ENABLE_MATERIAL_LAYER_PROTOTYPE
+				const TSharedPtr<FMaterialLayersFunctionsRuntimeGraphCache> LayerTreeCache = LayersFunctions->RuntimeGraphCache;
+				if (LayerTreeCache)
+				{
+					ExpressionIndex = LayerTreeCache->FindExpressionIndex(Op.MaterialExpressionGuid, LayersFunctions);
+				}
+				if (ExpressionIndex >= 0)
+				{
+					ThumbnailObject = LayerTreeCache->NodePreviewMaterials[ExpressionIndex];
+				}
+#endif
+				const TSharedPtr<FAssetThumbnail> AssetThumbnail = MakeShareable(new FAssetThumbnail(ThumbnailObject, ThumbnailSize, ThumbnailSize, UThumbnailManager::Get().GetSharedThumbnailPool()));
+				TSharedRef<SWidget> ThumbnailWidget = SNew(SBorder)
+					.Padding(2.0f)
+					[
+						AssetThumbnail->MakeThumbnailWidget()
+					];
+			AssetThumbnail->SetRealTime(true);
+			//	ThumbnailWidget->SetOnMouseDoubleClick(FPointerEventHandler::CreateSP(this, &SMaterialSubstrateTree::OnThumbnailDoubleClick, InAssociation, InIndex));
+			return ThumbnailWidget;
+		}
+		break;
+	}
+
+	static FString NoVisualization = FString(TEXT("Tree Operator Error"));
+	auto TreeOperatorError = SNew(SErrorText)
+		.ErrorText(FText::FromString(NoVisualization))
+		.BackgroundColor(FSlateColor(EStyleColor::AccentRed));
+	return TreeOperatorError->AsShared();
+}
+
+const TSharedRef<SWidget> SMaterialEditorTopologyWidget::ProcessOperatorAsThumbnails(const FSubstrateMaterialCompilationOutput& CompilationOutput, const FMaterialLayersFunctions* LayersFunctions)
+{
+	return InternalProcessOperatorAsThumbnails(CompilationOutput, CompilationOutput.Operators[CompilationOutput.RootOperatorIndex], TArray<FGuid>(), LayersFunctions, EStyleColor::MAX);
+}
+
+
+#undef LOCTEXT_NAMESPACE
diff --git a/Engine/Source/Editor/MaterialEditor/Private/SMaterialEditorTopologyWidget.h b/Engine/Source/Editor/MaterialEditor/Private/SMaterialEditorTopologyWidget.h
new file mode 100644
--- /dev/null
+++ b/Engine/Source/Editor/MaterialEditor/Private/SMaterialEditorTopologyWidget.h
@@ -0,0 +1,52 @@
+// Copyright Epic Games, Inc. All Rights Reserved.
+
+#pragma once
+
+#include "CoreMinimal.h"
+#include "MaterialEditor/MaterialEditorInstanceConstant.h"
+#include "Widgets/SWidget.h"
+
+#include "Widgets/DeclarativeSyntaxSupport.h"
+#include "Widgets/SCompoundWidget.h"
+#include "Widgets/SBoxPanel.h"
+#include "Widgets/Input/SButton.h"
+#include "Widgets/Input/SCheckBox.h"
+#include "Widgets/Input/SNumericEntryBox.h"
+#include "Widgets/Text/STextBlock.h"
+#include "Widgets/Layout/SBox.h"
+
+class IMaterialEditor;
+
+class SMaterialEditorTopologyWidget : public SCompoundWidget
+{
+public:
+	SLATE_BEGIN_ARGS(SMaterialEditorTopologyWidget)
+	{
+	}
+	SLATE_END_ARGS()
+
+	void Construct(const FArguments& InArgs, TWeakPtr<IMaterialEditor> InMaterialEditorPtr);
+
+	void UpdateFromMaterial() { bUpdateRequested = true; }
+
+	/** Gets the widget contents of the app */
+	virtual TSharedRef<SWidget> GetContent();
+
+	virtual ~SMaterialEditorTopologyWidget();
+
+	/** SWidget interface */
+	virtual void Tick(const FGeometry& AllottedGeometry, const double InCurrentTime, const float InDeltaTime) override;
+
+	static const TSharedRef<SWidget> ProcessOperatorAsThumbnails(const struct FSubstrateMaterialCompilationOutput& CompilationOutput, const struct FMaterialLayersFunctions* LayersFunctions);
+	
+private:
+
+	TSharedPtr<class SBox> MaterialBox;
+
+	/** Pointer back to the material editor that owns this */
+	TWeakPtr<IMaterialEditor> MaterialEditorPtr;
+
+	bool bUpdateRequested = true;
+
+};
+
diff --git a/Engine/Source/Editor/MaterialEditor/Private/SMaterialEditorViewport.cpp b/Engine/Source/Editor/MaterialEditor/Private/SMaterialEditorViewport.cpp
--- a/Engine/Source/Editor/MaterialEditor/Private/SMaterialEditorViewport.cpp
+++ b/Engine/Source/Editor/MaterialEditor/Private/SMaterialEditorViewport.cpp
@@ -51,6 +51,7 @@
 #define LOCTEXT_NAMESPACE "MaterialEditor"
 #include "MaterialEditorTabs.h"
 #include "PreviewProfileController.h"
+#include "SMaterialEditorTopologyWidget.h"
 #include "UnrealWidget.h"
 
 /** Viewport Client for the preview viewport */
@@ -343,6 +344,11 @@ void SMaterialEditor3DPreviewViewport::RefreshViewport()
 			}
 		}
 	}
+
+	// Also request to update the Substrate slab.
+	if (SubstrateWidget)
+		SubstrateWidget->UpdateFromMaterial();
+
 }
 
 bool SMaterialEditor3DPreviewViewport::SetPreviewAsset(UObject* InAsset)
@@ -1086,6 +1092,22 @@ void SMaterialEditor3DPreviewViewport::PopulateViewportOverlays(TSharedRef<class
 		];
 		// clang-format on
 
+#if ENABLE_MATERIAL_LAYER_PROTOTYPE
+	// only show topology widget for material instances
+	if (Cast<UMaterialInstance>(MaterialEditorPtr.Pin()->GetMaterialInterface()))
+	{
+		Overlay->AddSlot()
+			// clang-format off
+			.VAlign(VAlign_Top)
+			.HAlign(HAlign_Right)
+			.Padding(FMargin(5.0f, 5.0f, 50.0f, 5.0f))
+			[
+				SAssignNew(SubstrateWidget, SMaterialEditorTopologyWidget, MaterialEditorPtr)
+			];
+		// clang-format on
+	}
+#endif
+	// clang-format on
 	// add the feature level display widget
 	Overlay->AddSlot()
 		// clang-format off
diff --git a/Engine/Source/Editor/MaterialEditor/Private/SMaterialEditorViewport.h b/Engine/Source/Editor/MaterialEditor/Private/SMaterialEditorViewport.h
--- a/Engine/Source/Editor/MaterialEditor/Private/SMaterialEditorViewport.h
+++ b/Engine/Source/Editor/MaterialEditor/Private/SMaterialEditorViewport.h
@@ -108,6 +108,9 @@ private:
 	/** The parent tab where this viewport resides */
 	TWeakPtr<SDockTab> ParentTab;
 
+	/** The Substrate control tab */
+	TSharedPtr<class SMaterialEditorTopologyWidget> SubstrateWidget;
+
 	bool IsVisible() const override;
 
 	/** Pointer back to the material editor tool that owns us */
diff --git a/Engine/Source/Editor/MaterialEditor/Private/SMaterialSubstrateTree.cpp b/Engine/Source/Editor/MaterialEditor/Private/SMaterialSubstrateTree.cpp
--- a/Engine/Source/Editor/MaterialEditor/Private/SMaterialSubstrateTree.cpp
+++ b/Engine/Source/Editor/MaterialEditor/Private/SMaterialSubstrateTree.cpp
@@ -121,18 +121,17 @@ FSlateColor SMaterialSubstrateTreeItem::GetOuterBackgroundColor(TSharedPtr<FSort
 	return FAppStyle::Get().GetSlateColor("Colors.Panel");
 }
 
-void SMaterialSubstrateTreeItem::RefreshOnRowChange(const FAssetData& AssetData, SMaterialSubstrateTree* InTree)
+void SMaterialSubstrateTreeItem::RefreshMaterialViews()
 {
-	if (SMaterialLayersFunctionsInstanceWrapper* Wrapper = InTree->GetWrapper())
+	if (SMaterialLayersFunctionsInstanceWrapper* Wrapper = Tree->GetWrapper())
 	{
+		Tree->CreateGroupsWidget();
+		Tree->RequestTreeRefresh();
+		
 		if (Wrapper->OnLayerPropertyChanged.IsBound())
 		{
 			Wrapper->OnLayerPropertyChanged.Execute();
 		}
-		else
-		{
-			InTree->CreateGroupsWidget();
-		}
 	}
 }
 
@@ -372,8 +371,7 @@ FReply SMaterialSubstrateTreeItem::OnLayerDrop(const FDragDropEvent& DragDropEve
 					
 					Tree->FunctionInstance->MoveLayerNode(SourceNodeId, DestParent, DestSiblingIdx, bShouldDuplicate);
 					Tree->FunctionInstanceHandle->NotifyPostChange(EPropertyChangeType::ValueSet);
-					Tree->CreateGroupsWidget();
-					Tree->RequestTreeRefresh();
+					RefreshMaterialViews();
 				}
 			}
 			return FReply::Handled();
@@ -471,6 +469,7 @@ FReply SMaterialSubstrateTreeItem::OnLayerDrop(const FDragDropEvent& DragDropEve
 
 			if (DidModifyTree)
 			{
+				RefreshMaterialViews();
 				return FReply::Handled();
 			}
 		}
@@ -644,8 +643,7 @@ void SMaterialSubstrateTreeItem::Construct(const FArguments& InArgs, const TShar
 							AssetChild->ParameterHandle->SetValue(ObjPath.TryLoad());
 
 							Tree->FunctionInstanceHandle->NotifyPostChange(EPropertyChangeType::ValueSet);
-							Tree->CreateGroupsWidget();
-							Tree->RequestTreeRefresh();
+                			RefreshMaterialViews();
                 		})
 						.OnShouldFilterAsset_Static(&FilterAssetDataAMaterialLayerBlendFunction)
                 		.AllowClear(true)
@@ -1106,9 +1104,7 @@ void SMaterialSubstrateTree::RefreshOnAssetChange(const struct FAssetData& InAss
 	
 	//set their overrides back to 0
 	MaterialEditorInstance->CleanParameterStack(Index, MaterialType);
-	CreateGroupsWidget();
 	MaterialEditorInstance->ResetOverrides(Index, MaterialType);
-	RequestTreeRefresh();
 }
 
 void SMaterialSubstrateTree::ResetAssetToDefault(TSharedPtr<FSortedParamData> InData)
@@ -1323,6 +1319,16 @@ void SMaterialSubstrateTree::RecursiveCreateWidgets(FRecursiveCreateWidgetsConte
 				{
 					MaterialEditorInstance->StoredLayerPreviews[Payload.Layer]->SetParentEditorOnly(EditedMaterial);
 				}
+				UMaterialInterface* MaterialInterface = MaterialEditorInstance->GetMaterialInterface();
+				const FMaterialResource* MaterialResource = MaterialInterface->GetMaterialResource(ERHIFeatureLevel::SM6);
+				const FMaterialLayersFunctions* LayersFunctions = MaterialResource->GetMaterialLayers();
+#if ENABLE_MATERIAL_LAYER_PROTOTYPE
+				TSharedPtr<FMaterialLayersFunctionsRuntimeGraphCache> LayerTreeCache = LayersFunctions->RuntimeGraphCache;
+				if (LayerTreeCache)
+				{
+					LayerTreeCache->NodePreviewMaterials[InNodeId] = (UMaterial*)EditedMaterial;
+				}
+#endif
 			}
 #if ENABLE_MATERIAL_LAYER_PROTOTYPE
 			else 
diff --git a/Engine/Source/Editor/MaterialEditor/Private/SMaterialSubstrateTree.h b/Engine/Source/Editor/MaterialEditor/Private/SMaterialSubstrateTree.h
--- a/Engine/Source/Editor/MaterialEditor/Private/SMaterialSubstrateTree.h
+++ b/Engine/Source/Editor/MaterialEditor/Private/SMaterialSubstrateTree.h
@@ -163,7 +163,7 @@ private:
 	FSlateColor GetOuterBackgroundColor(TSharedPtr<FSortedParamData> InParamData) const;
 public:
 
-	void RefreshOnRowChange(const FAssetData& AssetData, SMaterialSubstrateTree* InTree);
+	void RefreshMaterialViews();
 	bool GetFilterState(SMaterialSubstrateTree* InTree, TSharedPtr<FSortedParamData> InStackData) const;
 	void FilterClicked(const ECheckBoxState NewCheckedState, SMaterialSubstrateTree* InTree, TSharedPtr<FSortedParamData> InStackData);
 	ECheckBoxState GetFilterChecked(SMaterialSubstrateTree* InTree, TSharedPtr<FSortedParamData> InStackData) const;
diff --git a/Engine/Source/Runtime/Engine/Private/Materials/HLSLMaterialTranslator.cpp b/Engine/Source/Runtime/Engine/Private/Materials/HLSLMaterialTranslator.cpp
--- a/Engine/Source/Runtime/Engine/Private/Materials/HLSLMaterialTranslator.cpp
+++ b/Engine/Source/Runtime/Engine/Private/Materials/HLSLMaterialTranslator.cpp
@@ -16269,7 +16269,7 @@ void FHLSLMaterialTranslator::PrepareEnvironmentDefines()
 						{
 							if (!It.IsDiscarded() && It.MaxDistanceFromLeaves == DistanceToLeaves)
 							{
-								SubstrateMaterialContextDescription += FString::Printf(TEXT("\tIdx=%d Op=%s ParentIdx=%d LeftIndex=%d RightIndex=%d BSDFIdx=%d LayerDepth=%d IsTop=%d IsBot=%d BSDFType=%s SSS=%d MFP=%d F90=%d Rough2=%d Fuzz=%d Aniso=%d Glint=%d SpecularProfile=%d Eye=%d Hair=%d\n"),
+								SubstrateMaterialContextDescription += FString::Printf(TEXT("\tIdx=%d Op=%s ParentIdx=%d LeftIndex=%d RightIndex=%d BSDFIdx=%d LayerDepth=%d IsTop=%d IsBot=%d BSDFType=%s SSS=%d MFP=%d F90=%d Rough2=%d Fuzz=%d Aniso=%d Glint=%d SpecularProfile=%d Eye=%d Hair=%d\n GUID=%s\n"),
 									It.Index, GetSubstrateOperatorStr(It.OperatorType), It.ParentIndex, It.LeftIndex, It.RightIndex, It.BSDFIndex, It.LayerDepth, It.bIsTop, It.bIsBottom,
 									*GetSubstrateBSDFName(It.BSDFType), 
 									It.Has(ESubstrateBsdfFeature_SSS), 
@@ -16281,7 +16281,8 @@ void FHLSLMaterialTranslator::PrepareEnvironmentDefines()
 									It.Has(ESubstrateBsdfFeature_Glint), 
 									It.Has(ESubstrateBsdfFeature_SpecularProfile), 
 									It.Has(ESubstrateBsdfFeature_Eye), 
-									It.Has(ESubstrateBsdfFeature_Hair));
+									It.Has(ESubstrateBsdfFeature_Hair),
+									*It.MaterialExpressionGuid.ToString());
 							}
 						}
 					}
diff --git a/Engine/Source/Runtime/Engine/Private/Materials/MaterialExpressions.cpp b/Engine/Source/Runtime/Engine/Private/Materials/MaterialExpressions.cpp
--- a/Engine/Source/Runtime/Engine/Private/Materials/MaterialExpressions.cpp
+++ b/Engine/Source/Runtime/Engine/Private/Materials/MaterialExpressions.cpp
@@ -7256,9 +7256,10 @@ void UMaterialExpressionMaterialAttributeLayers::RebuildLayerGraph(bool bReportE
 
 		LayersTreeCache->NodeMaterialGraphExpressions.Empty(LayersTree->Tree.Nodes.Num());
 		LayersTreeCache->NodeMaterialGraphExpressions.InsertZeroed(0, LayersTree->Tree.Nodes.Num());
+		LayersTreeCache->NodeMaterialGraphExpressionsBlends.Empty(LayersTree->Tree.Nodes.Num());
+		LayersTreeCache->NodeMaterialGraphExpressionsBlends.InsertZeroed(0, LayersTree->Tree.Nodes.Num());
 		LayersTreeCache->NodePreviewMaterials.Empty(LayersTree->Tree.Nodes.Num());
 		LayersTreeCache->NodePreviewMaterials.InsertZeroed(0, LayersTree->Tree.Nodes.Num());
-
 		struct TreeGraphBuilder
 		{
 			UMaterialExpressionMaterialAttributeLayers* This = nullptr;
@@ -7292,6 +7293,8 @@ void UMaterialExpressionMaterialAttributeLayers::RebuildLayerGraph(bool bReportE
 					UMaterialExpressionMaterialFunctionCall* CurrentBlendCaller = (TreeCache->BlendCallers.IsValidIndex(BlendFuncIdx) ? TreeCache->BlendCallers[BlendFuncIdx] : nullptr);
 					CurrentBlendCaller = (CurrentBlendCaller && CurrentBlendCaller->MaterialFunction ? CurrentBlendCaller : nullptr);
 
+					TreeCache->NodeMaterialGraphExpressionsBlends[NodeId] = CurrentBlendCaller;
+					
 					if (CurrentBlendCaller) // Blend is the new current layer if valid and connect the previous layer and current layer is valid
 					{
 						if (PreviousLayerInput)
@@ -15814,7 +15817,70 @@ int FMaterialLayersFunctionsTree::RemoveNodeBranch(FNode* BranchParentNode, TArr
 
 	return NumRemoved;
 }
+#if WITH_EDITOR
+int32 FMaterialLayersFunctionsRuntimeGraphCache::FMaterialLayersFunctionsRuntimeGraphCache::FindExpressionIndex(FGuid InExpressionGuid, const FMaterialLayersFunctions* Layersfunctions) const
+{
+	TArray<int32> NodeIndicesLayer; 
+	TArray<int32> NodeIndicesBlends;
+
+	GetNodeIndices(InExpressionGuid, NodeMaterialGraphExpressions, NodeIndicesLayer);
+	GetNodeIndices(InExpressionGuid, NodeMaterialGraphExpressionsBlends, NodeIndicesBlends);
 
+	int32 NodeIndex = INDEX_NONE;
+
+	const int32 NumLayers = NodeIndicesLayer.Num();
+	for(int32 LayerIndex = NumLayers - 1; LayerIndex >= 0; LayerIndex--)
+	{
+		int32 CurrentNodeIndex = NodeIndicesLayer[LayerIndex];
+		if (Layersfunctions->GetNodeChildren(CurrentNodeIndex).IsEmpty())
+		{
+			NodeIndex = CurrentNodeIndex;
+			break;
+		}
+	}
+	if (NodeIndex == INDEX_NONE)
+	{
+		const int32 NumBlends = NodeIndicesBlends.Num();
+		for(int32 BlendIndex = NumBlends - 1; BlendIndex >= 0; BlendIndex--)
+		{
+			int32 CurrentNodeIndex = NodeIndicesBlends[BlendIndex];
+			if (Layersfunctions->GetNodeChildren(CurrentNodeIndex).IsEmpty())
+			{
+				NodeIndex = CurrentNodeIndex;
+				break;
+			}
+		}
+	}
+	return NodeIndex;
+}
+
+void FMaterialLayersFunctionsRuntimeGraphCache::FMaterialLayersFunctionsRuntimeGraphCache::GetNodeIndices(FGuid InExpressionGuid,
+	const TArray<TObjectPtr<UMaterialExpressionMaterialFunctionCall>>& ExpressionArrayToUse, TArray<int32>& OutIndices) const
+{
+	const int32 NumGraphExpressions = ExpressionArrayToUse.Num();
+	for(int32 MaterialExpressionIndex = NumGraphExpressions - 1; MaterialExpressionIndex >= 0; MaterialExpressionIndex--)
+	{
+		TObjectPtr<UMaterialExpressionMaterialFunctionCall> FunctionCall = ExpressionArrayToUse[MaterialExpressionIndex];
+		if (FunctionCall)
+		{
+			const uint32 OutputCount = FunctionCall->FunctionOutputs.Num();
+			for (uint32 OutputIndex = 0; OutputIndex < OutputCount; ++OutputIndex)
+			{
+				if (FunctionCall->IsResultSubstrateMaterial(OutputIndex))
+				{
+					FSubstrateMaterialInfo SubstrateMaterialInfo(true/*bGatherGuids*/);
+					FunctionCall->GatherSubstrateMaterialInfo(SubstrateMaterialInfo, OutputIndex);
+					TArray<FGuid> Guids = SubstrateMaterialInfo.GetGuids();
+    
+					if (Guids.Contains(InExpressionGuid))
+						OutIndices.Add(MaterialExpressionIndex);
+				}
+			}
+		}
+	}
+}
+#endif
+	
 std::pair<TArray<FMaterialLayersFunctionsTree::FNodeId>, TArray<FMaterialLayersFunctionsTree::FPayload>> FMaterialLayersFunctionsTree::RemoveNode(FNodeId RemovedNodeId)
 {
 	TArray<FNodeId> RemovedIds;
diff --git a/Engine/Source/Runtime/Engine/Public/Materials/MaterialLayersFunctions.h b/Engine/Source/Runtime/Engine/Public/Materials/MaterialLayersFunctions.h
--- a/Engine/Source/Runtime/Engine/Public/Materials/MaterialLayersFunctions.h
+++ b/Engine/Source/Runtime/Engine/Public/Materials/MaterialLayersFunctions.h
@@ -5,6 +5,7 @@
 #include "CoreMinimal.h"
 #include "MaterialTypes.h"
 #include "MaterialExpression.h"
+#include "Materials/MaterialExpressionMaterialFunctionCall.h"
 
 #include "MaterialLayersFunctions.generated.h"
 
@@ -390,7 +391,11 @@ struct FMaterialLayersFunctionsRuntimeGraphCache
 	FMaterialExpressionCollection								ExpressionCollection;
 
 	TArray<TObjectPtr<UMaterialExpressionMaterialFunctionCall>> NodeMaterialGraphExpressions;
+	TArray<TObjectPtr<UMaterialExpressionMaterialFunctionCall>> NodeMaterialGraphExpressionsBlends;
 	TArray<TObjectPtr<UMaterial>>								NodePreviewMaterials;
+
+	ENGINE_API void GetNodeIndices(FGuid InExpressionGuid, const TArray<TObjectPtr<UMaterialExpressionMaterialFunctionCall>>& ExpressionArrayToUse, TArray<int32>& OutIndices) const;
+	ENGINE_API int32 FindExpressionIndex(FGuid InExpressionGuid, const FMaterialLayersFunctions* Layersfunctions) const;
 #endif
 };
 
