[MaterialEditor] Fix material stats visualization in preview viewport and give hard coded X-position a name.

#rnx

diff --git a/Engine/Source/Editor/MaterialEditor/Private/MaterialEditor.cpp b/Engine/Source/Editor/MaterialEditor/Private/MaterialEditor.cpp
--- a/Engine/Source/Editor/MaterialEditor/Private/MaterialEditor.cpp
+++ b/Engine/Source/Editor/MaterialEditor/Private/MaterialEditor.cpp
@@ -2436,7 +2436,9 @@ void FMaterialEditor::DrawMaterialInfoStrings(
 
 	// The font to use when displaying info strings
 	UFont* FontToUse = GEngine->GetTinyFont();
-	const int32 SpacingBetweenLines = 13;
+
+	constexpr int32 DrawPositionX = 5;
+	constexpr int32 SpacingBetweenLines = 13;
 
 	if (bDrawInstructions)
 	{
@@ -2447,8 +2449,7 @@ void FMaterialEditor::DrawMaterialInfoStrings(
 		for (int32 InstructionIndex = 0; InstructionIndex < Descriptions.Num(); InstructionIndex++)
 		{
 			FString InstructionCountString = FString::Printf(TEXT("%s: %u instructions\nStats: %s"), *Descriptions[InstructionIndex].ShaderDescription, Descriptions[InstructionIndex].InstructionCount, *Descriptions[InstructionIndex].ShaderStatisticsString);
-			Canvas->DrawShadowedString(5, DrawPositionY, *InstructionCountString, FontToUse, FLinearColor(1, 1, 0));
-			DrawPositionY += SpacingBetweenLines;
+			DrawPositionY += Canvas->DrawShadowedString(DrawPositionX, DrawPositionY, *InstructionCountString, FontToUse, FLinearColor(1, 1, 0));
 		}
 
 		// Display the number of texture samplers and samplers used by the material.
@@ -2459,7 +2460,7 @@ void FMaterialEditor::DrawMaterialInfoStrings(
 			int32 MaxSamplers = GetExpectedFeatureLevelMaxTextureSamplers(MaterialResource->GetFeatureLevel());
 
 			Canvas->DrawShadowedString(
-				5,
+				DrawPositionX,
 				DrawPositionY,
 				*FString::Printf(TEXT("%s samplers: %u/%u"), FeatureLevel <= ERHIFeatureLevel::ES3_1 ? TEXT("Mobile texture") : TEXT("Texture"), SamplersUsed, MaxSamplers),
 				FontToUse,
@@ -2474,7 +2475,7 @@ void FMaterialEditor::DrawMaterialInfoStrings(
 		if (NumVSTextureSamples > 0 || NumPSTextureSamples > 0)
 		{
 			Canvas->DrawShadowedString(
-				5,
+				DrawPositionX,
 				DrawPositionY,
 				*FString::Printf(TEXT("Texture Lookups (Est.): VS(%u), PS(%u)"), NumVSTextureSamples, NumPSTextureSamples),
 				FontToUse,
@@ -2487,7 +2488,7 @@ void FMaterialEditor::DrawMaterialInfoStrings(
 		if (NumVirtualTextureLookups > 0)
 		{
 			Canvas->DrawShadowedString(
-				5,
+				DrawPositionX,
 				DrawPositionY,
 				*FString::Printf(TEXT("Virtual Texture Lookups (Est.): %u"), NumVirtualTextureLookups),
 				FontToUse,
@@ -2500,7 +2501,7 @@ void FMaterialEditor::DrawMaterialInfoStrings(
 		if (NumVirtualTextureStacks > 0)
 		{
 			Canvas->DrawShadowedString(
-				5,
+				DrawPositionX,
 				DrawPositionY,
 				*FString::Printf(TEXT("Virtual Texture Stacks: %u"), NumVirtualTextureStacks),
 				FontToUse,
@@ -2519,7 +2520,7 @@ void FMaterialEditor::DrawMaterialInfoStrings(
 			if (Usages > 0)
 			{
 				Canvas->DrawShadowedString(
-					5,
+					DrawPositionX,
 					DrawPositionY,
 					*FString::Printf(
 						TEXT("LWC %s usages (Est.): %u (VS), %u (PS), %u (CS)"), 
@@ -2546,7 +2547,7 @@ void FMaterialEditor::DrawMaterialInfoStrings(
 			if (NumShaders)
 			{
 				FString ShaderCountString = FString::Printf(TEXT("Num shaders added: %i"), NumShaders);
-				Canvas->DrawShadowedString(5, DrawPositionY, *ShaderCountString, FontToUse, FLinearColor(1, 0.8, 0));
+				Canvas->DrawShadowedString(DrawPositionX, DrawPositionY, *ShaderCountString, FontToUse, FLinearColor(1, 0.8, 0));
 				DrawPositionY += SpacingBetweenLines;
 			}
 		}
@@ -2554,7 +2555,7 @@ void FMaterialEditor::DrawMaterialInfoStrings(
 
 	for(int32 ErrorIndex = 0; ErrorIndex < CompileErrors.Num(); ErrorIndex++)
 	{
-		Canvas->DrawShadowedString(5, DrawPositionY, *FString::Printf(TEXT("[%s] %s"), *FeatureLevelName, *CompileErrors[ErrorIndex]), FontToUse, FLinearColor(1, 0, 0));
+		Canvas->DrawShadowedString(DrawPositionX, DrawPositionY, *FString::Printf(TEXT("[%s] %s"), *FeatureLevelName, *CompileErrors[ErrorIndex]), FontToUse, FLinearColor(1, 0, 0));
 		DrawPositionY += SpacingBetweenLines;
 	}
 }
@@ -3289,10 +3290,10 @@ void FMaterialEditor::UpdateMaterialinfoList_Old()
 
 			for (int32 InstructionIndex = 0; InstructionIndex < Results.Num(); InstructionIndex++)
 			{
-					FString InstructionCountString = FString::Printf(TEXT("%s: %u instructions\nStats: %s"), *Results[InstructionIndex].ShaderDescription, Results[InstructionIndex].InstructionCount, *Results[InstructionIndex].ShaderStatisticsString);
+				FString InstructionCountString = FString::Printf(TEXT("%s: %u instructions\nStats: %s"), *Results[InstructionIndex].ShaderDescription, Results[InstructionIndex].InstructionCount, *Results[InstructionIndex].ShaderStatisticsString);
 				if (bBuiltinStats)
 				{
-						InstructionCountString += FString::Printf(TEXT(" - Built-in instructions: %u\nStats: %s"), EmptyMaterialResults[InstructionIndex].InstructionCount, *EmptyMaterialResults[InstructionIndex].ShaderStatisticsString);
+					InstructionCountString += FString::Printf(TEXT(" - Built-in instructions: %u\nStats: %s"), EmptyMaterialResults[InstructionIndex].InstructionCount, *EmptyMaterialResults[InstructionIndex].ShaderStatisticsString);
 				}
 				TempMaterialInfoList.Add(MakeShareable(new FMaterialInfo(InstructionCountString, FLinearColor::Yellow)));
 				TSharedRef<FTokenizedMessage> Line = FTokenizedMessage::Create(EMessageSeverity::Info);
