Add material instance override support for SpecularProfile when Substrate is enabled.

diff --git a/Engine/Source/Editor/DetailCustomizations/Private/MaterialInstanceDynamicDetails.cpp b/Engine/Source/Editor/DetailCustomizations/Private/MaterialInstanceDynamicDetails.cpp
--- a/Engine/Source/Editor/DetailCustomizations/Private/MaterialInstanceDynamicDetails.cpp
+++ b/Engine/Source/Editor/DetailCustomizations/Private/MaterialInstanceDynamicDetails.cpp
@@ -37,6 +37,8 @@ void FMaterialInstanceDynamicDetails::CustomizeDetails(IDetailLayoutBuilder& Det
 	PhysMatHandle->MarkHiddenByCustomization();
 	TSharedRef<IPropertyHandle> OverrideSubSurfaceHandle = DetailLayout.GetProperty(GET_MEMBER_NAME_CHECKED(UMaterialInstance, bOverrideSubsurfaceProfile), UMaterialInstance::StaticClass());
 	OverrideSubSurfaceHandle->MarkHiddenByCustomization();
+	TSharedRef<IPropertyHandle> OverrideSpecularProfileHandle = DetailLayout.GetProperty(GET_MEMBER_NAME_CHECKED(UMaterialInstance, bOverrideSpecularProfile), UMaterialInstance::StaticClass());
+	OverrideSpecularProfileHandle->MarkHiddenByCustomization();
 	TSharedRef<IPropertyHandle> BasePropertiesHandle = DetailLayout.GetProperty(GET_MEMBER_NAME_CHECKED(UMaterialInstance, BasePropertyOverrides), UMaterialInstance::StaticClass());
 	BasePropertiesHandle->MarkHiddenByCustomization();
 
diff --git a/Engine/Source/Editor/MaterialEditor/Private/MaterialEditorInstanceDetailCustomization.cpp b/Engine/Source/Editor/MaterialEditor/Private/MaterialEditorInstanceDetailCustomization.cpp
--- a/Engine/Source/Editor/MaterialEditor/Private/MaterialEditorInstanceDetailCustomization.cpp
+++ b/Engine/Source/Editor/MaterialEditor/Private/MaterialEditorInstanceDetailCustomization.cpp
@@ -165,6 +165,8 @@ void FMaterialInstanceParameterDetails::CustomizeDetails(IDetailLayoutBuilder& D
 		DetailLayout.HideProperty("RefractionDepthBias");
 		DetailLayout.HideProperty("bOverrideSubsurfaceProfile");
 		DetailLayout.HideProperty("SubsurfaceProfile");
+		DetailLayout.HideProperty("bOverrideSpecularProfile");
+		DetailLayout.HideProperty("SpecularProfile");
 		DetailLayout.HideProperty("BasePropertyOverrides");
 		DetailLayout.HideProperty("MaterialLayersParameterValues");
 	}
@@ -237,6 +239,8 @@ void FMaterialInstanceParameterDetails::CustomizeDetails(IDetailLayoutBuilder& D
 			DetailLayout.HideProperty("RefractionDepthBias");
 			DetailLayout.HideProperty("bOverrideSubsurfaceProfile");
 			DetailLayout.HideProperty("SubsurfaceProfile");
+			DetailLayout.HideProperty("bOverrideSpecularProfile");
+			DetailLayout.HideProperty("SpecularProfile");
 			DetailLayout.HideProperty("BasePropertyOverrides");
 		}
 		else
@@ -315,20 +319,40 @@ void FMaterialInstanceParameterDetails::CustomizeDetails(IDetailLayoutBuilder& D
 				// Hide the originals, these will be recreated manually
 				DetailLayout.HideProperty("bOverrideSubsurfaceProfile");
 				DetailLayout.HideProperty("SubsurfaceProfile");
+				DetailLayout.HideProperty("bOverrideSpecularProfile");
+				DetailLayout.HideProperty("SpecularProfile");
 				DetailLayout.HideProperty("BasePropertyOverrides");
 
 				// Set up the override logic for the subsurface profile
-				TAttribute<bool> IsParamEnabled = TAttribute<bool>::Create(TAttribute<bool>::FGetter::CreateLambda([this](){ return (bool)MaterialEditorInstance->bOverrideSubsurfaceProfile; }));
+				{
+					TAttribute<bool> IsParamEnabled = TAttribute<bool>::Create(TAttribute<bool>::FGetter::CreateLambda([this](){ return (bool)MaterialEditorInstance->bOverrideSubsurfaceProfile; }));
+
+					IDetailPropertyRow& PropertyRow = MaterialPropertyOverrideGroup.AddPropertyRow(DetailLayout.GetProperty("SubsurfaceProfile"));
+					PropertyRow
+						.EditCondition(IsParamEnabled, 
+							FOnBooleanValueChanged::CreateLambda([this](bool NewValue) {
+								MaterialEditorInstance->bOverrideSubsurfaceProfile = (uint32)NewValue;
+								MaterialEditorInstance->PostEditChange();
+								FEditorSupportDelegates::RedrawAllViewports.Broadcast();
+						}))
+						.Visibility(TAttribute<EVisibility>::Create(TAttribute<EVisibility>::FGetter::CreateSP(this, &FMaterialInstanceParameterDetails::ShouldShowSubsurfaceProfile)));
+				}
+
+				// Set up the override logic for the specular profile
+				if (Substrate::IsSubstrateEnabled())
+				{
+					TAttribute<bool> IsParamEnabled = TAttribute<bool>::Create(TAttribute<bool>::FGetter::CreateLambda([this](){ return (bool)MaterialEditorInstance->bOverrideSpecularProfile; }));
 
-				IDetailPropertyRow& PropertyRow = MaterialPropertyOverrideGroup.AddPropertyRow(DetailLayout.GetProperty("SubsurfaceProfile"));
-				PropertyRow
+					IDetailPropertyRow& PropertyRow = MaterialPropertyOverrideGroup.AddPropertyRow(DetailLayout.GetProperty("SpecularProfile"));
+					PropertyRow
 					.EditCondition(IsParamEnabled, 
-						FOnBooleanValueChanged::CreateLambda([this](bool NewValue) {
-							MaterialEditorInstance->bOverrideSubsurfaceProfile = (uint32)NewValue;
-							MaterialEditorInstance->PostEditChange();
-							FEditorSupportDelegates::RedrawAllViewports.Broadcast();
+								   FOnBooleanValueChanged::CreateLambda([this](bool NewValue) {
+									   MaterialEditorInstance->bOverrideSpecularProfile = (uint32)NewValue;
+									   MaterialEditorInstance->PostEditChange();
+									   FEditorSupportDelegates::RedrawAllViewports.Broadcast();
 					}))
-					.Visibility(TAttribute<EVisibility>::Create(TAttribute<EVisibility>::FGetter::CreateSP(this, &FMaterialInstanceParameterDetails::ShouldShowSubsurfaceProfile)));
+					.Visibility(TAttribute<EVisibility>::Create(TAttribute<EVisibility>::FGetter::CreateSP(this, &FMaterialInstanceParameterDetails::ShouldShowSpecularProfile)));
+				}
 			
 				// Append the base property overrides to the Material Property Override Group
 				CreateBasePropertyOverrideWidgets(DetailLayout, MaterialPropertyOverrideGroup);
@@ -1102,6 +1126,12 @@ EVisibility FMaterialInstanceParameterDetails::ShouldShowSubsurfaceProfile() con
 	return UseSubsurfaceProfile(ShadingModels) ? EVisibility::Visible : EVisibility::Collapsed;
 }
 
+EVisibility FMaterialInstanceParameterDetails::ShouldShowSpecularProfile() const
+{
+	const bool bHasProfiles = Substrate::IsSubstrateEnabled() && MaterialEditorInstance->Parent && MaterialEditorInstance->Parent->SpecularProfiles.Num() > 0;
+	return bHasProfiles ? EVisibility::Visible : EVisibility::Collapsed;
+}
+
 void FMaterialInstanceParameterDetails::OnCopyParameterValues(int32 ParameterGroupIndex)
 {
 	if (!MaterialEditorInstance || !MaterialEditorInstance->ParameterGroups.IsValidIndex(ParameterGroupIndex))
diff --git a/Engine/Source/Editor/MaterialEditor/Private/MaterialEditorInstanceDetailCustomization.h b/Engine/Source/Editor/MaterialEditor/Private/MaterialEditorInstanceDetailCustomization.h
--- a/Engine/Source/Editor/MaterialEditor/Private/MaterialEditorInstanceDetailCustomization.h
+++ b/Engine/Source/Editor/MaterialEditor/Private/MaterialEditorInstanceDetailCustomization.h
@@ -78,6 +78,9 @@ private:
 	/** Returns true if the Subsurface Profile options should be displayed */
 	EVisibility ShouldShowSubsurfaceProfile() const;
 
+	/** Returns true if the Specular Profile options should be displayed */
+	EVisibility ShouldShowSpecularProfile() const;
+
 	//Functions supporting copy/paste of entire parameter groups.
 
 	/**
diff --git a/Engine/Source/Editor/UnrealEd/Classes/MaterialEditor/MaterialEditorInstanceConstant.h b/Engine/Source/Editor/UnrealEd/Classes/MaterialEditor/MaterialEditorInstanceConstant.h
--- a/Engine/Source/Editor/UnrealEd/Classes/MaterialEditor/MaterialEditorInstanceConstant.h
+++ b/Engine/Source/Editor/UnrealEd/Classes/MaterialEditor/MaterialEditorInstanceConstant.h
@@ -324,10 +324,18 @@ private:
 	UPROPERTY(EditAnywhere, BlueprintReadOnly, Category = Material, meta = (DisplayName = "Subsurface Profile"))
 	TObjectPtr<class USubsurfaceProfile> SubsurfaceProfile;
 
-	/** Defines if SubsurfaceProfile from tis instance is used or it uses the parent one. */
+	/** Defines if SubsurfaceProfile from this instance is used or it uses the parent one. */
 	UPROPERTY(EditAnywhere, Category = MaterialEditorInstanceConstant)
 	uint32 bOverrideSubsurfaceProfile : 1;
 
+	/** Specular profile */
+	UPROPERTY(EditAnywhere, BlueprintReadOnly, Category = Material, meta = (DisplayName = "Specular Profile"))
+	TObjectPtr<class USpecularProfile> SpecularProfile;
+
+	/** Defines if SpecularProfile from this instance is used or it uses the parent one. */
+	UPROPERTY(EditAnywhere, Category = MaterialEditorInstanceConstant)
+	uint32 bOverrideSpecularProfile : 1;
+
 	UPROPERTY()
 	uint32 bOverrideBaseProperties_DEPRECATED : 1;
 
diff --git a/Engine/Source/Editor/UnrealEd/Private/PreviewMaterial.cpp b/Engine/Source/Editor/UnrealEd/Private/PreviewMaterial.cpp
--- a/Engine/Source/Editor/UnrealEd/Private/PreviewMaterial.cpp
+++ b/Engine/Source/Editor/UnrealEd/Private/PreviewMaterial.cpp
@@ -37,6 +37,7 @@
 #include "MaterialPropertyHelpers.h"
 #include "MaterialStatsCommon.h"
 #include "MaterialDomain.h"
+#include "RenderUtils.h"
 
 /**
  * Class for rendering the material on the preview mesh in the Material Editor
@@ -983,6 +984,9 @@ void UMaterialEditorInstanceConstant::CopyToSourceInstance(const bool bForceStat
 		SourceInstance->bOverrideSubsurfaceProfile = bOverrideSubsurfaceProfile;
 		SourceInstance->SubsurfaceProfile = SubsurfaceProfile;
 
+		SourceInstance->bOverrideSpecularProfile = bOverrideSpecularProfile;
+		SourceInstance->SpecularProfileOverride = SpecularProfile;
+
 		// Update object references and parameter names.
 		SourceInstance->UpdateParameterNames();
 		VisibleExpressions.Empty();
@@ -1169,6 +1173,15 @@ void UMaterialEditorInstanceConstant::CopyBasePropertiesFromParent()
 	// Copy the subsurface profile. GetSubsurfaceProfile_Internal() will return either the overridden profile or one from a parent
 	SubsurfaceProfile = SourceInstance->GetSubsurfaceProfile_Internal();
 
+	if (Substrate::IsSubstrateEnabled())
+	{
+		bOverrideSpecularProfile = SourceInstance->bOverrideSpecularProfile;
+		if (SourceInstance->NumSpecularProfile_Internal() > 0)
+		{
+			SpecularProfile = SourceInstance->GetSpecularProfileOverride_Internal();
+		}
+	}
+
 	// Post process blendable location and priority overrides
 	PostProcessOverrides.bOverrideBlendableLocation = SourceInstance->bOverrideBlendableLocation;
 	if (SourceInstance->bOverrideBlendableLocation)
diff --git a/Engine/Source/Runtime/Engine/Private/Materials/Material.cpp b/Engine/Source/Runtime/Engine/Private/Materials/Material.cpp
--- a/Engine/Source/Runtime/Engine/Private/Materials/Material.cpp
+++ b/Engine/Source/Runtime/Engine/Private/Materials/Material.cpp
@@ -468,7 +468,9 @@ public:
 				{
 					if (ParameterInfo.Name == SpecularProfile::GetSpecularProfileParameterName(GetSpecularProfileRT(It)))
 					{
-						OutValue = SpecularProfile::GetSpecularProfileId(GetSpecularProfileRT(It));
+						const USpecularProfile* SpecularProfile = GetSpecularProfileRT(It);
+						const USpecularProfile* SpecularProfileOverride = GetSpecularProfileOverrideRT();
+						OutValue = SpecularProfile::GetSpecularProfileId(SpecularProfileOverride ? SpecularProfileOverride : SpecularProfile);
 						return true;
 					}
 				}
diff --git a/Engine/Source/Runtime/Engine/Private/Materials/MaterialInstance.cpp b/Engine/Source/Runtime/Engine/Private/Materials/MaterialInstance.cpp
--- a/Engine/Source/Runtime/Engine/Private/Materials/MaterialInstance.cpp
+++ b/Engine/Source/Runtime/Engine/Private/Materials/MaterialInstance.cpp
@@ -322,6 +322,7 @@ bool FMaterialInstanceResource::GetParameterValue(EMaterialParameterType Type, c
 	}
 	else if (Type == EMaterialParameterType::Scalar && NumSpecularProfileRT() > 0)
 	{
+		const USpecularProfile* SPOverrideRT = GetSpecularProfileOverrideRT();
 		for (uint32 It=0,Count=NumSpecularProfileRT();It<Count;++It)
 		{
 			if (ParameterInfo.Name == SpecularProfile::GetSpecularProfileParameterName(GetSpecularProfileRT(It)))
@@ -329,7 +330,7 @@ bool FMaterialInstanceResource::GetParameterValue(EMaterialParameterType Type, c
 #if !ENABLE_MATERIAL_LAYER_PROTOTYPE
 				check(ParameterInfo.Association == EMaterialParameterAssociation::GlobalParameter);
 #endif
-				OutValue = SpecularProfile::GetSpecularProfileId(GetSpecularProfileRT(It));
+				OutValue = SpecularProfile::GetSpecularProfileId(SPOverrideRT ? SPOverrideRT : GetSpecularProfileRT(It));
 				bResult = true;
 				break;
 			}
@@ -4952,6 +4953,26 @@ USubsurfaceProfile* UMaterialInstance::GetSubsurfaceProfileOverride_Internal() c
 	return Parent ? Parent->GetSubsurfaceProfileOverride_Internal() : 0;
 }
 
+uint32 UMaterialInstance::NumSpecularProfile_Internal() const
+{
+	// Return the subsurface profile count form the root material.
+	checkSlow(IsInGameThread());
+	return Parent ? Parent->NumSpecularProfile_Internal() : 0;
+}
+
+USpecularProfile* UMaterialInstance::GetSpecularProfile_Internal(uint32 Index) const
+{
+	// Return the Specular profile from the root material.
+	checkSlow(IsInGameThread());
+
+	return Parent ? Parent->GetSpecularProfile_Internal(Index) : 0;
+}
+
+USpecularProfile* UMaterialInstance::GetSpecularProfileOverride_Internal() const
+{
+	return bOverrideSpecularProfile ? SpecularProfileOverride : nullptr;
+}
+
 bool UMaterialInstance::CastsRayTracedShadows() const
 {
 	//#dxr_todo: do per material instance override?
@@ -5166,6 +5187,10 @@ bool UMaterialInstance::IsRedundant() const
 	{
 		return false;
 	}
+	if (Substrate::IsSubstrateEnabled() && NumSpecularProfile_Internal() > 0 && GetSpecularProfile_Internal(0) != Parent->GetSpecularProfile_Internal(0))
+	{
+		return false;
+	}
 	// Assume that if any properties are overridden they are different to their parent
 	if (
 	   BasePropertyOverrides.bOverride_OpacityMaskClipValue 
diff --git a/Engine/Source/Runtime/Engine/Private/Materials/MaterialInterface.cpp b/Engine/Source/Runtime/Engine/Private/Materials/MaterialInterface.cpp
--- a/Engine/Source/Runtime/Engine/Private/Materials/MaterialInterface.cpp
+++ b/Engine/Source/Runtime/Engine/Private/Materials/MaterialInterface.cpp
@@ -1487,6 +1487,11 @@ USpecularProfile* UMaterialInterface::GetSpecularProfile_Internal(uint32 Index)
 	return nullptr;
 }
 
+USpecularProfile* UMaterialInterface::GetSpecularProfileOverride_Internal() const
+{
+	return nullptr;
+}
+
 UNeuralProfile* UMaterialInterface::GetNeuralProfile_Internal() const
 {
 	return nullptr;
@@ -1647,10 +1652,11 @@ void UMaterialInterface::UpdateMaterialRenderProxy(FMaterialRenderProxy& Proxy)
 				}
 			}
 
-			if (Entries.Num() > 0)
+			USpecularProfile* SpecularProfileOverride = GetSpecularProfileOverride_Internal();
+			if (Entries.Num() > 0 || SpecularProfileOverride)
 			{
 				ENQUEUE_RENDER_COMMAND(UpdateMaterialRenderProxySpecular)(
-				[InProxy, Entries](FRHICommandListImmediate& RHICmdList)
+				[InProxy, Entries, SpecularProfileOverride](FRHICommandListImmediate& RHICmdList)
 				{
 					for (const FEntry& Entry : Entries)
 					{
@@ -1661,6 +1667,20 @@ void UMaterialInterface::UpdateMaterialRenderProxy(FMaterialRenderProxy& Proxy)
 						}
 						InProxy->AddSpecularProfileRT(Entry.Profile);
 					}
+
+					if (SpecularProfileOverride)
+					{
+						const FTextureReference* SPOverrideTexture = nullptr;
+						if (!SpecularProfileOverride->Settings.IsProcedural() && SpecularProfileOverride->Settings.Texture)
+						{
+							SPOverrideTexture = &SpecularProfileOverride->Settings.Texture->TextureReference;
+						}
+
+						SpecularProfile::AddOrUpdateProfile(SpecularProfileOverride, SpecularProfileOverride->Guid, SpecularProfileOverride->Settings, SPOverrideTexture);
+					}
+
+					// Set the specular profile override using the default SpecularProfile pointer.
+					InProxy->SetSpecularProfileOverrideRT(SpecularProfileOverride);
 				});
 			}
 		}
diff --git a/Engine/Source/Runtime/Engine/Public/Materials/MaterialInstance.h b/Engine/Source/Runtime/Engine/Public/Materials/MaterialInstance.h
--- a/Engine/Source/Runtime/Engine/Public/Materials/MaterialInstance.h
+++ b/Engine/Source/Runtime/Engine/Public/Materials/MaterialInstance.h
@@ -625,6 +625,10 @@ class UMaterialInstance : public UMaterialInterface
 
 	const FMaterialInstanceCachedData& GetCachedInstanceData() const { return CachedData ? *CachedData : FMaterialInstanceCachedData::EmptyData; }
 
+	/** Defines SpecularProfile override used by this instance. If not set, it uses the parent ones. */
+	UPROPERTY()
+	TObjectPtr<class USpecularProfile> SpecularProfileOverride;
+
 	/**
 	 * Indicates whether the instance has static permutation resources (which are required when static parameters are present) 
 	 * Read directly from the rendering thread, can only be modified with the use of a FMaterialUpdateContext.
@@ -637,6 +641,10 @@ class UMaterialInstance : public UMaterialInterface
 	UPROPERTY(EditAnywhere, BlueprintReadOnly, Category = MaterialInstance)
 	uint8 bOverrideSubsurfaceProfile:1;
 
+	/** Defines if SpecularProfile from this instance is used or it uses the parent one. */
+	UPROPERTY(EditAnywhere, BlueprintReadOnly, Category = MaterialInstance)
+	uint8 bOverrideSpecularProfile:1;
+
 	/** For post process materials, use BlendableLocationOverride. */
 	UPROPERTY(EditAnywhere, BlueprintReadWrite, Category = MaterialInstance)
 	uint8 bOverrideBlendableLocation:1;
@@ -911,6 +919,9 @@ public:
 	ENGINE_API uint32 NumSubsurfaceProfileRoot_Internal() const override;
 	ENGINE_API USubsurfaceProfile* GetSubsurfaceProfileRoot_Internal(uint32 Index) const override;
 	ENGINE_API virtual USubsurfaceProfile* GetSubsurfaceProfileOverride_Internal() const override;
+	ENGINE_API virtual uint32 NumSpecularProfile_Internal() const override;
+	ENGINE_API virtual USpecularProfile* GetSpecularProfile_Internal(uint32 Index) const override;
+	ENGINE_API virtual USpecularProfile* GetSpecularProfileOverride_Internal() const override;
 	ENGINE_API virtual bool CastsRayTracedShadows() const override;
 	ENGINE_API virtual bool IsTessellationEnabled() const override;
 	ENGINE_API virtual bool HasSubstrateRoughnessTracking() const override;
diff --git a/Engine/Source/Runtime/Engine/Public/Materials/MaterialInterface.h b/Engine/Source/Runtime/Engine/Public/Materials/MaterialInterface.h
--- a/Engine/Source/Runtime/Engine/Public/Materials/MaterialInterface.h
+++ b/Engine/Source/Runtime/Engine/Public/Materials/MaterialInterface.h
@@ -1004,6 +1004,7 @@ public:
 	ENGINE_API virtual USubsurfaceProfile* GetSubsurfaceProfileOverride_Internal() const;
 	ENGINE_API virtual uint32 NumSpecularProfile_Internal() const;
 	ENGINE_API virtual USpecularProfile* GetSpecularProfile_Internal(uint32 Index) const;
+	ENGINE_API virtual USpecularProfile* GetSpecularProfileOverride_Internal() const;
 	ENGINE_API virtual UNeuralProfile* GetNeuralProfile_Internal() const;
 	ENGINE_API virtual bool CastsRayTracedShadows() const;
 	ENGINE_API virtual bool IsTessellationEnabled() const;
diff --git a/Engine/Source/Runtime/Engine/Public/Materials/MaterialRenderProxy.h b/Engine/Source/Runtime/Engine/Public/Materials/MaterialRenderProxy.h
--- a/Engine/Source/Runtime/Engine/Public/Materials/MaterialRenderProxy.h
+++ b/Engine/Source/Runtime/Engine/Public/Materials/MaterialRenderProxy.h
@@ -218,7 +218,7 @@ public:
 	}
 #endif
 
-	// Subsurface profile.
+	// Subsurface profile
 	// When Substrate is enabled, this is ONLY used as an override for Subsurface Profile on material instance (override all Subsurface Profiles at once for now)
 	void SetSubsurfaceProfileRT(const USubsurfaceProfile* Ptr) { SubsurfaceProfileRT = Ptr; }
 	const USubsurfaceProfile* GetSubsurfaceProfileRT() const { return SubsurfaceProfileRT; }
@@ -232,6 +232,9 @@ public:
 	void AddSpecularProfileRT(const USpecularProfile* Ptr) { SpecularProfilesRT.Add(Ptr); }
 	const USpecularProfile* GetSpecularProfileRT(uint32 Index) const { check(Index<uint32(SpecularProfilesRT.Num())); return SpecularProfilesRT[Index]; }
 	const uint32 NumSpecularProfileRT() const { return SpecularProfilesRT.Num(); }
+	// Specular profiles override
+	void SetSpecularProfileOverrideRT(const USpecularProfile* Ptr) { SpecularProfileOverrideRT = Ptr; }
+	const USpecularProfile* GetSpecularProfileOverrideRT() const { return SpecularProfileOverrideRT; }
 
 	// Neural profiles
 	void SetNeuralProfileRT(const UNeuralProfile* Ptr) { NeuralProfileRT = Ptr; }
@@ -257,6 +260,7 @@ private:
 	const USubsurfaceProfile* SubsurfaceProfileRT;	// Overrides all SubsurfaceProfilesRT when set, used when set from material instance to respect the legacy workflow
 	TArray<const USubsurfaceProfile*> SubsurfaceProfilesRT;
 	TArray<const USpecularProfile*> SpecularProfilesRT;
+	const USpecularProfile* SpecularProfileOverrideRT = nullptr;
 	const UNeuralProfile* NeuralProfileRT;
 	FString MaterialName;
 
