Combine the different wrappers in the Material Editors in prep for removing ENABLE_MATERIAL_LAYER_PROTOTYPE

diff --git a/Engine/Source/Editor/MaterialEditor/Private/MaterialEditor.cpp b/Engine/Source/Editor/MaterialEditor/Private/MaterialEditor.cpp
--- a/Engine/Source/Editor/MaterialEditor/Private/MaterialEditor.cpp
+++ b/Engine/Source/Editor/MaterialEditor/Private/MaterialEditor.cpp
@@ -1796,15 +1796,20 @@ void FMaterialEditor::CreateInternalWidgets()
 		.InGenerator(Generator)
 		.InMaterialDetailsView(MaterialDetailsView);
 
-#if ENABLE_MATERIAL_LAYER_PROTOTYPE
-	MaterialLayersFunctionsInstance = SNew(SMaterialLayersFunctionsInstanceWrapper)
-		.InMaterialEditorInstance(MaterialEditorInstance)
-		.InGenerator(Generator);
-#else
-	MaterialLayersFunctionsInstance = SNew(SMaterialLayersFunctionsMaterialWrapper)
+	
+	if (GBEnableMaterialLayerPrototype)
+	{
+		MaterialLayersFunctionsInstance = SNew(SMaterialLayersFunctionsInstanceWrapper)
 			.InMaterialEditorInstance(MaterialEditorInstance)
 			.InGenerator(Generator);
-#endif
+	}
+	else
+	{
+		MaterialLayersFunctionsInstance = SNew(SMaterialLayersFunctionsMaterialWrapper)
+			.InMaterialEditorInstance(MaterialEditorInstance)
+			.InGenerator(Generator);
+	}
+
 	Palette = SNew(SMaterialPalette, SharedThis(this));
 
 	FMessageLogModule& MessageLogModule = FModuleManager::LoadModuleChecked<FMessageLogModule>("MessageLog");
diff --git a/Engine/Source/Editor/MaterialEditor/Private/MaterialEditor.h b/Engine/Source/Editor/MaterialEditor/Private/MaterialEditor.h
--- a/Engine/Source/Editor/MaterialEditor/Private/MaterialEditor.h
+++ b/Engine/Source/Editor/MaterialEditor/Private/MaterialEditor.h
@@ -1017,11 +1017,8 @@ private:
 	TSharedPtr<class SMaterialCustomPrimitiveDataPanel> MaterialCustomPrimitiveDataWidget;
 
 	/** Layer Properties View */
-#if ENABLE_MATERIAL_LAYER_PROTOTYPE
-	TSharedPtr<class SMaterialLayersFunctionsInstanceWrapper> MaterialLayersFunctionsInstance;
-#else
-	TSharedPtr<class SMaterialLayersFunctionsMaterialWrapper> MaterialLayersFunctionsInstance;
-#endif
+	TSharedPtr<class SMaterialLayersFunctionsWrapper> MaterialLayersFunctionsInstance;
+	
 	/** The current transaction. */
 	FScopedTransaction* ScopedTransaction;
 
diff --git a/Engine/Source/Editor/MaterialEditor/Private/SMaterialLayersFunctionsTree.cpp b/Engine/Source/Editor/MaterialEditor/Private/SMaterialLayersFunctionsTree.cpp
--- a/Engine/Source/Editor/MaterialEditor/Private/SMaterialLayersFunctionsTree.cpp
+++ b/Engine/Source/Editor/MaterialEditor/Private/SMaterialLayersFunctionsTree.cpp
@@ -1784,87 +1784,96 @@ void SMaterialLayersFunctionsInstanceWrapper::Refresh()
 	LayerParameter.Reset();
 	TSharedPtr<SHorizontalBox> HeaderBox;
 	NestedTree->CreateGroupsWidget();
-	LayerParameter = NestedTree->FunctionParameter;
+	LayerParameter = NestedTree->GetFunctionParameter();
 	FOnClicked 	OnChildButtonClicked = FOnClicked::CreateStatic(&FMaterialPropertyHelpers::OnClickedSaveNewMaterialInstance, MaterialEditorInstance->GetMaterialInterface(), ImplicitConv<UObject*>(MaterialEditorInstance));
 	FOnClicked	OnSiblingButtonClicked = FOnClicked::CreateStatic(&FMaterialPropertyHelpers::OnClickedSaveNewMaterialInstance, MaterialEditorInstance->GetParentMaterialInterface(), ImplicitConv<UObject*>(MaterialEditorInstance));
 
-#if ENABLE_MATERIAL_LAYER_PROTOTYPE
-	const float ThumbnailSize = 48.0f;
-	TSharedPtr<SBox> ThumbnailBox;
-	UObject* ThumbnailObject = MaterialEditorInstance->GetMaterialInterface().Get();
-	const TSharedPtr<FAssetThumbnail> AssetThumbnail = MakeShareable(new FAssetThumbnail(ThumbnailObject, ThumbnailSize, ThumbnailSize, NestedTree->GetTreeThumbnailPool()));
-	TSharedRef<SWidget> ThumbnailWidget = AssetThumbnail->MakeThumbnailWidget();
-	FText MaterialName = FText::FromName(MaterialEditorInstance->GetMaterialInterface()->GetFName());
-	
-#endif
+	constexpr float ThumbnailSize = 48.0f;
+	TSharedRef<SWidget> ThumbnailWidget = SNullWidget::NullWidget;
+	TSharedPtr<SBox> ThumbnailBox = nullptr;
+	FText MaterialName = FText::GetEmpty();
+	if (GBEnableMaterialLayerPrototype)
+	{
+		UObject* ThumbnailObject = MaterialEditorInstance->GetMaterialInterface().Get();
+		const TSharedPtr<FAssetThumbnail> AssetThumbnail = MakeShareable(new FAssetThumbnail(ThumbnailObject, ThumbnailSize, ThumbnailSize, NestedTree->GetTreeThumbnailPool()));
+		ThumbnailWidget = AssetThumbnail->MakeThumbnailWidget();
+		MaterialName = FText::FromName(MaterialEditorInstance->GetMaterialInterface()->GetFName());
+	}
 	
 	if (LayerParameter != nullptr)
 	{
-#if !ENABLE_MATERIAL_LAYER_PROTOTYPE
-		FOnClicked OnRelinkToParent = FOnClicked::CreateSP(NestedTree.ToSharedRef(), &SMaterialLayersFunctionsInstanceTree::RelinkLayersToParent);
-#endif
 		
-		this->ChildSlot
-			[
-				SNew(SVerticalBox)
-				+ SVerticalBox::Slot()
-				.Padding(FMargin(0.0f, 0.0f, 0.0f, 10.0f))
-				.AutoHeight()
-				[
-#if ENABLE_MATERIAL_LAYER_PROTOTYPE
-					SAssignNew(HeaderBox, SHorizontalBox)
-					+ SHorizontalBox::Slot()
-					.AutoWidth()
-							.HAlign(HAlign_Center)
-							.VAlign(VAlign_Center)
-							.Padding(4.0f)
-							.MaxWidth(ThumbnailSize)
-							[
-								SAssignNew(ThumbnailBox, SBox)
-								.MaxDesiredWidth(ThumbnailSize)
-								.MaxDesiredWidth(ThumbnailSize)
-								.MaxDesiredHeight(ThumbnailSize)
-								.MinDesiredHeight(ThumbnailSize)
-								[
-									ThumbnailWidget
-								]
-							]
-
-					+ SHorizontalBox::Slot()
+		FOnClicked OnRelinkToParent = FOnClicked::CreateSPLambda(this, [this]()
+			{
+				return NestedTree->RelinkLayersToParent();
+			});
+		
+		if (GBEnableMaterialLayerPrototype)
+		{
+			HeaderBox = SNew(SHorizontalBox)
+			+ SHorizontalBox::Slot()
+				.AutoWidth()
+					.HAlign(HAlign_Center)
 					.VAlign(VAlign_Center)
-					.AutoWidth()
-					.Padding(8.0f, 0.0f)
+					.Padding(4.0f)
+					.MaxWidth(ThumbnailSize)
 					[
-						SNew(SVerticalBox)
-						+ SVerticalBox::Slot()
-						.AutoHeight()
+						SAssignNew(ThumbnailBox, SBox)
+						.MaxDesiredWidth(ThumbnailSize)
+						.MaxDesiredWidth(ThumbnailSize)
+						.MaxDesiredHeight(ThumbnailSize)
+						.MinDesiredHeight(ThumbnailSize)
 						[
-							SNew(STextBlock)
-							.Text(MaterialName)
-							.TextStyle(FAppStyle::Get(), "LargeText")
+							ThumbnailWidget
 						]
+					]
 
-						+ SVerticalBox::Slot()
-						.Padding(0, 4)
-						.AutoHeight()
-						[
-							SNew(STextBlock)
-							.Text(LOCTEXT("MaterialInstance", "Material Instance"))
-							.TextStyle(FAppStyle::Get(), "NormalText.Important")
-						]
+				+ SHorizontalBox::Slot()
+				.VAlign(VAlign_Center)
+				.AutoWidth()
+				.Padding(8.0f, 0.0f)
+				[
+					SNew(SVerticalBox)
+					+ SVerticalBox::Slot()
+					.AutoHeight()
+					[
+						SNew(STextBlock)
+						.Text(MaterialName)
+						.TextStyle(FAppStyle::Get(), "LargeText")
 					]
-#else
-					SAssignNew(HeaderBox, SHorizontalBox)
-					+ SHorizontalBox::Slot()
-					.Padding(FMargin(4.0f, 0.0f))
-					.HAlign(HAlign_Left)
-					.AutoWidth()
-					.VAlign(VAlign_Center)
+
+					+ SVerticalBox::Slot()
+					.Padding(0, 4)
+					.AutoHeight()
 					[
 						SNew(STextBlock)
-						.Text(LOCTEXT("MaterialLayers", "Material Layers"))
+						.Text(LOCTEXT("MaterialInstance", "Material Instance"))
+						.TextStyle(FAppStyle::Get(), "NormalText.Important")
 					]
-#endif
+				];
+		}
+		else
+		{
+			HeaderBox = SNew(SHorizontalBox)
+				+ SHorizontalBox::Slot()
+				.Padding(FMargin(4.0f, 0.0f))
+				.HAlign(HAlign_Left)
+				.AutoWidth()
+				.VAlign(VAlign_Center)
+				[
+					SNew(STextBlock)
+					.Text(LOCTEXT("MaterialLayers", "Material Layers"))
+				];
+		}
+
+		this->ChildSlot
+			[
+				SNew(SVerticalBox)
+				+ SVerticalBox::Slot()
+				.Padding(FMargin(0.0f, 0.0f, 0.0f, 10.0f))
+				.AutoHeight()
+				[
+					HeaderBox.ToSharedRef()
 				]
 				+ SVerticalBox::Slot()
 				.Padding(FMargin(2.0f))
@@ -1872,106 +1881,121 @@ void SMaterialLayersFunctionsInstanceWrapper::Refresh()
 					NestedTree.ToSharedRef()
 				]
 		];
-		if (FMaterialPropertyHelpers::IsOverriddenExpression(NestedTree->FunctionParameter))
+		if (FMaterialPropertyHelpers::IsOverriddenExpression(NestedTree->GetFunctionParameter()))
 		{
-#if ENABLE_MATERIAL_LAYER_PROTOTYPE
-			FMenuBuilder DetailViewOptions( true, nullptr);
-			DetailViewOptions.AddMenuEntry(
-			LOCTEXT("AddMaterialSlab", "Add Material Slab"),
-			LOCTEXT("AddMaterialSlab_ToolTip", "Add new Material Slab to the Substrate tree"),
-			FSlateIcon(),
-			FUIAction(FExecuteAction::CreateSP(NestedTree.Get(), &SMaterialSubstrateTree::AddRootNodeLayer)));
-
-			if (NestedTree->GetRelinkLayersToParentVisibility() == EVisibility::Visible)
+			if (GBEnableMaterialLayerPrototype)
 			{
+				FMenuBuilder DetailViewOptions( true, nullptr);
+				DetailViewOptions.AddMenuEntry(
+				LOCTEXT("AddMaterialSlab", "Add Material Slab"),
+				LOCTEXT("AddMaterialSlab_ToolTip", "Add new Material Slab to the Substrate tree"),
+				FSlateIcon(),
+				FUIAction(FExecuteAction::CreateSPLambda(this, [this]()
+					{
+						NestedTree->AddRootNodeLayer();
+					})));
+
+				if (NestedTree->GetRelinkLayersToParentVisibility() == EVisibility::Visible)
+				{
+					DetailViewOptions.AddMenuEntry(
+					LOCTEXT("Relink", "Relink"),
+					LOCTEXT("RelinkToParentLayers", "Relink to Parent Layers and Blends"),
+					FSlateIcon(),
+					FUIAction(FExecuteAction::CreateSPLambda(this, [this]()
+					{
+						NestedTree->RelinkLayersToParent();
+					})));
+				}
+				DetailViewOptions.AddMenuEntry(
+				LOCTEXT("SaveSibling", "Save Sibling"),
+				LOCTEXT("SaveToSiblingInstance", "Save to Sibling Instance"),
+				FSlateIcon(),
+				FUIAction(FExecuteAction::CreateLambda([this, OnSiblingButtonClicked]()	{ OnSiblingButtonClicked.Execute(); })));
+
 				DetailViewOptions.AddMenuEntry(
-				LOCTEXT("Relink", "Relink"),
-				LOCTEXT("RelinkToParentLayers", "Relink to Parent Layers and Blends"),
+				LOCTEXT("SaveChild", "Save Child"),
+				LOCTEXT("SaveToChildInstance", "Save To Child Instance"),
 				FSlateIcon(),
-				FUIAction(FExecuteAction::CreateSP(NestedTree.Get(), &SMaterialSubstrateTree::RelinkLayersToParent)));
+				FUIAction(FExecuteAction::CreateLambda([this, OnChildButtonClicked]()	{ OnChildButtonClicked.Execute(); })));
+
+				HeaderBox->AddSlot()
+				.FillWidth(1.0f)
+				[
+					SNullWidget::NullWidget
+				];
+				HeaderBox->AddSlot()
+					.HAlign(HAlign_Right)
+					.AutoWidth()
+					.Padding(2.0f)
+					.VAlign(VAlign_Center)
+					[
+						SNew(SComboButton)
+							.ContentPadding(0.f)
+							.ComboButtonStyle(&FAppStyle::Get().GetWidgetStyle<FComboButtonStyle>("StatusBar.StatusBarEllipsisComboButton"))
+							.HasDownArrow(true)
+							.MenuContent()
+							[
+								DetailViewOptions.MakeWidget()
+							]
+					];
 			}
-			DetailViewOptions.AddMenuEntry(
-			LOCTEXT("SaveSibling", "Save Sibling"),
-			LOCTEXT("SaveToSiblingInstance", "Save to Sibling Instance"),
-			FSlateIcon(),
-			FUIAction(FExecuteAction::CreateLambda([this, OnSiblingButtonClicked]()	{ OnSiblingButtonClicked.Execute(); })));
-
-			DetailViewOptions.AddMenuEntry(
-			LOCTEXT("SaveChild", "Save Child"),
-			LOCTEXT("SaveToChildInstance", "Save To Child Instance"),
-			FSlateIcon(),
-			FUIAction(FExecuteAction::CreateLambda([this, OnChildButtonClicked]()	{ OnChildButtonClicked.Execute(); })));
+			else
+			{
+				HeaderBox->AddSlot()
+					.HAlign(HAlign_Left)
+					.AutoWidth()
+					.VAlign(VAlign_Center)
+					[
+					PropertyCustomizationHelpers::MakeAddButton(FSimpleDelegate::CreateSPLambda(this, [this]()
+						{
+							NestedTree->AddRootNodeLayer();
+						}))
+					];
+			}
+		}
 
+		if (!GBEnableMaterialLayerPrototype)
+		{
 			HeaderBox->AddSlot()
-			.FillWidth(1.0f)
-			[
-				SNullWidget::NullWidget
-			];
-			HeaderBox->AddSlot()
-				.HAlign(HAlign_Right)
-				.AutoWidth()
-				.Padding(2.0f)
-				.VAlign(VAlign_Center)
+				.FillWidth(1.0f)
 				[
-					SNew(SComboButton)
-						.ContentPadding(0.f)
-						.ComboButtonStyle(&FAppStyle::Get().GetWidgetStyle<FComboButtonStyle>("StatusBar.StatusBarEllipsisComboButton"))
-						.HasDownArrow(true)
-						.MenuContent()
-						[
-							DetailViewOptions.MakeWidget()
-						]
+					SNullWidget::NullWidget
 				];
-			
-#else
 			HeaderBox->AddSlot()
-				.HAlign(HAlign_Left)
 				.AutoWidth()
-				.VAlign(VAlign_Center)
+				.Padding(2.0f)
 				[
-					PropertyCustomizationHelpers::MakeAddButton(FSimpleDelegate::CreateSP(NestedTree.Get(), &SMaterialLayersFunctionsInstanceTree::AddLayer))
+					SNew(SButton)
+					.Text(LOCTEXT("Relink", "Relink"))
+					.HAlign(HAlign_Center)
+					.OnClicked(OnRelinkToParent)
+					.ToolTipText(LOCTEXT("RelinkToParentLayers", "Relink to Parent Layers and Blends"))
+					.Visibility_Lambda([this]()
+					{
+						return NestedTree->GetRelinkLayersToParentVisibility();
+					})
 				];
-#endif
-		}
-
-#if !ENABLE_MATERIAL_LAYER_PROTOTYPE
-		HeaderBox->AddSlot()
-			.FillWidth(1.0f)
-			[
-				SNullWidget::NullWidget
-			];
-		HeaderBox->AddSlot()
-			.AutoWidth()
-			.Padding(2.0f)
-			[
-				SNew(SButton)
-				.Text(LOCTEXT("Relink", "Relink"))
-				.HAlign(HAlign_Center)
-				.OnClicked(OnRelinkToParent)
-				.ToolTipText(LOCTEXT("RelinkToParentLayers", "Relink to Parent Layers and Blends"))
-				.Visibility(NestedTree.Get(), &SMaterialLayersFunctionsInstanceTree::GetRelinkLayersToParentVisibility)
-			];
-		HeaderBox->AddSlot()
+			HeaderBox->AddSlot()
+					.AutoWidth()
+					.Padding(2.0f)
+					[
+						SNew(SButton)
+						.Text(LOCTEXT("SaveSibling", "Save Sibling"))
+						.HAlign(HAlign_Center)
+						.OnClicked(OnSiblingButtonClicked)
+						.ToolTipText(LOCTEXT("SaveToSiblingInstance", "Save to Sibling Instance"))
+					];
+			HeaderBox->AddSlot()
 				.AutoWidth()
 				.Padding(2.0f)
 				[
 					SNew(SButton)
-					.Text(LOCTEXT("SaveSibling", "Save Sibling"))
+					.Text(LOCTEXT("SaveChild", "Save Child"))
 					.HAlign(HAlign_Center)
-					.OnClicked(OnSiblingButtonClicked)
-					.ToolTipText(LOCTEXT("SaveToSiblingInstance", "Save to Sibling Instance"))
+					.OnClicked(OnChildButtonClicked)
+					.ToolTipText(LOCTEXT("SaveToChildInstance", "Save To Child Instance"))
 				];
-		HeaderBox->AddSlot()
-			.AutoWidth()
-			.Padding(2.0f)
-			[
-				SNew(SButton)
-				.Text(LOCTEXT("SaveChild", "Save Child"))
-				.HAlign(HAlign_Center)
-				.OnClicked(OnChildButtonClicked)
-				.ToolTipText(LOCTEXT("SaveToChildInstance", "Save To Child Instance"))
-			];
-#endif
+		}
 	}
 	else
 	{
@@ -1991,21 +2015,24 @@ void SMaterialLayersFunctionsInstanceWrapper::Refresh()
 
 void SMaterialLayersFunctionsInstanceWrapper::Construct(const FArguments& InArgs)
 {
-#if ENABLE_MATERIAL_LAYER_PROTOTYPE
-	NestedTree = SNew(SMaterialSubstrateTree)
-		.InMaterialEditorInstance(InArgs._InMaterialEditorInstance)
-		.InGenerator(InArgs._InGenerator)
-		.InWrapper(this)
-		.InShowHiddenDelegate(InArgs._InShowHiddenDelegate);
-#else
-	NestedTree = SNew(SMaterialLayersFunctionsInstanceTree)
-		.InMaterialEditorInstance(Cast<UMaterialEditorInstanceConstant>(InArgs._InMaterialEditorInstance))
-		.InWrapper(this)
-		.InShowHiddenDelegate(InArgs._InShowHiddenDelegate);
-#endif
+	if (GBEnableMaterialLayerPrototype)
+	{
+		NestedTree = SNew(SMaterialSubstrateTree)
+			.InMaterialEditorInstance(InArgs._InMaterialEditorInstance)
+			.InGenerator(InArgs._InGenerator)
+			.InWrapper(this)
+			.InShowHiddenDelegate(InArgs._InShowHiddenDelegate);
+	}
+	else
+	{
+		NestedTree = SNew(SMaterialLayersFunctionsInstanceTree)
+			.InMaterialEditorInstance(Cast<UMaterialEditorInstanceConstant>(InArgs._InMaterialEditorInstance))
+			.InWrapper(this)
+			.InShowHiddenDelegate(InArgs._InShowHiddenDelegate);
+	}
 	
 
-	LayerParameter = NestedTree->FunctionParameter;
+	LayerParameter = NestedTree->GetFunctionParameter();
 
 	MaterialEditorInstance = InArgs._InMaterialEditorInstance;
 	Refresh();
@@ -2014,12 +2041,7 @@ void SMaterialLayersFunctionsInstanceWrapper::Construct(const FArguments& InArgs
 
 void SMaterialLayersFunctionsInstanceWrapper::SetEditorInstance(UMaterialEditorParameters* InMaterialEditorInstance)
 {
-#if ENABLE_MATERIAL_LAYER_PROTOTYPE
-	NestedTree->MaterialEditorInstance = InMaterialEditorInstance;
-#else
-	NestedTree->MaterialEditorInstance = Cast<UMaterialEditorInstanceConstant>(InMaterialEditorInstance);
-#endif
-	
+	NestedTree->SetMaterialEditorInstance(InMaterialEditorInstance);
 	Refresh();
 }
 
@@ -3304,9 +3326,9 @@ void SMaterialLayersFunctionsMaterialWrapper::Construct(const FArguments& InArgs
 
 }
 
-void SMaterialLayersFunctionsMaterialWrapper::SetEditorInstance(UMaterialEditorPreviewParameters* InMaterialEditorInstance)
+void SMaterialLayersFunctionsMaterialWrapper::SetEditorInstance(UMaterialEditorParameters* InMaterialEditorInstance)
 {
-	NestedTree->MaterialEditorInstance = InMaterialEditorInstance;
+	NestedTree->MaterialEditorInstance = dynamic_cast<UMaterialEditorPreviewParameters*>(InMaterialEditorInstance);
 	Refresh();
 }
 
diff --git a/Engine/Source/Editor/MaterialEditor/Private/SMaterialLayersFunctionsTree.h b/Engine/Source/Editor/MaterialEditor/Private/SMaterialLayersFunctionsTree.h
--- a/Engine/Source/Editor/MaterialEditor/Private/SMaterialLayersFunctionsTree.h
+++ b/Engine/Source/Editor/MaterialEditor/Private/SMaterialLayersFunctionsTree.h
@@ -19,6 +19,12 @@
 #include "Widgets/Views/STreeView.h"
 #include "Materials/MaterialExpressionMaterialSample.h"
 
+#if ENABLE_MATERIAL_LAYER_PROTOTYPE
+static constexpr bool GBEnableMaterialLayerPrototype = true;
+#else
+static constexpr bool GBEnableMaterialLayerPrototype = false;
+#endif
+
 class IPropertyHandle;
 class SMaterialLayersFunctionsInstanceTree;
 class UDEditorParameterValue;
@@ -97,7 +103,15 @@ public:
 	FString GetInstancePath(SMaterialLayersFunctionsInstanceTree* InTree) const;
 };
 
-class SMaterialLayersFunctionsInstanceWrapper : public SCompoundWidget
+class SMaterialLayersFunctionsWrapper: public SCompoundWidget
+{
+public:
+	virtual void SetEditorInstance(UMaterialEditorParameters* InMaterialEditorInstance) = 0;
+	virtual void Refresh() = 0;
+	
+};
+
+class SMaterialLayersFunctionsInstanceWrapper : public SMaterialLayersFunctionsWrapper
 {
 public:
 	SLATE_BEGIN_ARGS(SMaterialLayersFunctionsInstanceWrapper)
@@ -109,21 +123,18 @@ public:
 	SLATE_ARGUMENT(UMaterialEditorParameters*, InMaterialEditorInstance)
 	SLATE_ARGUMENT(TSharedPtr<class IPropertyRowGenerator>, InGenerator)
 	SLATE_ARGUMENT(FGetShowHiddenParameters, InShowHiddenDelegate)
-
 	SLATE_END_ARGS()
-	void Refresh();
+	
 	void Construct(const FArguments& InArgs);
-	void SetEditorInstance(UMaterialEditorParameters* InMaterialEditorInstance);
+	virtual void SetEditorInstance(UMaterialEditorParameters* InMaterialEditorInstance) override;
+	virtual void Refresh() override;
 	
 	TSharedPtr<class IPropertyRowGenerator> GetGenerator();
 	TAttribute<ECheckBoxState> IsParamChecked;
 	TWeakObjectPtr<class UDEditorParameterValue> LayerParameter;
 	class UMaterialEditorParameters* MaterialEditorInstance;
-#if ENABLE_MATERIAL_LAYER_PROTOTYPE
-	TSharedPtr<class SMaterialSubstrateTree> NestedTree;
-#else
-	TSharedPtr<class SMaterialLayersFunctionsInstanceTree> NestedTree;
-#endif
+
+	TSharedPtr<class SMaterialLayersTree> NestedTree;
 	FSimpleDelegate OnLayerPropertyChanged;
 	
 	
@@ -131,7 +142,20 @@ private:
 	TWeakPtr<class IPropertyRowGenerator> Generator;
 };
 
-class SMaterialLayersFunctionsInstanceTree : public STreeView<TSharedPtr<FSortedParamData>>
+class SMaterialLayersTree: public STreeView<TSharedPtr<FSortedParamData>>
+{
+public:
+	virtual void CreateGroupsWidget() = 0;
+	virtual TWeakObjectPtr<class UDEditorParameterValue> GetFunctionParameter() = 0;
+	virtual TSharedPtr<class FAssetThumbnailPool> GetTreeThumbnailPool() = 0;
+	virtual EVisibility GetRelinkLayersToParentVisibility() const = 0;
+	virtual FReply RelinkLayersToParent() = 0;
+	virtual void SetMaterialEditorInstance(UMaterialEditorParameters* InMaterialEditorInstance) = 0;
+	virtual void UpdateThumbnailMaterial(TEnumAsByte<EMaterialParameterAssociation> InAssociation, int32 InIndex, bool bAlterBlendIndex = false) = 0;
+	virtual void CollectAssetStackItemsRecursively(TSharedPtr<FSortedParamData> Item, TArray<TSharedPtr<FSortedParamData>>& OutGroupsContainer, TArray<uint32>& OutNodeIdsContainer) {}
+	virtual void AddRootNodeLayer() {}
+};
+class SMaterialLayersFunctionsInstanceTree : public SMaterialLayersTree
 {
 	friend class SMaterialLayersFunctionsInstanceTreeItem;
 public:
@@ -158,23 +182,23 @@ public:
 	void RefreshOnAssetChange(const struct FAssetData& InAssetData, int32 Index, EMaterialParameterAssociation MaterialType);
 	void ResetAssetToDefault(TSharedPtr<FSortedParamData> InData);
 	void AddLayer();
+	virtual void AddRootNodeLayer() override { AddLayer(); }
 	void RemoveLayer(int32 Index);
 	FReply UnlinkLayer(int32 Index);
-	FReply RelinkLayersToParent();
+	virtual FReply RelinkLayersToParent() override;
 	EVisibility GetUnlinkLayerVisibility(int32 Index) const;
-	EVisibility GetRelinkLayersToParentVisibility() const;
+	virtual EVisibility GetRelinkLayersToParentVisibility() const override;
+	virtual void SetMaterialEditorInstance(UMaterialEditorParameters* InMaterialEditorInstance) override { MaterialEditorInstance = Cast<UMaterialEditorInstanceConstant>(InMaterialEditorInstance); } 
 	FReply ToggleLayerVisibility(int32 Index);
-	TSharedPtr<class FAssetThumbnailPool> GetTreeThumbnailPool();
+	virtual TSharedPtr<class FAssetThumbnailPool> GetTreeThumbnailPool() override;
 
 	/** Object that stores all of the possible parameters we can edit */
 	UMaterialEditorInstanceConstant* MaterialEditorInstance;
 
 	/** Builds the custom parameter groups category */
-	void CreateGroupsWidget();
-
-	// Adding this noop method for compatibility with the Substrate enabled Material Layer. 
-	void CollectAssetStackItemsRecursively(TSharedPtr<FSortedParamData> Item, TArray<TSharedPtr<FSortedParamData>>& OutGroupsContainer, TArray<uint32>& OutNodeIdsContainer) {}
-
+	virtual void CreateGroupsWidget() override;
+	virtual TWeakObjectPtr<class UDEditorParameterValue> GetFunctionParameter() override{ return FunctionParameter; }
+	
 	bool IsLayerVisible(int32 Index) const;
 
 	SMaterialLayersFunctionsInstanceWrapper* GetWrapper() { return Wrapper; }
@@ -182,7 +206,7 @@ public:
 
 
 	TSharedRef<SWidget> CreateThumbnailWidget(EMaterialParameterAssociation InAssociation, int32 InIndex, float InThumbnailSize);
-	void UpdateThumbnailMaterial(TEnumAsByte<EMaterialParameterAssociation> InAssociation, int32 InIndex, bool bAlterBlendIndex = false);
+	virtual void UpdateThumbnailMaterial(TEnumAsByte<EMaterialParameterAssociation> InAssociation, int32 InIndex, bool bAlterBlendIndex = false) override;
 	FReply OnThumbnailDoubleClick(const FGeometry& Geometry, const FPointerEvent& MouseEvent, EMaterialParameterAssociation InAssociation, int32 InIndex);
 	bool IsOverriddenExpression(class UDEditorParameterValue* Parameter, int32 InIndex);
 	bool IsOverriddenExpression(TObjectPtr<UDEditorParameterValue> Parameter, int32 InIndex);
@@ -260,7 +284,7 @@ public:
 	FString GetInstancePath(SMaterialLayersFunctionsMaterialTree* InTree) const;
 };
 
-class SMaterialLayersFunctionsMaterialWrapper : public SCompoundWidget
+class SMaterialLayersFunctionsMaterialWrapper : public SMaterialLayersFunctionsWrapper
 {
 public:
 	SLATE_BEGIN_ARGS(SMaterialLayersFunctionsMaterialWrapper)
@@ -271,9 +295,10 @@ public:
 	SLATE_ARGUMENT(UMaterialEditorPreviewParameters*, InMaterialEditorInstance)
 	SLATE_ARGUMENT(TSharedPtr<class IPropertyRowGenerator>, InGenerator)
 	SLATE_END_ARGS()
-	void Refresh();
+
+	virtual void Refresh() override;
 	void Construct(const FArguments& InArgs);
-	void SetEditorInstance(UMaterialEditorPreviewParameters* InMaterialEditorInstance);
+	virtual void SetEditorInstance(UMaterialEditorParameters* InMaterialEditorInstance) override;
 	TSharedPtr<class IPropertyRowGenerator> GetGenerator();
 	TWeakObjectPtr<class UDEditorParameterValue> LayerParameter;
 	class UMaterialEditorPreviewParameters* MaterialEditorInstance;
diff --git a/Engine/Source/Editor/MaterialEditor/Private/SMaterialSubstrateTree.cpp b/Engine/Source/Editor/MaterialEditor/Private/SMaterialSubstrateTree.cpp
--- a/Engine/Source/Editor/MaterialEditor/Private/SMaterialSubstrateTree.cpp
+++ b/Engine/Source/Editor/MaterialEditor/Private/SMaterialSubstrateTree.cpp
@@ -1171,7 +1171,7 @@ void SMaterialSubstrateTree::UnlinkLayer(int32 Index)
 	Wrapper->Refresh();
 }
 
-void SMaterialSubstrateTree::RelinkLayersToParent()
+FReply SMaterialSubstrateTree::RelinkLayersToParent()
 {
 	const FScopedTransaction Transaction(LOCTEXT("RelinkLayersToParent", "Relink layers to parent"));
 	FunctionInstanceHandle->NotifyPreChange();
@@ -1180,6 +1180,7 @@ void SMaterialSubstrateTree::RelinkLayersToParent()
 	MaterialEditorInstance->RegenerateArrays();
 	CreateGroupsWidget();
 	Wrapper->Refresh();
+	return FReply::Handled();
 }
 
 EVisibility SMaterialSubstrateTree::GetUnlinkLayerVisibility(int32 Index) const
diff --git a/Engine/Source/Editor/MaterialEditor/Private/SMaterialSubstrateTree.h b/Engine/Source/Editor/MaterialEditor/Private/SMaterialSubstrateTree.h
--- a/Engine/Source/Editor/MaterialEditor/Private/SMaterialSubstrateTree.h
+++ b/Engine/Source/Editor/MaterialEditor/Private/SMaterialSubstrateTree.h
@@ -17,6 +17,7 @@
 #include "Widgets/DeclarativeSyntaxSupport.h"
 #include "Widgets/SCompoundWidget.h"
 #include "Widgets/Views/STreeView.h"
+#include "SMaterialLayersFunctionsTree.h"
 
 class IPropertyHandle;
 class SMaterialSubstrateTree;
@@ -26,7 +27,7 @@ class UMaterialEditorInstanceConstant;
 struct FRecursiveCreateWidgetsContext;
 
 typedef TSharedPtr<FSortedParamData> FSortedParamDataPtr;
-class SMaterialSubstrateTree : public STreeView<FSortedParamDataPtr>
+class SMaterialSubstrateTree : public SMaterialLayersTree
 {
 	friend class SMaterialSubstrateTreeItem;
 public:
@@ -61,15 +62,17 @@ public:
 	TSharedPtr<IPropertyHandle> FunctionInstanceHandle;
 	void RefreshOnAssetChange(const struct FAssetData& InAssetData, int32 InNodeId, EMaterialParameterAssociation MaterialType);
 	void ResetAssetToDefault(TSharedPtr<FSortedParamData> InData);
-	
-	void AddRootNodeLayer() { AddNodeLayer(); }
+
+	virtual void AddRootNodeLayer() override { AddNodeLayer(); }
 	void AddNodeLayer(int32 InParent = -1);
 	void RemoveNodeLayer(int32 InNodeId);
 
 	void UnlinkLayer(int32 Index);
-	void RelinkLayersToParent();
+	virtual FReply RelinkLayersToParent() override;
 	EVisibility GetUnlinkLayerVisibility(int32 Index) const;
-	EVisibility GetRelinkLayersToParentVisibility() const;
+	virtual EVisibility GetRelinkLayersToParentVisibility() const override;
+	virtual void SetMaterialEditorInstance(UMaterialEditorParameters* InMaterialEditorInstance) override { MaterialEditorInstance = InMaterialEditorInstance; } 
+	
 	FReply ToggleLayerVisibility(int32 Index);
 	bool IsLayerVisible(int32 Index) const;
 
@@ -79,15 +82,15 @@ public:
 	UMaterialEditorParameters* MaterialEditorInstance;
 
 	/** Builds the custom parameter groups category */
-	void CreateGroupsWidget();
-
+	virtual void CreateGroupsWidget() override;
+	virtual TWeakObjectPtr<class UDEditorParameterValue> GetFunctionParameter() override{ return FunctionParameter; }
 
 	SMaterialLayersFunctionsInstanceWrapper* GetWrapper() { return Wrapper; }
 
 	TSharedPtr<IDetailTreeNode> FindParameterGroupsNode(TSharedPtr<IPropertyRowGenerator> PropertyRowGenerator);
 
 	TSharedRef<SWidget> CreateThumbnailWidget(EMaterialParameterAssociation InAssociation, int32 InIndex, float InThumbnailSize);
-	void UpdateThumbnailMaterial(TEnumAsByte<EMaterialParameterAssociation> InAssociation, int32 InIndex, bool bAlterBlendIndex = false);
+	virtual void UpdateThumbnailMaterial(TEnumAsByte<EMaterialParameterAssociation> InAssociation, int32 InIndex, bool bAlterBlendIndex = false) override;
 	FReply OnThumbnailDoubleClick(const FGeometry& Geometry, const FPointerEvent& MouseEvent, EMaterialParameterAssociation InAssociation, int32 InIndex);
 	bool IsOverriddenExpression(class UDEditorParameterValue* Parameter, int32 InIndex);
 	bool IsOverriddenExpression(TObjectPtr<UDEditorParameterValue> Parameter, int32 InIndex);
@@ -97,7 +100,7 @@ public:
 
 	void CollectStackItemsRecursively(TSharedPtr<FSortedParamData> Item, TArray<TSharedPtr<FSortedParamData>>& OutGroupsContainer);
 
-	void CollectAssetStackItemsRecursively(TSharedPtr<FSortedParamData> Item, TArray<TSharedPtr<FSortedParamData>>& OutGroupsContainer, TArray<uint32>& OutNodeIdsContainer);
+	virtual void CollectAssetStackItemsRecursively(TSharedPtr<FSortedParamData> Item, TArray<TSharedPtr<FSortedParamData>>& OutGroupsContainer, TArray<uint32>& OutNodeIdsContainer) override;
 
 	void CollectStackItemsForMaterialFunctionAsset(EMaterialParameterAssociation InAssociation, int32 InAssetIndex, TArray<TSharedPtr<FSortedParamData>>& OutGroupsContainer);
 
