Material Graph: Allow Channel Names Customization on Texture Object Parameters. Primarily for showing up in Material Instance Constants

- Also allow Custom Channel Names on any UMaterialExpressionTextureSample, rather than just expressions inheriting from UMaterialExpressionTextureSampleParameter

diff --git a/Engine/Source/Editor/MaterialEditor/Private/MaterialEditorDetailCustomization.cpp b/Engine/Source/Editor/MaterialEditor/Private/MaterialEditorDetailCustomization.cpp
--- a/Engine/Source/Editor/MaterialEditor/Private/MaterialEditorDetailCustomization.cpp
+++ b/Engine/Source/Editor/MaterialEditor/Private/MaterialEditorDetailCustomization.cpp
@@ -281,167 +281,160 @@ void FMaterialExpressionParameterDetails::CustomizeDetails( IDetailLayoutBuilder
 		{
 			TSharedPtr<IPropertyHandle> ChannelHandle = DetailLayout.GetProperty("ChannelNames", UMaterialExpressionTextureSampleParameter::StaticClass());
 			TSharedPtr<IPropertyHandle> ValueHandle = DetailLayout.GetProperty("Texture", UMaterialExpressionTextureBase::StaticClass());
-			if (TextureParameter->GetOutputType(0) != MCT_Texture)
+			
+			if (ChannelHandle.IsValid() && ChannelHandle->IsValidHandle())
 			{
-				if (ChannelHandle.IsValid() && ChannelHandle->IsValidHandle())
+				static const FName Red("R");
+				static const FName Green("G");
+				static const FName Blue("B");
+				static const FName Alpha("A");
+				// Rebuild the layout when the ChannelNames property changes
+				ChannelHandle->GetChildHandle(Red)->SetOnPropertyValueChanged(FSimpleDelegate::CreateLambda([&DetailLayout]()
 				{
-					static const FName Red("R");
-					static const FName Green("G");
-					static const FName Blue("B");
-					static const FName Alpha("A");
-					// Rebuild the layout when the ChannelNames property changes
-					ChannelHandle->GetChildHandle(Red)->SetOnPropertyValueChanged(FSimpleDelegate::CreateLambda([&DetailLayout]()
-					{
-						DetailLayout.ForceRefreshDetails();
-					}));
-					ChannelHandle->GetChildHandle(Green)->SetOnPropertyValueChanged(FSimpleDelegate::CreateLambda([&DetailLayout]()
-					{
-						DetailLayout.ForceRefreshDetails();
-					}));
-					ChannelHandle->GetChildHandle(Blue)->SetOnPropertyValueChanged(FSimpleDelegate::CreateLambda([&DetailLayout]()
-					{
-						DetailLayout.ForceRefreshDetails();
-					}));
-					ChannelHandle->GetChildHandle(Alpha)->SetOnPropertyValueChanged(FSimpleDelegate::CreateLambda([&DetailLayout]()
-					{
-						DetailLayout.ForceRefreshDetails();
-					}));
-				}
-
-				if (ValueHandle.IsValid() && ValueHandle->IsValidHandle())
+					DetailLayout.ForceRefreshDetails();
+				}));
+				ChannelHandle->GetChildHandle(Green)->SetOnPropertyValueChanged(FSimpleDelegate::CreateLambda([&DetailLayout]()
 				{
-					IDetailPropertyRow& PropertyRow = *DetailLayout.EditDefaultProperty(ValueHandle);
-					TSharedPtr<SWidget> NameWidget;
-					TSharedPtr<SWidget> ValueWidget;
-					FDetailWidgetRow DefaultRow;
-					PropertyRow.GetDefaultWidgets(NameWidget, ValueWidget, DefaultRow);
-
-					FDetailWidgetRow &DetailWidgetRow = PropertyRow.CustomWidget();
-					TSharedPtr<SVerticalBox> NameVerticalBox;
-					DetailWidgetRow.NameContent()
-						[
-							SAssignNew(NameVerticalBox, SVerticalBox)
-							+ SVerticalBox::Slot()
-							.AutoHeight()
-							[
-								SNew(STextBlock)
-								.Text(FText::FromName(TextureParameter->ParameterName))
-								.Font(FAppStyle::GetFontStyle(TEXT("PropertyWindow.NormalFont")))
-							]
-						];
+					DetailLayout.ForceRefreshDetails();
+				}));
+				ChannelHandle->GetChildHandle(Blue)->SetOnPropertyValueChanged(FSimpleDelegate::CreateLambda([&DetailLayout]()
+				{
+					DetailLayout.ForceRefreshDetails();
+				}));
+				ChannelHandle->GetChildHandle(Alpha)->SetOnPropertyValueChanged(FSimpleDelegate::CreateLambda([&DetailLayout]()
+				{
+					DetailLayout.ForceRefreshDetails();
+				}));
+			}
 
-					DetailWidgetRow.ValueContent()
-						.MinDesiredWidth(DefaultRow.ValueWidget.MinWidth)
-						.MaxDesiredWidth(DefaultRow.ValueWidget.MaxWidth)
+			if (ValueHandle.IsValid() && ValueHandle->IsValidHandle())
+			{
+				IDetailPropertyRow& PropertyRow = *DetailLayout.EditDefaultProperty(ValueHandle);
+				TSharedPtr<SWidget> NameWidget;
+				TSharedPtr<SWidget> ValueWidget;
+				FDetailWidgetRow DefaultRow;
+				PropertyRow.GetDefaultWidgets(NameWidget, ValueWidget, DefaultRow);
+
+				FDetailWidgetRow &DetailWidgetRow = PropertyRow.CustomWidget();
+				TSharedPtr<SVerticalBox> NameVerticalBox;
+				DetailWidgetRow.NameContent()
+					[
+						SAssignNew(NameVerticalBox, SVerticalBox)
+						+ SVerticalBox::Slot()
+						.AutoHeight()
 						[
-							ValueWidget.ToSharedRef()
-						];
+							SNew(STextBlock)
+							.Text(FText::FromName(TextureParameter->ParameterName))
+							.Font(FAppStyle::GetFontStyle(TEXT("PropertyWindow.NormalFont")))
+						]
+					];
 
-					static const FName Red("R");
-					static const FName Green("G");
-					static const FName Blue("B");
-					static const FName Alpha("A");
+				DetailWidgetRow.ValueContent()
+					.MinDesiredWidth(DefaultRow.ValueWidget.MinWidth)
+					.MaxDesiredWidth(DefaultRow.ValueWidget.MaxWidth)
+					[
+						ValueWidget.ToSharedRef()
+					];
 
-					if (!TextureParameter->ChannelNames.R.IsEmpty())
-					{
-						NameVerticalBox->AddSlot()
-							[
-								SNew(SHorizontalBox)
-								+ SHorizontalBox::Slot()
-								.AutoWidth()
-								.Padding(20.0, 2.0, 4.0, 2.0)
-								[
-									SNew(STextBlock)
-									.Text(FText::FromName(Red))
-									.Font(FAppStyle::GetFontStyle(TEXT("PropertyWindow.BoldFont")))
-								]
-								+ SHorizontalBox::Slot()
-								.HAlign(HAlign_Left)
-								.Padding(4.0, 2.0)
-								[
-									SNew(STextBlock)
-									.Text(TextureParameter->ChannelNames.R)
-									.Font(FAppStyle::GetFontStyle(TEXT("PropertyWindow.NormalFont")))
-								]
-							];
-					}
-					if (!TextureParameter->ChannelNames.G.IsEmpty())
-					{
-						NameVerticalBox->AddSlot()
-							[
-								SNew(SHorizontalBox)
-								+ SHorizontalBox::Slot()
-								.Padding(20.0, 2.0, 4.0, 2.0)
-								.AutoWidth()
-								[
-									SNew(STextBlock)
-									.Text(FText::FromName(Green))
-									.Font(FAppStyle::GetFontStyle(TEXT("PropertyWindow.BoldFont")))
-								]
-								+ SHorizontalBox::Slot()
-								.HAlign(HAlign_Left)
-								.Padding(4.0, 2.0)
-								[
-									SNew(STextBlock)
-									.Text(TextureParameter->ChannelNames.G)
-									.Font(FAppStyle::GetFontStyle(TEXT("PropertyWindow.NormalFont")))
-								]
-							];
-					}
-					if (!TextureParameter->ChannelNames.B.IsEmpty())
-					{
-							NameVerticalBox->AddSlot()
-							[
-								SNew(SHorizontalBox)
-								+ SHorizontalBox::Slot()
-								.Padding(20.0, 2.0, 4.0, 2.0)
-								.AutoWidth()
-								[
-									SNew(STextBlock)
-									.Text(FText::FromName(Blue))
-									.Font(FAppStyle::GetFontStyle(TEXT("PropertyWindow.BoldFont")))
-								]
-								+ SHorizontalBox::Slot()
-								.HAlign(HAlign_Left)
-								.Padding(4.0, 2.0)
-								[
-									SNew(STextBlock)
-									.Text(TextureParameter->ChannelNames.B)
-									.Font(FAppStyle::GetFontStyle(TEXT("PropertyWindow.NormalFont")))
-								]
-							];
-					}
-					if (!TextureParameter->ChannelNames.A.IsEmpty())
-					{
-						NameVerticalBox->AddSlot()
-							[
-								SNew(SHorizontalBox)
-								+ SHorizontalBox::Slot()
-								.Padding(20.0, 2.0, 4.0, 2.0)
-								.AutoWidth()
-								[
-									SNew(STextBlock)
-									.Text(FText::FromName(Alpha))
-									.Font(FAppStyle::GetFontStyle(TEXT("PropertyWindow.BoldFont")))
-								]
-								+ SHorizontalBox::Slot()
-								.HAlign(HAlign_Left)
-								.Padding(4.0, 2.0)
-								[
-									SNew(STextBlock)
-									.Text(TextureParameter->ChannelNames.A)
-									.Font(FAppStyle::GetFontStyle(TEXT("PropertyWindow.NormalFont")))
-								]
-							];
-					}
+				static const FName Red("R");
+				static const FName Green("G");
+				static const FName Blue("B");
+				static const FName Alpha("A");
+
+				if (!TextureParameter->ChannelNames.R.IsEmpty())
+				{
+					NameVerticalBox->AddSlot()
+					[
+						SNew(SHorizontalBox)
+						+ SHorizontalBox::Slot()
+						.AutoWidth()
+						.Padding(20.0, 2.0, 4.0, 2.0)
+						[
+							SNew(STextBlock)
+							.Text(FText::FromName(Red))
+							.Font(FAppStyle::GetFontStyle(TEXT("PropertyWindow.BoldFont")))
+						]
+						+ SHorizontalBox::Slot()
+						.HAlign(HAlign_Left)
+						.Padding(4.0, 2.0)
+						[
+							SNew(STextBlock)
+							.Text(TextureParameter->ChannelNames.R)
+							.Font(FAppStyle::GetFontStyle(TEXT("PropertyWindow.NormalFont")))
+						]
+					];
+				}
+				if (!TextureParameter->ChannelNames.G.IsEmpty())
+				{
+					NameVerticalBox->AddSlot()
+					[
+						SNew(SHorizontalBox)
+						+ SHorizontalBox::Slot()
+						.Padding(20.0, 2.0, 4.0, 2.0)
+						.AutoWidth()
+						[
+							SNew(STextBlock)
+							.Text(FText::FromName(Green))
+							.Font(FAppStyle::GetFontStyle(TEXT("PropertyWindow.BoldFont")))
+						]
+						+ SHorizontalBox::Slot()
+						.HAlign(HAlign_Left)
+						.Padding(4.0, 2.0)
+						[
+							SNew(STextBlock)
+							.Text(TextureParameter->ChannelNames.G)
+							.Font(FAppStyle::GetFontStyle(TEXT("PropertyWindow.NormalFont")))
+						]
+					];
+				}
+				if (!TextureParameter->ChannelNames.B.IsEmpty())
+				{
+					NameVerticalBox->AddSlot()
+					[
+						SNew(SHorizontalBox)
+						+ SHorizontalBox::Slot()
+						.Padding(20.0, 2.0, 4.0, 2.0)
+						.AutoWidth()
+						[
+							SNew(STextBlock)
+							.Text(FText::FromName(Blue))
+							.Font(FAppStyle::GetFontStyle(TEXT("PropertyWindow.BoldFont")))
+						]
+						+ SHorizontalBox::Slot()
+						.HAlign(HAlign_Left)
+						.Padding(4.0, 2.0)
+						[
+							SNew(STextBlock)
+							.Text(TextureParameter->ChannelNames.B)
+							.Font(FAppStyle::GetFontStyle(TEXT("PropertyWindow.NormalFont")))
+						]
+					];
+				}
+				if (!TextureParameter->ChannelNames.A.IsEmpty())
+				{
+					NameVerticalBox->AddSlot()
+					[
+						SNew(SHorizontalBox)
+						+ SHorizontalBox::Slot()
+						.Padding(20.0, 2.0, 4.0, 2.0)
+						.AutoWidth()
+						[
+							SNew(STextBlock)
+							.Text(FText::FromName(Alpha))
+							.Font(FAppStyle::GetFontStyle(TEXT("PropertyWindow.BoldFont")))
+						]
+						+ SHorizontalBox::Slot()
+						.HAlign(HAlign_Left)
+						.Padding(4.0, 2.0)
+						[
+							SNew(STextBlock)
+							.Text(TextureParameter->ChannelNames.A)
+							.Font(FAppStyle::GetFontStyle(TEXT("PropertyWindow.NormalFont")))
+						]
+					];
 				}
-			}
-			else
-			{
-				DetailLayout.HideProperty(ChannelHandle);
 			}
 		}
-
 	}
 
 	check(ScalarParameterObjects.Num() == DefaultValueHandles.Num());
diff --git a/Engine/Source/Runtime/Engine/Private/Materials/MaterialExpressions.cpp b/Engine/Source/Runtime/Engine/Private/Materials/MaterialExpressions.cpp
--- a/Engine/Source/Runtime/Engine/Private/Materials/MaterialExpressions.cpp
+++ b/Engine/Source/Runtime/Engine/Private/Materials/MaterialExpressions.cpp
@@ -2481,6 +2481,8 @@ UMaterialExpressionTextureSample::UMaterialExpressionTextureSample(const FObject
 	ConstCoordinate = 0;
 	ConstMipValue = INDEX_NONE;
 	AutomaticViewMipBias = true;
+
+	ApplyChannelNames();
 #endif // WITH_EDITORONLY_DATA
 }
 
@@ -2532,6 +2534,16 @@ void UMaterialExpressionTextureSample::PostEditChangeProperty(FPropertyChangedEv
 			GraphNode->ReconstructNode();
 		}
 	}
+
+	if (PropertyChangedEvent.GetMemberPropertyName() == GET_MEMBER_NAME_STRING_CHECKED(ThisClass, ChannelNames))
+	{
+		ApplyChannelNames();
+
+		if (GraphNode)
+		{
+			GraphNode->ReconstructNode();
+		}
+	}
 	
 	// Need to update expression properties before super call (which triggers recompile)
 	Super::PostEditChangeProperty( PropertyChangedEvent );	
@@ -2548,6 +2560,21 @@ void UMaterialExpressionTextureSample::PostLoad()
 	}
 }
 
+void UMaterialExpressionTextureSample::ApplyChannelNames()
+{
+	static const FName Red("R");
+	static const FName Green("G");
+	static const FName Blue("B");
+	static const FName Alpha("A");
+	if (GetOutputType(0) != MCT_Texture)
+	{
+		Outputs[1].OutputName = !ChannelNames.R.IsEmpty() ? FName(*ChannelNames.R.ToString()) : Red;
+		Outputs[2].OutputName = !ChannelNames.G.IsEmpty() ? FName(*ChannelNames.G.ToString()) : Green;
+		Outputs[3].OutputName = !ChannelNames.B.IsEmpty() ? FName(*ChannelNames.B.ToString()) : Blue;
+		Outputs[4].OutputName = !ChannelNames.A.IsEmpty() ? FName(*ChannelNames.A.ToString()) : Alpha;
+	}
+}
+
 TArrayView<FExpressionInput*> UMaterialExpressionTextureSample::GetInputsView()
 {
 	CachedInputs.Empty();
@@ -3854,8 +3881,6 @@ UMaterialExpressionTextureSampleParameter::UMaterialExpressionTextureSampleParam
 #if WITH_EDITORONLY_DATA
 	bIsParameterExpression = true;
 	bShowTextureInputPin = false;
-
-	ApplyChannelNames();
 #endif
 }
 
@@ -3917,36 +3942,6 @@ bool UMaterialExpressionTextureSampleParameter::SetParameterValue(FName InParame
 	return false;
 }
 
-void UMaterialExpressionTextureSampleParameter::PostEditChangeProperty(FPropertyChangedEvent& PropertyChangedEvent)
-{
-	if (PropertyChangedEvent.GetMemberPropertyName() == GET_MEMBER_NAME_STRING_CHECKED(ThisClass, ChannelNames))
-	{
-		ApplyChannelNames();
-
-		if (GraphNode)
-		{
-			GraphNode->ReconstructNode();
-		}
-	}
-
-	Super::PostEditChangeProperty(PropertyChangedEvent);
-}
-
-void UMaterialExpressionTextureSampleParameter::ApplyChannelNames()
-{
-	static const FName Red("R");
-	static const FName Green("G");
-	static const FName Blue("B");
-	static const FName Alpha("A");
-	if (GetOutputType(0) != MCT_Texture)
-	{
-		Outputs[1].OutputName = !ChannelNames.R.IsEmpty() ? FName(*ChannelNames.R.ToString()) : Red;
-		Outputs[2].OutputName = !ChannelNames.G.IsEmpty() ? FName(*ChannelNames.G.ToString()) : Green;
-		Outputs[3].OutputName = !ChannelNames.B.IsEmpty() ? FName(*ChannelNames.B.ToString()) : Blue;
-		Outputs[4].OutputName = !ChannelNames.A.IsEmpty() ? FName(*ChannelNames.A.ToString()) : Alpha;
-	}
-}
-
 bool UMaterialExpressionTextureSampleParameter::TextureIsValid(UTexture* /*InTexture*/, FString& OutMessage)
 {
 	OutMessage = TEXT("Invalid texture type");
diff --git a/Engine/Source/Runtime/Engine/Public/Materials/MaterialExpressionTextureSample.h b/Engine/Source/Runtime/Engine/Public/Materials/MaterialExpressionTextureSample.h
--- a/Engine/Source/Runtime/Engine/Public/Materials/MaterialExpressionTextureSample.h
+++ b/Engine/Source/Runtime/Engine/Public/Materials/MaterialExpressionTextureSample.h
@@ -6,6 +6,7 @@
 #include "CoreMinimal.h"
 #include "UObject/ObjectMacros.h"
 #include "MaterialExpressionIO.h"
+#include "MaterialTypes.h"
 #include "Materials/MaterialExpressionTextureBase.h"
 #include "MaterialExpressionTextureSample.generated.h"
 
@@ -79,6 +80,9 @@ class UMaterialExpressionTextureSample : public UMaterialExpressionTextureBase
 	UPROPERTY(EditAnywhere, Category = MaterialExpressionTextureSample)
 	uint8 AutomaticViewMipBias : 1;
 
+	UPROPERTY(EditAnywhere, Category = ParameterCustomization)
+	FParameterChannelNames ChannelNames;
+
 protected:
 	// Inherited parameter expressions can hide unused input pin
 	uint8 bShowTextureInputPin : 1;
@@ -101,6 +105,12 @@ public:
 	ENGINE_API virtual bool CanEditChange(const FProperty* InProperty) const override;
 	ENGINE_API virtual void PostEditChangeProperty(FPropertyChangedEvent& PropertyChangedEvent) override;
 	ENGINE_API virtual void PostLoad() override;
+
+	ENGINE_API void ApplyChannelNames();
+	FParameterChannelNames GetTextureChannelNames() const
+	{
+		return ChannelNames;
+	}
 #endif // WITH_EDITOR
 	//~ End UObject Interface
 
diff --git a/Engine/Source/Runtime/Engine/Public/Materials/MaterialExpressionTextureSampleParameter.h b/Engine/Source/Runtime/Engine/Public/Materials/MaterialExpressionTextureSampleParameter.h
--- a/Engine/Source/Runtime/Engine/Public/Materials/MaterialExpressionTextureSampleParameter.h
+++ b/Engine/Source/Runtime/Engine/Public/Materials/MaterialExpressionTextureSampleParameter.h
@@ -6,7 +6,6 @@
 #include "CoreMinimal.h"
 #include "UObject/ObjectMacros.h"
 #include "Misc/Guid.h"
-#include "MaterialTypes.h"
 #include "Materials/MaterialExpressionTextureSample.h"
 #include "MaterialExpressionTextureSampleParameter.generated.h"
 
@@ -33,9 +32,6 @@ class UMaterialExpressionTextureSampleParameter : public UMaterialExpressionText
 	UPROPERTY(EditAnywhere, Category = MaterialExpressionTextureSampleParameter)
 	int32 SortPriority = 32;
 
-	UPROPERTY(EditAnywhere, Category = ParameterCustomization)
-	FParameterChannelNames ChannelNames;
-
 	//~ Begin UMaterialExpression Interface
 #if WITH_EDITOR
 	ENGINE_API virtual void Build(UE::MIR::FEmitter& Emitter) override;
@@ -58,7 +54,7 @@ class UMaterialExpressionTextureSampleParameter : public UMaterialExpressionText
 		OutMeta.Group = Group;
 		OutMeta.SortPriority = SortPriority;
 		OutMeta.AssetPath = GetAssetPathName();
-		OutMeta.ChannelNames = ChannelNames;
+		OutMeta.ChannelNames = GetTextureChannelNames();
 		return true;
 	}
 	virtual bool SetParameterValue(const FName& Name, const FMaterialParameterMetadata& Meta, EMaterialExpressionSetParameterValueFlags Flags) override
@@ -82,8 +78,6 @@ class UMaterialExpressionTextureSampleParameter : public UMaterialExpressionText
 
 #if WITH_EDITOR
 	ENGINE_API bool SetParameterValue(FName InParameterName, UTexture* InValue, EMaterialExpressionSetParameterValueFlags Flags = EMaterialExpressionSetParameterValueFlags::None);
-	ENGINE_API virtual void PostEditChangeProperty(FPropertyChangedEvent& PropertyChangedEvent) override;
-	ENGINE_API void ApplyChannelNames();
 
 	/**
 	 * Return true if the texture is a movie texture
@@ -104,11 +98,4 @@ class UMaterialExpressionTextureSampleParameter : public UMaterialExpressionText
 	{
 		return ExpressionGUID;
 	}
-
-#if WITH_EDITOR
-	FParameterChannelNames GetTextureChannelNames() const
-	{
-		return ChannelNames;
-	}
-#endif
 };
