Implement "Create Reroute Usage Node From Reroute Declaration" and expose in the right click menu for Reroute Declaration nodes.

- The function gets the currently selected node, casts to a UMaterialExpressionNamedRerouteDeclaration, and creates a new UMaterialExpressionNamedRerouteUsage node, which is 150 units to the right of the declaration node
- Zooms the display to fit the newly created node.
- Setup the properties on the usage node which link back to the declaration.

[RN] minor | Rendering - Materials and Shaders

diff --git a/Engine/Source/Editor/MaterialEditor/Private/MaterialEditor.cpp b/Engine/Source/Editor/MaterialEditor/Private/MaterialEditor.cpp
--- a/Engine/Source/Editor/MaterialEditor/Private/MaterialEditor.cpp
+++ b/Engine/Source/Editor/MaterialEditor/Private/MaterialEditor.cpp
@@ -3821,6 +3821,10 @@ void FMaterialEditor::BindCommands()
 	ToolkitCommands->MapAction(
 		Commands.SelectNamedRerouteUsages,
 		FExecuteAction::CreateSP(this, &FMaterialEditor::OnSelectNamedRerouteUsages));
+
+	ToolkitCommands->MapAction(
+		Commands.CreateRerouteUsageFromDeclaration,
+		FExecuteAction::CreateSP(this, &FMaterialEditor::OnCreateRerouteUsageFromDeclaration));
 	
 	ToolkitCommands->MapAction(
 		Commands.ConvertRerouteToNamedReroute,
@@ -4777,6 +4781,41 @@ void FMaterialEditor::OnSelectNamedRerouteUsages()
 	}
 }
 
+void FMaterialEditor::OnCreateRerouteUsageFromDeclaration()
+{
+	if (TSharedPtr<SGraphEditor> FocusedGraphEd = FocusedGraphEdPtr.Pin())
+	{
+		const FGraphPanelSelectionSet SelectedNodes = FocusedGraphEd->GetSelectedNodes();
+		if (SelectedNodes.Num() == 1)
+		{
+			FocusedGraphEd->ClearSelectionSet();
+			for (FGraphPanelSelectionSet::TConstIterator NodeIt(SelectedNodes); NodeIt; ++NodeIt)
+			{
+				UMaterialGraphNode* GraphNode = Cast<UMaterialGraphNode>(*NodeIt);
+				if (GraphNode)
+				{
+					UMaterialExpression* CurrentSelectedExpression = GraphNode->MaterialExpression;
+					UMaterialExpressionNamedRerouteDeclaration* Declaration = Cast<UMaterialExpressionNamedRerouteDeclaration>(CurrentSelectedExpression);
+
+					if (Declaration)
+					{
+						UMaterialExpressionNamedRerouteUsage* RerouteUsage = Cast<UMaterialExpressionNamedRerouteUsage>(FMaterialEditorUtilities::CreateNewMaterialExpression(
+							GraphNode->GetGraph(),
+							UMaterialExpressionNamedRerouteUsage::StaticClass(),
+							FVector2D(GraphNode->NodePosX + 150, GraphNode->NodePosY),
+							true,
+							false));
+
+						RerouteUsage->Declaration = Declaration;
+						RerouteUsage->DeclarationGuid = Declaration->VariableGuid;
+					}
+				}
+			}
+			FocusedGraphEd->ZoomToFit(true);
+		}
+	}
+}
+
 void FMaterialEditor::OnConvertRerouteToNamedReroute()
 {
 	if (TSharedPtr<SGraphEditor> FocusedGraphEd = FocusedGraphEdPtr.Pin())
@@ -7762,6 +7801,10 @@ TSharedRef<SGraphEditor> FMaterialEditor::CreateGraphEditorWidget(TSharedRef<cla
 			FExecuteAction::CreateSP(this, &FMaterialEditor::OnSelectNamedRerouteUsages)
 			);
 
+		GraphEditorCommands->MapAction( FMaterialEditorCommands::Get().CreateRerouteUsageFromDeclaration,
+			FExecuteAction::CreateSP(this, &FMaterialEditor::OnCreateRerouteUsageFromDeclaration)
+			);
+
 		GraphEditorCommands->MapAction( FMaterialEditorCommands::Get().ConvertRerouteToNamedReroute,
 			FExecuteAction::CreateSP(this, &FMaterialEditor::OnConvertRerouteToNamedReroute)
 			);
diff --git a/Engine/Source/Editor/MaterialEditor/Private/MaterialEditor.h b/Engine/Source/Editor/MaterialEditor/Private/MaterialEditor.h
--- a/Engine/Source/Editor/MaterialEditor/Private/MaterialEditor.h
+++ b/Engine/Source/Editor/MaterialEditor/Private/MaterialEditor.h
@@ -801,6 +801,8 @@ private:
 	void OnSelectNamedRerouteDeclaration();
 	/** Command to select local variable usages */
 	void OnSelectNamedRerouteUsages();
+	/** Command to create reroute usage node from declaration */
+	void OnCreateRerouteUsageFromDeclaration();
 	/** Command to convert a reroute node to local variables */
 	void OnConvertRerouteToNamedReroute();
 	/** Command to convert local variables to a reroute node */
diff --git a/Engine/Source/Editor/MaterialEditor/Private/MaterialEditorActions.cpp b/Engine/Source/Editor/MaterialEditor/Private/MaterialEditorActions.cpp
--- a/Engine/Source/Editor/MaterialEditor/Private/MaterialEditorActions.cpp
+++ b/Engine/Source/Editor/MaterialEditor/Private/MaterialEditorActions.cpp
@@ -56,6 +56,7 @@ void FMaterialEditorCommands::RegisterCommands()
 	UI_COMMAND( ConvertToConstant, "Convert to Constant", "Converts the parameters to constants.", EUserInterfaceActionType::Button, FInputChord() );
 	UI_COMMAND( SelectNamedRerouteDeclaration, "Select Named Reroute Declaration", "Select this named reroute's declaration node.", EUserInterfaceActionType::Button, FInputChord() );
 	UI_COMMAND( SelectNamedRerouteUsages, "Select Named Reroute Usages", "Select this named reroute's usage nodes.", EUserInterfaceActionType::Button, FInputChord() );
+	UI_COMMAND( CreateRerouteUsageFromDeclaration, "Create Reroute Usage", "Create Reroute Usage From Declaration.", EUserInterfaceActionType::Button, FInputChord() );
 	UI_COMMAND( ConvertRerouteToNamedReroute, "Convert to Named Reroute", "Replace this reroute with a named reroute.", EUserInterfaceActionType::Button, FInputChord() );
 	UI_COMMAND( ConvertNamedRerouteToReroute, "Convert to Reroute", "Replace this named reroute with an (unnamed) reroute node.", EUserInterfaceActionType::Button, FInputChord() );
 	UI_COMMAND( ConvertToTextureObjects, "Convert to Texture Object", "Converts the objects to texture objects.", EUserInterfaceActionType::Button, FInputChord() );
diff --git a/Engine/Source/Editor/MaterialEditor/Public/MaterialEditorActions.h b/Engine/Source/Editor/MaterialEditor/Public/MaterialEditorActions.h
--- a/Engine/Source/Editor/MaterialEditor/Public/MaterialEditorActions.h
+++ b/Engine/Source/Editor/MaterialEditor/Public/MaterialEditorActions.h
@@ -137,6 +137,7 @@ public:
 	/** Local variables select helpers */
 	TSharedPtr< FUICommandInfo > SelectNamedRerouteDeclaration;
 	TSharedPtr< FUICommandInfo > SelectNamedRerouteUsages;
+	TSharedPtr<FUICommandInfo> CreateRerouteUsageFromDeclaration;
 
 	/** Conversion between reroute nodes and local variables */
 	TSharedPtr< FUICommandInfo > ConvertRerouteToNamedReroute;
diff --git a/Engine/Source/Editor/UnrealEd/Private/MaterialGraphNode.cpp b/Engine/Source/Editor/UnrealEd/Private/MaterialGraphNode.cpp
--- a/Engine/Source/Editor/UnrealEd/Private/MaterialGraphNode.cpp
+++ b/Engine/Source/Editor/UnrealEd/Private/MaterialGraphNode.cpp
@@ -403,6 +403,7 @@ void UMaterialGraphNode::GetNodeContextMenuActions(UToolMenu* Menu, UGraphNodeCo
 				if (MaterialExpression->IsA(UMaterialExpressionNamedRerouteDeclaration::StaticClass()))
 				{
 					NodeSection.AddMenuEntry(FMaterialEditorCommands::Get().SelectNamedRerouteUsages);
+					NodeSection.AddMenuEntry(FMaterialEditorCommands::Get().CreateRerouteUsageFromDeclaration);
 				}
 				if (MaterialExpression->IsA(UMaterialExpressionNamedRerouteUsage::StaticClass()))
 				{
