"Select Named Reroute Usages" right click functionality will now display the nodes in the "Find Results".

- Updated the SFindInMaterial widget to allow for specifying a list of UEdGraphNode nodes to display directly.
- Attempt to spawn the find results tab when the user invokes this functionality (in case they don' t have the tab open)
- Clear the search text field in case there was something already there.

[RN] minor | Rendering - Materials and Shaders

diff --git a/Engine/Source/Editor/MaterialEditor/Private/FindInMaterial.cpp b/Engine/Source/Editor/MaterialEditor/Private/FindInMaterial.cpp
--- a/Engine/Source/Editor/MaterialEditor/Private/FindInMaterial.cpp
+++ b/Engine/Source/Editor/MaterialEditor/Private/FindInMaterial.cpp
@@ -184,6 +184,35 @@ void SFindInMaterial::FocusForUse()
 	FSlateApplication::Get().SetKeyboardFocus(FilterTextBoxWidgetPath, EFocusCause::SetDirectly);
 }
 
+void SFindInMaterial::PopulateSearchItems(const TArray<UEdGraphNode*>& Nodes)
+{
+	if (Nodes.Num() == 0)
+	{
+		return;
+	}
+
+	// Reset to default state.
+	SearchTextField->SetText(FText::GetEmpty());
+	RootSearchResult.Reset();
+	RootSearchResult = FSearchResult(new FFindInMaterialResult(FString("MaterialTreeRoot")));
+	ItemsFound.Empty();
+
+	for (UEdGraphNode* Node : Nodes)
+	{
+		const FString NodeName = Node->GetNodeTitle(ENodeTitleType::FullTitle).ToString();
+		const FString NodeType = Node->GetNodeTitle(ENodeTitleType::ListView).ToString();
+		ItemsFound.Add(FSearchResult(new FFindInMaterialResult(NodeName == NodeType ? NodeName : NodeName + " - " + NodeType, RootSearchResult, Node)));
+	}
+
+	// Refresh the tree and expand items.
+	TreeView->RequestTreeRefresh();
+
+	for (auto It(ItemsFound.CreateIterator()); It; ++It)
+	{
+		TreeView->SetItemExpansion(*It, true);
+	}
+}
+
 void SFindInMaterial::OnSearchTextChanged(const FText& Text)
 {
 	SearchValue = Text.ToString();
diff --git a/Engine/Source/Editor/MaterialEditor/Private/MaterialEditor.cpp b/Engine/Source/Editor/MaterialEditor/Private/MaterialEditor.cpp
--- a/Engine/Source/Editor/MaterialEditor/Private/MaterialEditor.cpp
+++ b/Engine/Source/Editor/MaterialEditor/Private/MaterialEditor.cpp
@@ -4752,31 +4752,41 @@ void FMaterialEditor::OnSelectNamedRerouteUsages()
 	if (TSharedPtr<SGraphEditor> FocusedGraphEd = FocusedGraphEdPtr.Pin())
 	{
 		const FGraphPanelSelectionSet SelectedNodes = FocusedGraphEd->GetSelectedNodes();
-	if (SelectedNodes.Num() == 1)
-	{
-			FocusedGraphEd->ClearSelectionSet();
-		for (FGraphPanelSelectionSet::TConstIterator NodeIt(SelectedNodes); NodeIt; ++NodeIt)
+		if (SelectedNodes.Num() == 1)
 		{
-			UMaterialGraphNode* GraphNode = Cast<UMaterialGraphNode>(*NodeIt);
-			if (GraphNode)
+			FocusedGraphEd->ClearSelectionSet();
+			for (FGraphPanelSelectionSet::TConstIterator NodeIt(SelectedNodes); NodeIt; ++NodeIt)
 			{
-				UMaterialExpression* CurrentSelectedExpression = GraphNode->MaterialExpression;
-				UMaterialExpressionNamedRerouteDeclaration* Declaration = Cast<UMaterialExpressionNamedRerouteDeclaration>(CurrentSelectedExpression);
-				for(UMaterialExpression* Expression : Material->GetExpressions())
+				UMaterialGraphNode* GraphNode = Cast<UMaterialGraphNode>(*NodeIt);
+				if (GraphNode)
 				{
-					auto* Usage = Cast<UMaterialExpressionNamedRerouteUsage>(Expression);
-					if (Usage && Usage->Declaration == Declaration)
+					TArray<UEdGraphNode*> FoundGraphNodes;
+
+					UMaterialExpression* CurrentSelectedExpression = GraphNode->MaterialExpression;
+					UMaterialExpressionNamedRerouteDeclaration* Declaration = Cast<UMaterialExpressionNamedRerouteDeclaration>(CurrentSelectedExpression);
+					for(UMaterialExpression* Expression : Material->GetExpressions())
 					{
-						UEdGraphNode* UsageGraphNode = Usage->GraphNode;
-						if (UsageGraphNode)
+						auto* Usage = Cast<UMaterialExpressionNamedRerouteUsage>(Expression);
+						if (Usage && Usage->Declaration == Declaration)
 						{
-								FocusedGraphEd->SetNodeSelection(UsageGraphNode, true);
+							UEdGraphNode* UsageGraphNode = Usage->GraphNode;
+							if (UsageGraphNode)
+							{
+								FoundGraphNodes.Add(UsageGraphNode);
 							}
 						}
 					}
+					
+					if (FoundGraphNodes.Num())
+					{
+						// Spawn the tab in case the user doesn't have it open
+						TabManager->TryInvokeTab(FMaterialEditorTabs::FindTabId);
+
+						FindResults->PopulateSearchItems(FoundGraphNodes);
+						FindResults->FocusForUse();
+					}
 				}
 			}
-			FocusedGraphEd->ZoomToFit(true);
 		}
 	}
 }
diff --git a/Engine/Source/Editor/MaterialEditor/Public/FindInMaterial.h b/Engine/Source/Editor/MaterialEditor/Public/FindInMaterial.h
--- a/Engine/Source/Editor/MaterialEditor/Public/FindInMaterial.h
+++ b/Engine/Source/Editor/MaterialEditor/Public/FindInMaterial.h
@@ -83,6 +83,9 @@ public:
 	/** Focuses this widget's search box */
 	void FocusForUse();
 
+	/** Clears out the existing search items and repopulates the list with these nodes */
+	void PopulateSearchItems(const TArray<UEdGraphNode*>& Nodes);
+
 protected:
 	typedef TSharedPtr<FFindInMaterialResult> FSearchResult;
 	typedef STreeView<FSearchResult> STreeViewType;
