Implement "JumpToDefinition" for all material graph nodes so we can jump from the node in the graph to the C++ header.

- Previously JumpToDefinition and CanJumpToDefinition was only implemented for Composite nodes.  The jump to definition, which in the editor typically means jumping to the C++ header like in blueprints, was repurposed to open up composite nodes to show what was inside.
- This change now makes it so that this functionality will jump you from the Material Expression to the C++ header.
- Added `GoToDefinition` to each node's context menu, this command hooks into two functions on the FMaterialEditor which will do the "jump to definition" for the selected node.
- Added a simple helper function `GetFirstSelectedGraphNode` to get the first currently selected node.
- Removed old behavior in the Material Editor where double clicking a node would cause it to call JumpToDefinition.  JumpToDefinition was previously only implemented for the composite node, but now that every node implements it, we don't want double click to open up the C++ header.
- To maintain the old behavior of the composite node I've explicitly added double click behavior handling.

[RN] major | Rendering - Materials and Shaders

diff --git a/Engine/Source/Editor/MaterialEditor/Private/MaterialEditor.cpp b/Engine/Source/Editor/MaterialEditor/Private/MaterialEditor.cpp
--- a/Engine/Source/Editor/MaterialEditor/Private/MaterialEditor.cpp
+++ b/Engine/Source/Editor/MaterialEditor/Private/MaterialEditor.cpp
@@ -5632,6 +5632,27 @@ bool FMaterialEditor::CanGoToDocumentation()
 	return !DocumentationLink.IsEmpty();
 }
 
+void FMaterialEditor::OnGoToDefinition()
+{
+	if (UMaterialGraphNode* Node = GetFirstSelectedGraphNode())
+	{
+		if (Node->CanJumpToDefinition())
+		{
+			Node->JumpToDefinition();
+		}
+	}
+}
+
+bool FMaterialEditor::CanGoToDefinition() const
+{
+	if (UMaterialGraphNode* Node = GetFirstSelectedGraphNode())
+	{
+		return Node->CanJumpToDefinition();
+	}
+
+	return false;
+}
+
 void FMaterialEditor::RenameAssetFromRegistry(const FAssetData& InAddedAssetData, const FString& InNewName)
 {
 	// Grab the asset class, it will be checked for being a material function.
@@ -6912,6 +6933,20 @@ void FMaterialEditor::CreateDerivedMaterialInstancesPreviews()
 	MaterialStatsManager->SetMaterialsDisplayNames(DisplayNames);
 }
 
+UMaterialGraphNode* FMaterialEditor::GetFirstSelectedGraphNode() const
+{
+	const FGraphPanelSelectionSet SelectedNodes = GetSelectedNodes();
+	for (FGraphPanelSelectionSet::TConstIterator NodeIt(SelectedNodes); NodeIt; ++NodeIt)
+	{
+		if (UMaterialGraphNode* GraphNode = Cast<UMaterialGraphNode>(*NodeIt))
+		{
+			return GraphNode;
+		}
+	}
+
+	return nullptr;
+}
+
 void FMaterialEditor::UpdateMaterialAfterGraphChange()
 {
 	FlushRenderingCommands();
@@ -7863,6 +7898,11 @@ TSharedRef<SGraphEditor> FMaterialEditor::CreateGraphEditorWidget(TSharedRef<cla
 			FExecuteAction::CreateSP(this, &FMaterialEditor::OnCreateComponentMaskNode)
 			);
 
+		GraphEditorCommands->MapAction(FGraphEditorCommands::Get().GoToDefinition,
+									   FExecuteAction::CreateSP(this, &FMaterialEditor::OnGoToDefinition),
+									   FCanExecuteAction::CreateSP(this, &FMaterialEditor::CanGoToDefinition)
+									   );
+
 		GraphEditorCommands->MapAction(FGraphEditorCommands::Get().GoToDocumentation,
 			FExecuteAction::CreateSP(this, &FMaterialEditor::OnGoToDocumentation),
 			FCanExecuteAction::CreateSP(this, &FMaterialEditor::CanGoToDocumentation)
@@ -8289,11 +8329,7 @@ void FMaterialEditor::OnNodeDoubleClicked(class UEdGraphNode* Node)
 {
 	UMaterialGraphNode* GraphNode = Cast<UMaterialGraphNode>(Node);
 
-	if (Node && Node->CanJumpToDefinition())
-	{
-		Node->JumpToDefinition();
-	}
-	else if (GraphNode && GraphNode->MaterialExpression)
+	if (GraphNode && GraphNode->MaterialExpression)
 	{
 		TOptional<FLinearColor> InitialColor;
 		bool bCanEditAlpha = true;
@@ -8379,6 +8415,19 @@ void FMaterialEditor::OnNodeDoubleClicked(class UEdGraphNode* Node)
 		{
 			OnSelectNamedRerouteDeclaration();
 		}
+		else if (GraphNode->MaterialExpression->IsA<UMaterialExpressionComposite>())
+		{
+			if (UMaterialGraphNode_Composite* CompositeGraphNode = Cast<UMaterialGraphNode_Composite>(GraphNode))
+			{
+				if (CompositeGraphNode->GetJumpTargetForDoubleClick())
+				{
+					if (UObject* HyperlinkTarget = CompositeGraphNode->GetJumpTargetForDoubleClick())
+					{
+						FMaterialEditorUtilities::BringFocusAttentionOnObject(HyperlinkTarget);
+					}
+				}
+			}
+		}
 	}
 }
 
diff --git a/Engine/Source/Editor/MaterialEditor/Private/MaterialEditor.h b/Engine/Source/Editor/MaterialEditor/Private/MaterialEditor.h
--- a/Engine/Source/Editor/MaterialEditor/Private/MaterialEditor.h
+++ b/Engine/Source/Editor/MaterialEditor/Private/MaterialEditor.h
@@ -729,6 +729,9 @@ private:
 	/** Create array of derived material instances used in conjunction with preview material in material stats */
 	void CreateDerivedMaterialInstancesPreviews();
 
+	/** Returns nullptr if nothing selected */
+	UMaterialGraphNode* GetFirstSelectedGraphNode() const;
+
 public:
 private:
 	/**
@@ -857,6 +860,11 @@ private:
 	/** Can we open documentation for the selected node */
 	bool CanGoToDocumentation();
 
+	/** Open C++ header for the selected node class */
+	void OnGoToDefinition();
+	/** Can we open the C++ header for the selected node */
+	bool CanGoToDefinition() const;
+
 	/** Util to try and get doc link for the currently selected node */
 	FString GetDocLinkForSelectedNode();
 
diff --git a/Engine/Source/Editor/UnrealEd/Classes/MaterialGraph/MaterialGraphNode.h b/Engine/Source/Editor/UnrealEd/Classes/MaterialGraph/MaterialGraphNode.h
--- a/Engine/Source/Editor/UnrealEd/Classes/MaterialGraph/MaterialGraphNode.h
+++ b/Engine/Source/Editor/UnrealEd/Classes/MaterialGraph/MaterialGraphNode.h
@@ -74,6 +74,8 @@ public:
 	virtual uint32 GetPinMaterialType(const UEdGraphPin* Pin) const override;
 	virtual void PinDefaultValueChanged(UEdGraphPin* Pin) override;
 	virtual TSharedPtr<SGraphNode> CreateVisualWidget() override;
+	virtual bool CanJumpToDefinition() const override;
+	virtual void JumpToDefinition() const override;
 	//~ End UEdGraphNode Interface.
 
 	//~ Begin UMaterialGraphNode_Base Interface
diff --git a/Engine/Source/Editor/UnrealEd/Classes/MaterialGraph/MaterialGraphNode_Composite.h b/Engine/Source/Editor/UnrealEd/Classes/MaterialGraph/MaterialGraphNode_Composite.h
--- a/Engine/Source/Editor/UnrealEd/Classes/MaterialGraph/MaterialGraphNode_Composite.h
+++ b/Engine/Source/Editor/UnrealEd/Classes/MaterialGraph/MaterialGraphNode_Composite.h
@@ -28,8 +28,6 @@ class UMaterialGraphNode_Composite : public UMaterialGraphNode
 
 	//~ Begin UEdGraphNode Interface.
 	virtual UObject* GetJumpTargetForDoubleClick() const override;
-	virtual bool CanJumpToDefinition() const override;
-	virtual void JumpToDefinition() const override;
 	virtual void DestroyNode() override; 
 	virtual void PrepareForCopying() override;
 	virtual void PostPasteNode() override; 
diff --git a/Engine/Source/Editor/UnrealEd/Private/MaterialGraphNode.cpp b/Engine/Source/Editor/UnrealEd/Private/MaterialGraphNode.cpp
--- a/Engine/Source/Editor/UnrealEd/Private/MaterialGraphNode.cpp
+++ b/Engine/Source/Editor/UnrealEd/Private/MaterialGraphNode.cpp
@@ -53,6 +53,7 @@
 #include "Framework/Commands/GenericCommands.h"
 #include "ScopedTransaction.h"
 #include "MaterialNodes/SGraphNodeMaterialBase.h"
+#include "SourceCodeNavigation.h"
 
 #define LOCTEXT_NAMESPACE "MaterialGraphNode"
 
@@ -488,6 +489,11 @@ void UMaterialGraphNode::GetNodeContextMenuActions(UToolMenu* Menu, UGraphNodeCo
 		// Break all links
 		{
 			NodeSection.AddMenuEntry(FGraphEditorCommands::Get().BreakNodeLinks);
+
+			if (Context->Node->CanJumpToDefinition())
+			{
+				NodeSection.AddMenuEntry(FGraphEditorCommands::Get().GoToDefinition);
+			}
 		}
 
 		// Separate the above frequently used options from the below less frequently used common options
@@ -650,6 +656,16 @@ TSharedPtr<SGraphNode> UMaterialGraphNode::CreateVisualWidget()
 	return SNew(SGraphNodeMaterialBase, this);
 }
 
+bool UMaterialGraphNode::CanJumpToDefinition() const
+{
+	return true;
+}
+
+void UMaterialGraphNode::JumpToDefinition() const
+{
+	FSourceCodeNavigation::NavigateToClass(MaterialExpression->GetClass());
+}
+
 void UMaterialGraphNode::CreateInputPins()
 {
 	int32 InputIndex = 0;
diff --git a/Engine/Source/Editor/UnrealEd/Private/MaterialGraphNode_Composite.cpp b/Engine/Source/Editor/UnrealEd/Private/MaterialGraphNode_Composite.cpp
--- a/Engine/Source/Editor/UnrealEd/Private/MaterialGraphNode_Composite.cpp
+++ b/Engine/Source/Editor/UnrealEd/Private/MaterialGraphNode_Composite.cpp
@@ -59,19 +59,6 @@ UObject* UMaterialGraphNode_Composite::GetJumpTargetForDoubleClick() const
 	return BoundGraph;
 }
 
-bool UMaterialGraphNode_Composite::CanJumpToDefinition() const
-{
-	return GetJumpTargetForDoubleClick() != nullptr;
-}
-
-void UMaterialGraphNode_Composite::JumpToDefinition() const
-{
-	if (UObject* HyperlinkTarget = GetJumpTargetForDoubleClick())
-	{
-		FMaterialEditorUtilities::BringFocusAttentionOnObject(HyperlinkTarget);
-	}
-}
-
 void UMaterialGraphNode_Composite::DestroyNode()
 {
 	UEdGraph* GraphToRemove = BoundGraph;
