Material Editor: Add command to preview the currently-hovered node, without having to select it.
- Default its shortcut to Shift + Spacebar

diff --git a/Engine/Source/Editor/GraphEditor/Public/SNodePanel.h b/Engine/Source/Editor/GraphEditor/Public/SNodePanel.h
--- a/Engine/Source/Editor/GraphEditor/Public/SNodePanel.h
+++ b/Engine/Source/Editor/GraphEditor/Public/SNodePanel.h
@@ -430,6 +430,34 @@ public:
 
 		typedef TSet<TWeakPtr<SNodePanel::SNode>> FNodeSet;
 
+		//~ Begin SWidget Interface
+		virtual void OnMouseEnter(const FGeometry& InGeometry, const FPointerEvent& InMouseEvent) override
+		{
+			SPanel::OnMouseEnter(InGeometry, InMouseEvent);
+
+			if (TSharedPtr<SNodePanel> ParentPanel = GetParentPanel())
+			{
+				ParentPanel->SetCurrentHoveredNode(GetObjectBeingDisplayed());
+			}
+		}
+
+		virtual void OnMouseLeave(const FPointerEvent& InMouseEvent) override
+		{
+			SPanel::OnMouseLeave(InMouseEvent);
+
+			if (TSharedPtr<SNodePanel> ParentPanel = GetParentPanel())
+			{
+				UObject* PreviousHoveredNode = ParentPanel->GetCurrentHoveredNode();
+				
+				// Only unset the current hovered node on mouse leave if its this node
+				if (PreviousHoveredNode == GetObjectBeingDisplayed())
+				{
+					ParentPanel->SetCurrentHoveredNode(nullptr);
+				}
+			}
+		}
+		//~ End SWidget Interface
+
 		// SPanel Interface
 		virtual FChildren* GetChildren() override
 		{
@@ -986,6 +1014,9 @@ public:
 		ZoomLevel = PreviousZoomLevel = ZoomLevels->GetNearestZoomLevel(OldZoomAmount);
 		PostChangedZoom();
 	}
+
+	void SetCurrentHoveredNode(UObject* InNode) { CurrentHoveredNode = InNode; }
+	UObject* GetCurrentHoveredNode() { return CurrentHoveredNode; }
 	
 protected:
 	
@@ -1037,7 +1068,11 @@ protected:
 public:
 	/** Nodes selected in this instance of the editor; the selection is per-instance of the GraphEditor */
 	FGraphSelectionManager SelectionManager;
+
 protected:
+	/** The currently hovered SNode's UObject */
+	TObjectPtr<class UObject> CurrentHoveredNode;
+
 	/** A pending marquee operation if it's active */
 	FMarqueeOperation Marquee;
 
diff --git a/Engine/Source/Editor/MaterialEditor/Private/MaterialEditor.cpp b/Engine/Source/Editor/MaterialEditor/Private/MaterialEditor.cpp
--- a/Engine/Source/Editor/MaterialEditor/Private/MaterialEditor.cpp
+++ b/Engine/Source/Editor/MaterialEditor/Private/MaterialEditor.cpp
@@ -129,7 +129,7 @@
 #include "Kismet2/BlueprintEditorUtils.h"
 #include "Logging/TokenizedMessage.h"
 #include "EdGraphUtilities.h"
-#include "SNodePanel.h"
+#include "SGraphPanel.h"
 #include "MaterialEditorUtilities.h"
 #include "SMaterialPalette.h"
 #include "FindInMaterial.h"
@@ -4989,19 +4989,30 @@ void FMaterialEditor::OnPreviewNode()
 		{
 			if (UMaterialGraphNode* GraphNode = Cast<UMaterialGraphNode>(*NodeIt))
 			{
-				PreviewNode(*GraphNode);
+				PreviewNode(GraphNode);
 			}
 		}
 	}
 }
 
-void FMaterialEditor::PreviewNode(UMaterialGraphNode& InGraphNode)
+void FMaterialEditor::OnPreviewHoveredNode()
 {
+	if (TSharedPtr<SGraphEditor> FocusedGraphEd = FocusedGraphEdPtr.Pin())
+	{
+		if (SGraphPanel* GraphPanel = FocusedGraphEd->GetGraphPanel())
+		{
+			PreviewNode(Cast<UMaterialGraphNode>(GraphPanel->GetCurrentHoveredNode()));
+		}
+	}
+}
+
+void FMaterialEditor::PreviewNode(UMaterialGraphNode* InGraphNode)
+{  
 	if (TSharedPtr<SGraphEditor> FocusedGraphEd = FocusedGraphEdPtr.Pin())
 	{
 		FocusedGraphEd->NotifyGraphChanged();
 	}
-	SetPreviewExpression(InGraphNode.MaterialExpression);
+	SetPreviewExpression(InGraphNode ? InGraphNode->MaterialExpression : nullptr);
 }
 
 void FMaterialEditor::OnToggleRealtimePreview()
@@ -5955,15 +5966,22 @@ void FMaterialEditor::SetPreviewExpression(UMaterialExpression* NewPreviewExpres
 			FunctionOutput->bLastPreviewed = false;
 		}
 		// If we are already previewing the selected expression toggle previewing off
-		PreviewExpression = NULL;
-		ExpressionPreviewMaterial->GetExpressionCollection().Empty();
-		SetPreviewMaterial( Material );
+		PreviewExpression = nullptr;
+
+		if (ExpressionPreviewMaterial)
+		{
+			ExpressionPreviewMaterial->GetExpressionCollection().Empty();
+		}
+		
+		// Set Preview Material to the material we're working on
+		SetPreviewMaterial(Material);
+		
 		// Recompile the preview material to get changes that might have been made during previewing
 		UpdatePreviewMaterial();
 	}
 	else
 	{
-		if( ExpressionPreviewMaterial == NULL )
+		if (ExpressionPreviewMaterial == nullptr)
 		{
 			// Create the expression preview material if it hasnt already been created
 			ExpressionPreviewMaterial = NewObject<UPreviewMaterial>(GetTransientPackage(), NAME_None, RF_Public);
@@ -7877,6 +7895,10 @@ TSharedRef<SGraphEditor> FMaterialEditor::CreateGraphEditorWidget(TSharedRef<cla
 			FExecuteAction::CreateSP(this, &FMaterialEditor::OnPreviewNode)
 			);
 
+		GraphEditorCommands->MapAction( FMaterialEditorCommands::Get().PreviewHoveredNode,
+			FExecuteAction::CreateSP(this, &FMaterialEditor::OnPreviewHoveredNode)
+		);
+
 		GraphEditorCommands->MapAction( FMaterialEditorCommands::Get().EnableRealtimePreviewNode,
 			FExecuteAction::CreateSP(this, &FMaterialEditor::OnToggleRealtimePreview)
 			);
diff --git a/Engine/Source/Editor/MaterialEditor/Private/MaterialEditor.h b/Engine/Source/Editor/MaterialEditor/Private/MaterialEditor.h
--- a/Engine/Source/Editor/MaterialEditor/Private/MaterialEditor.h
+++ b/Engine/Source/Editor/MaterialEditor/Private/MaterialEditor.h
@@ -812,12 +812,14 @@ private:
 	void OnConvertNamedRerouteToReroute();
 	/** Command for previewing a selected node */
 	void OnPreviewNode();
+	/** Command for previewing the currently-hovered node */
+	void OnPreviewHoveredNode();
 
 	/** 
 	 *	Call to Preview a InGraphNode.
 	 *	If InGraphNode is the currently-previewed node, will stop previewing
 	 */
-	virtual void PreviewNode(UMaterialGraphNode& InGraphNode) override;
+	virtual void PreviewNode(UMaterialGraphNode* InGraphNode) override;
 
 	/** Command for toggling real time preview of selected node */
 	void OnToggleRealtimePreview();
diff --git a/Engine/Source/Editor/MaterialEditor/Private/MaterialEditorActions.cpp b/Engine/Source/Editor/MaterialEditor/Private/MaterialEditorActions.cpp
--- a/Engine/Source/Editor/MaterialEditor/Private/MaterialEditorActions.cpp
+++ b/Engine/Source/Editor/MaterialEditor/Private/MaterialEditorActions.cpp
@@ -63,6 +63,7 @@ void FMaterialEditorCommands::RegisterCommands()
 	UI_COMMAND( ConvertToTextureSamples, "Convert to Texture Sample", "Converts the objects to texture samples.", EUserInterfaceActionType::Button, FInputChord() );
 	UI_COMMAND( StopPreviewNode, "Stop Previewing Node", "Stops the preview viewport from previewing this node", EUserInterfaceActionType::Button, FInputChord() );
 	UI_COMMAND( StartPreviewNode, "Start Previewing Node", "Makes the preview viewport start previewing this node", EUserInterfaceActionType::Button, FInputChord() );
+	UI_COMMAND(PreviewHoveredNode, "Preview Hovered Node", "Previews the currently-hovered node", EUserInterfaceActionType::Button, FInputChord(EModifierKey::Shift, EKeys::SpaceBar) );
 	UI_COMMAND( EnableRealtimePreviewNode, "Enable Realtime Preview", "Enables realtime previewing of this expression node", EUserInterfaceActionType::Button, FInputChord() );
 	UI_COMMAND( DisableRealtimePreviewNode, "Disable Realtime Preview", "Disables realtime previewing of this expression node", EUserInterfaceActionType::Button, FInputChord() );
 	UI_COMMAND( BreakAllLinks, "Break All Links", "Breaks all links leading out of this node.", EUserInterfaceActionType::Button, FInputChord() );
diff --git a/Engine/Source/Editor/MaterialEditor/Public/IMaterialEditor.h b/Engine/Source/Editor/MaterialEditor/Public/IMaterialEditor.h
--- a/Engine/Source/Editor/MaterialEditor/Public/IMaterialEditor.h
+++ b/Engine/Source/Editor/MaterialEditor/Public/IMaterialEditor.h
@@ -165,7 +165,7 @@ public:
 	virtual void RefreshStatsMaterials() {};
 
 	/** Start/Stop Previewing the specified material graph node */
-	virtual void PreviewNode(UMaterialGraphNode& InGraphNode) {};
+	virtual void PreviewNode(UMaterialGraphNode* InGraphNode) {};
 private:
 	FMaterialEditorClosedEvent MaterialEditorClosedEvent;
 	FRegisterTabSpawnersEvent RegisterTabSpawnersEvent;
diff --git a/Engine/Source/Editor/MaterialEditor/Public/MaterialEditorActions.h b/Engine/Source/Editor/MaterialEditor/Public/MaterialEditorActions.h
--- a/Engine/Source/Editor/MaterialEditor/Public/MaterialEditorActions.h
+++ b/Engine/Source/Editor/MaterialEditor/Public/MaterialEditorActions.h
@@ -146,8 +146,11 @@ public:
 	/** Stops a node from being previewed in the viewport */
 	TSharedPtr< FUICommandInfo > StopPreviewNode;
 	
-	/** Makes a new be previewed in the viewport */
+	/** Makes the selected node be previewed in the viewport */
 	TSharedPtr< FUICommandInfo > StartPreviewNode;
+
+	/** Previews the currently-hovered node */
+	TSharedPtr< FUICommandInfo > PreviewHoveredNode;
 	
 	/** Enables realtime previewing of this node */
 	TSharedPtr< FUICommandInfo > EnableRealtimePreviewNode;
diff --git a/Engine/Source/Editor/UnrealEd/Private/MaterialGraphNode.cpp b/Engine/Source/Editor/UnrealEd/Private/MaterialGraphNode.cpp
--- a/Engine/Source/Editor/UnrealEd/Private/MaterialGraphNode.cpp
+++ b/Engine/Source/Editor/UnrealEd/Private/MaterialGraphNode.cpp
@@ -850,7 +850,7 @@ void UMaterialGraphNode::ToggleNodePreview()
 	}
 
 	IMaterialEditor* MaterialEditor = static_cast<IMaterialEditor*>(FoundAssetEditor);
-	MaterialEditor->PreviewNode(*this);
+	MaterialEditor->PreviewNode(this);
 }
 
 void UMaterialGraphNode::ResetMaterialExpressionOwner()
