Fix for missing nullptr checks on Material Attribute resolves in MaterialEditor.cpp.

diff --git a/Engine/Source/Editor/MaterialEditor/Private/MaterialEditor.cpp b/Engine/Source/Editor/MaterialEditor/Private/MaterialEditor.cpp
--- a/Engine/Source/Editor/MaterialEditor/Private/MaterialEditor.cpp
+++ b/Engine/Source/Editor/MaterialEditor/Private/MaterialEditor.cpp
@@ -874,17 +874,17 @@ void FMaterialEditor::InitMaterialEditor( const EToolkitMode::Type Mode, const T
 									UMaterialExpressionSetMaterialAttributes* SetMaterialAttributes = Cast<UMaterialExpressionSetMaterialAttributes>(CreateNewMaterialExpression(UMaterialExpressionSetMaterialAttributes::StaticClass(), FVector2D(-75, 300), false, true));
 									UMaterialExpression* HorizontalMixingNode = CreateNewMaterialExpression(UMaterialExpressionSubstrateHorizontalMixing::StaticClass(), FVector2D(-350, 200), false, true);
 									UMaterialExpressionGetMaterialAttributes* GetBottomMaterialAttributes = Cast<UMaterialExpressionGetMaterialAttributes>(CreateNewMaterialExpression(UMaterialExpressionGetMaterialAttributes::StaticClass(), FVector2D(-750, 200), false, true));
-									int32 BottomFrontMaterialIndex = GetBottomMaterialAttributes->CreateOrGetOutputAttribute(MP_FrontMaterial);
+									int32 BottomFrontMaterialIndex = GetBottomMaterialAttributes ? GetBottomMaterialAttributes->CreateOrGetOutputAttribute(MP_FrontMaterial) : INDEX_NONE;
 
 									UMaterialExpression* LegacyBlendNode = CreateNewMaterialExpression(UMaterialExpressionBlendMaterialAttributes::StaticClass(), FVector2D(-350, 400), false, true);
 									UMaterialExpressionGetMaterialAttributes* GetTopMaterialAttributes = Cast< UMaterialExpressionGetMaterialAttributes>(CreateNewMaterialExpression(UMaterialExpressionGetMaterialAttributes::StaticClass(), FVector2D(-750, 400), false, true));
-									int32 TopFrontMaterialIndex = GetTopMaterialAttributes->CreateOrGetOutputAttribute(MP_FrontMaterial);
+									int32 TopFrontMaterialIndex = GetTopMaterialAttributes ? GetTopMaterialAttributes->CreateOrGetOutputAttribute(MP_FrontMaterial) : INDEX_NONE;
 
 									if (SetMaterialAttributes 
 										&& HorizontalMixingNode 
-										&& GetBottomMaterialAttributes && BottomFrontMaterialIndex > 0 
+										&& BottomFrontMaterialIndex > 0 
 										&& LegacyBlendNode 
-										&& GetTopMaterialAttributes && TopFrontMaterialIndex > 0)
+										&& TopFrontMaterialIndex > 0)
 									{
 										const FString FrontMaterialName = FMaterialAttributeDefinitionMap::GetAttributeName(MP_FrontMaterial);
 										const FString MaterialAttributesName = FMaterialAttributeDefinitionMap::GetAttributeName(MP_MaterialAttributes);
