Fix for Substrate crash in editor windows when layer reroute were not correctly assigned during preview material generation. The ParamLayer overrides embedded on the attribute layer expression have to be synchronised in and out based on the MI being processed, otherwise we either end up attempting to compile invalid layers/memory or the wrong layer stack for either the MI or the Parent.

diff --git a/Engine/Source/Editor/MaterialEditor/Private/SMaterialEditorTopologyWidget.cpp b/Engine/Source/Editor/MaterialEditor/Private/SMaterialEditorTopologyWidget.cpp
--- a/Engine/Source/Editor/MaterialEditor/Private/SMaterialEditorTopologyWidget.cpp
+++ b/Engine/Source/Editor/MaterialEditor/Private/SMaterialEditorTopologyWidget.cpp
@@ -4,6 +4,7 @@
 
 #include "EditorWidgetsModule.h"
 #include "AssetRegistry/AssetRegistryModule.h"
+#include "Materials/MaterialInterface.h"
 #include "MaterialEditor.h"
 #include "MaterialInstanceEditor.h"
 #include "SubstrateDefinitions.h"
@@ -115,9 +116,10 @@ void SMaterialEditorTopologyWidget::Tick(const FGeometry& AllottedGeometry, cons
 	if (MaterialEditorPtr.IsValid())
 	{
 		TSharedPtr<IMaterialEditor> MaterialEditorPinned = MaterialEditorPtr.Pin();
-		FMaterialResource* MaterialResource = MaterialEditorPinned->GetMaterialInterface()->GetMaterialResource(GMaxRHIFeatureLevel);
+		UMaterialInterface* MaterialInterface = MaterialEditorPinned->GetMaterialInterface();
+		FMaterialResource* MaterialResource = MaterialInterface ? MaterialInterface->GetMaterialResource(GMaxRHIFeatureLevel) : nullptr;
 
-		if (Cast<UMaterial>(MaterialEditorPinned->GetMaterialInterface()))
+		if (MaterialInterface && MaterialInterface->IsA<UMaterial>())
 		{
 			TSharedPtr<FMaterialEditor> MaterialEditor = StaticCastSharedPtr<FMaterialEditor>(MaterialEditorPinned);
 			UMaterial* MaterialForStats = MaterialEditor->bStatsFromPreviewMaterial ? MaterialEditor->Material : MaterialEditor->OriginalMaterial;
@@ -126,7 +128,7 @@ void SMaterialEditorTopologyWidget::Tick(const FGeometry& AllottedGeometry, cons
 
 		SubstrateMaterialDescription = FText::FromString(FString(TEXT("SubstrateMaterialDescription")));
 		
-		if (MaterialResource)
+		if (MaterialInterface && MaterialResource)
 		{
 			FString MaterialDescription;
 			FString MaterialBudget;
@@ -143,8 +145,13 @@ void SMaterialEditorTopologyWidget::Tick(const FGeometry& AllottedGeometry, cons
 				// Now generate a visual representation of the material from the topology tree of operators.
 				if (CompilationOutput.RootOperatorIndex >= 0)
 				{
-					const FMaterialLayersFunctions* LayersFunctions = MaterialResource->GetMaterialLayers();
-					MaterialBox->SetContent(ProcessOperatorAsThumbnails(CompilationOutput, LayersFunctions));
+					FMaterialLayersFunctions LayersFunctions;
+					if (MaterialInterface->GetMaterialLayers(LayersFunctions))
+					{
+						MaterialInterface->SyncLayersRuntimeGraphCache(&LayersFunctions);
+					}
+					MaterialBox->SetContent(ProcessOperatorAsThumbnails(CompilationOutput, &LayersFunctions));
+					MaterialInterface->SyncLayersRuntimeGraphCache(nullptr);
 				}
 				else
 				{
diff --git a/Engine/Source/Editor/MaterialEditor/Private/SMaterialSubstrateTree.cpp b/Engine/Source/Editor/MaterialEditor/Private/SMaterialSubstrateTree.cpp
--- a/Engine/Source/Editor/MaterialEditor/Private/SMaterialSubstrateTree.cpp
+++ b/Engine/Source/Editor/MaterialEditor/Private/SMaterialSubstrateTree.cpp
@@ -1336,26 +1336,31 @@ void SMaterialSubstrateTree::RecursiveCreateWidgets(FRecursiveCreateWidgetsConte
 			}
 			else if(Substrate::IsMaterialLayeringSupportEnabled())
 			{
-				UMaterialInterface* MaterialInterface = MaterialEditorInstance->GetMaterialInterface();
-				const FMaterialResource* MaterialResource = MaterialInterface->GetMaterialResource(ERHIFeatureLevel::SM6);
-				const FMaterialLayersFunctions* LayersFunctions = MaterialResource->GetMaterialLayers();
-				TSharedPtr<FMaterialLayersFunctionsRuntimeGraphCache> LayerTreeCache = LayersFunctions->RuntimeGraphCache;
-				if (LayerTreeCache)
+				if(UMaterialInterface* MaterialInterface = MaterialEditorInstance->GetMaterialInterface())
 				{
-					UMaterialExpressionMaterialFunctionCall* Call = LayerTreeCache->NodeMaterialGraphExpressions[InNodeId];
-					if (Call && Call->MaterialFunction)
+					FMaterialLayersFunctions LayersFunctions;
+					if(MaterialInterface->GetMaterialLayers(LayersFunctions))
 					{
-						if (MaterialEditorInstance->StoredLayerPreviews[Payload.Layer] == nullptr)
+						MaterialInterface->SyncLayersRuntimeGraphCache(&LayersFunctions);
+						if (TSharedPtr<FMaterialLayersFunctionsRuntimeGraphCache> LayerTreeCache = LayersFunctions.RuntimeGraphCache)
 						{
-							MaterialEditorInstance->StoredLayerPreviews[Payload.Layer] = (NewObject<UMaterialInstanceConstant>(MaterialEditorInstance, NAME_None));
-						}
+							UMaterialExpressionMaterialFunctionCall* Call = LayerTreeCache->NodeMaterialGraphExpressions[InNodeId];
+							if (Call && Call->MaterialFunction)
+							{
+								if (MaterialEditorInstance->StoredLayerPreviews[Payload.Layer] == nullptr)
+								{
+									MaterialEditorInstance->StoredLayerPreviews[Payload.Layer] = (NewObject<UMaterialInstanceConstant>(MaterialEditorInstance, NAME_None));
+								}
 
-						UMaterial* PreviewMaterial = LayersFunctions->GetRuntimeNodePreviewMaterial(InNodeId);
+								UMaterial* PreviewMaterial = LayersFunctions.GetRuntimeNodePreviewMaterial(InNodeId);
 
-						if (MaterialEditorInstance->StoredLayerPreviews[Payload.Layer] && MaterialEditorInstance->StoredLayerPreviews[Payload.Layer]->Parent != PreviewMaterial)
-						{
-							MaterialEditorInstance->StoredLayerPreviews[Payload.Layer]->SetParentEditorOnly(PreviewMaterial);
+								if (MaterialEditorInstance->StoredLayerPreviews[Payload.Layer] && MaterialEditorInstance->StoredLayerPreviews[Payload.Layer]->Parent != PreviewMaterial)
+								{
+									MaterialEditorInstance->StoredLayerPreviews[Payload.Layer]->SetParentEditorOnly(PreviewMaterial);
+								}
+							}
 						}
+						MaterialInterface->SyncLayersRuntimeGraphCache(nullptr);
 					}
 				}
 			}
@@ -1405,7 +1410,17 @@ void SMaterialSubstrateTree::CreateGroupsWidget()
 	check(MaterialEditorInstance);
 	if (MaterialEditorInstance->IsA<UMaterialEditorInstanceConstant>())
 	{
-		MaterialEditorInstance->RegenerateArrays();
+		FMaterialLayersFunctions LayersFunctions;
+		UMaterialInterface* MaterialInterface = MaterialEditorInstance->GetMaterialInterface();
+		if (MaterialInterface && MaterialInterface->GetMaterialLayers(LayersFunctions))
+		{
+			MaterialInterface->SyncLayersRuntimeGraphCache(&LayersFunctions);
+		}
+		MaterialEditorInstance->RegenerateArrays();		
+		if (MaterialInterface)
+		{
+			MaterialInterface->SyncLayersRuntimeGraphCache(nullptr);
+		}
 	}
 	
 	TArray<FUnsortedParamData> NonLayerProperties; // NonLayerProperties.Empty();
diff --git a/Engine/Source/Runtime/Engine/Private/Materials/Material.cpp b/Engine/Source/Runtime/Engine/Private/Materials/Material.cpp
--- a/Engine/Source/Runtime/Engine/Private/Materials/Material.cpp
+++ b/Engine/Source/Runtime/Engine/Private/Materials/Material.cpp
@@ -2268,6 +2268,23 @@ bool UMaterial::GetMaterialLayers(FMaterialLayersFunctions& OutLayers, TMicRecur
 	return false;
 }
 
+#if WITH_EDITORONLY_DATA
+void UMaterial::SyncLayersRuntimeGraphCache(FMaterialLayersFunctions* OverrideLayers)
+{
+	if(UMaterialEditorOnlyData* EditorOnly = GetEditorOnlyData())
+	{
+		for (UMaterialExpression* Expression : EditorOnly->ExpressionCollection.Expressions)
+		{
+			if (UMaterialExpressionMaterialAttributeLayers* LayersExpression = Cast<UMaterialExpressionMaterialAttributeLayers>(Expression))
+			{
+				LayersExpression->OverrideLayerGraph(OverrideLayers);
+				return;
+			}
+		}
+	}
+}
+#endif
+
 bool UMaterial::GetRefractionSettings(float& OutBiasValue) const
 {
 	OutBiasValue = RefractionDepthBias;
diff --git a/Engine/Source/Runtime/Engine/Private/Materials/MaterialCachedData.cpp b/Engine/Source/Runtime/Engine/Private/Materials/MaterialCachedData.cpp
--- a/Engine/Source/Runtime/Engine/Private/Materials/MaterialCachedData.cpp
+++ b/Engine/Source/Runtime/Engine/Private/Materials/MaterialCachedData.cpp
@@ -605,7 +605,7 @@ void FMaterialCachedExpressionData::UpdateForExpressions(const FMaterialCachedEx
 				MaterialLayers = Layers.GetRuntime();
 				EditorOnlyData->MaterialLayers = Layers.EditorOnly;
 				FMaterialLayersFunctions::Validate(MaterialLayers, EditorOnlyData->MaterialLayers);
-				LayersExpression->RebuildLayerGraph(false);
+				LayersExpression->OverrideLayerGraph(&Layers);
 			}
 		}
 		else if (UMaterialExpressionMaterialFunctionCall* FunctionCall = Cast<UMaterialExpressionMaterialFunctionCall>(Expression))
diff --git a/Engine/Source/Runtime/Engine/Private/Materials/MaterialInstance.cpp b/Engine/Source/Runtime/Engine/Private/Materials/MaterialInstance.cpp
--- a/Engine/Source/Runtime/Engine/Private/Materials/MaterialInstance.cpp
+++ b/Engine/Source/Runtime/Engine/Private/Materials/MaterialInstance.cpp
@@ -3583,6 +3583,14 @@ void UMaterialInstance::DeclareConstructClasses(TArray<FTopLevelAssetPath>& OutC
 	Super::DeclareConstructClasses(OutConstructClasses, SpecificSubclass);
 	OutConstructClasses.Add(FTopLevelAssetPath(UMaterialInstanceEditorOnlyData::StaticClass()));
 }
+
+void UMaterialInstance::SyncLayersRuntimeGraphCache(FMaterialLayersFunctions* OverrideLayers)
+{
+	if (Parent)
+	{
+		Parent->SyncLayersRuntimeGraphCache(OverrideLayers);
+	}
+}
 #endif
 
 void UMaterialInstance::BeginDestroy()
diff --git a/Engine/Source/Runtime/Engine/Public/Materials/Material.h b/Engine/Source/Runtime/Engine/Public/Materials/Material.h
--- a/Engine/Source/Runtime/Engine/Public/Materials/Material.h
+++ b/Engine/Source/Runtime/Engine/Public/Materials/Material.h
@@ -1267,6 +1267,8 @@ public:
 #if WITH_EDITORONLY_DATA
 	/** Return true if the material support shading model override */
 	ENGINE_API bool SupportsShadingModelOverride() const;
+
+	ENGINE_API virtual void SyncLayersRuntimeGraphCache(FMaterialLayersFunctions* OverrideLayers) override;
 #endif
 
 #if WITH_EDITOR
diff --git a/Engine/Source/Runtime/Engine/Public/Materials/MaterialInstance.h b/Engine/Source/Runtime/Engine/Public/Materials/MaterialInstance.h
--- a/Engine/Source/Runtime/Engine/Public/Materials/MaterialInstance.h
+++ b/Engine/Source/Runtime/Engine/Public/Materials/MaterialInstance.h
@@ -947,6 +947,7 @@ public:
 	virtual ENGINE_API void PostLoad() override;
 #if WITH_EDITORONLY_DATA
 	ENGINE_API static void DeclareConstructClasses(TArray<FTopLevelAssetPath>& OutConstructClasses, const UClass* SpecificSubclass);
+	ENGINE_API virtual void SyncLayersRuntimeGraphCache(FMaterialLayersFunctions* OverrideLayers) override;
 #endif
 	virtual ENGINE_API void BeginDestroy() override;
 	virtual ENGINE_API bool IsReadyForFinishDestroy() override;
diff --git a/Engine/Source/Runtime/Engine/Public/Materials/MaterialInterface.h b/Engine/Source/Runtime/Engine/Public/Materials/MaterialInterface.h
--- a/Engine/Source/Runtime/Engine/Public/Materials/MaterialInterface.h
+++ b/Engine/Source/Runtime/Engine/Public/Materials/MaterialInterface.h
@@ -678,6 +678,11 @@ public:
 	* @return					True if successful
 	*/
 	ENGINE_API bool GetStaticComponentMaskParameterValue(const FHashedMaterialParameterInfo& ParameterInfo, bool& R, bool& G, bool& B, bool& A, FGuid& OutExpressionGuid, bool bOveriddenOnly = false) const;
+
+	/**
+	 * Update layers functions RuntimeGraphCache for compiling (primarily for Preview Materials)
+	 */
+	ENGINE_API virtual void SyncLayersRuntimeGraphCache(FMaterialLayersFunctions* OverrideLayers) PURE_VIRTUAL(UMaterialInterface::SyncLayersRuntimeGraphCache);
 #endif // WITH_EDITORONLY_DATA
 
 #if WITH_EDITOR
