Improve HLSL Code panel in the Material Editor

- By default the HLSL code panel will only display the generated HLSL code that comes from the material graph itself.  It does not display the full generated shader source.
- This makes it easier to grok, especially for Tech Artists.
- We now display a listing of the key, value pairs from MaterialTemplate.ush "%s" key to the string that is injected into the code:
- ExtractHLSLCode now pulls all data from the new FMaterialInsights object instead of the string returned by GetMaterialExpressionSource()
- ExtractHLSLCode pulls the key,value map, makes a copy, sorts by key, and generates a string to display to the user.  This is done so we don't mutate state on the Material's FMaterialInsight object.
- If the new translator is enabled we display three separate panels:
1. New Translator
2. Old Translator
3. Intermediate Representation string
- The HLSL code panel has been refactor to contain a "Copy" button and I've refactored the code to make it easier to re-use.
- DebugDumpIR returns the string instead of dumping to file.
- Anytime we translate we should clear the FMaterialInsights ptr since it can be used by new and legacy translator.

[RN] Major | Rendering - Materials and Shaders

diff --git a/Engine/Source/Editor/MaterialEditor/Private/MaterialStats.cpp b/Engine/Source/Editor/MaterialEditor/Private/MaterialStats.cpp
--- a/Engine/Source/Editor/MaterialEditor/Private/MaterialStats.cpp
+++ b/Engine/Source/Editor/MaterialEditor/Private/MaterialStats.cpp
@@ -15,6 +15,7 @@
 #include "HAL/PlatformApplicationMisc.h"
 #include "MaterialEditorActions.h"
 #include "Materials/MaterialInstance.h"
+#include "Materials/MaterialInsights.h"
 #include "IMaterialEditor.h"
 #include "MaterialEditorSettings.h"
 #include "ShaderCompiler.h"
@@ -998,85 +999,132 @@ static FText GenerateLineNoTextForCodeView(const FString& InSourceCode)
 	return FText::FromString(MoveTemp(LineNumberStr));
 }
 
-TSharedRef<SScrollBox> FMaterialStats::BuildShaderCodeWidget(TFunction<FText(void)>&& InShaderCodeCallback)
+TSharedRef<SWidget> FMaterialStats::BuildShaderCodeWidget(FText TextLabel, TFunction<FText(void)>&& InShaderCodeCallback, TFunction<FReply(void)>&& InCopyCodeCallback, TAttribute<EVisibility> InVisibilityAttribute)
 {
-	return SNew(SScrollBox)
-		+ SScrollBox::Slot().Padding(5)
-		[
-			SNew(SSplitter)
-			.Orientation(EOrientation::Orient_Horizontal)
-			+SSplitter::Slot()
-			.Resizable(false)
-			.SizeRule(SSplitter::SizeToContent)
-			[
-				SNew(STextBlock)
-				.Text_Lambda([InShaderCodeCallback]()
-				{
-					return GenerateLineNoTextForCodeView(InShaderCodeCallback().ToString());
-				})
-				.TextStyle(FAppStyle::Get(), "MessageLog")
-				.ColorAndOpacity(FLinearColor::Gray)
-			]
-			+SSplitter::Slot()
-			[
-				SNew(SMultiLineEditableText)
-				.Text_Lambda([InShaderCodeCallback](){ return InShaderCodeCallback(); })
-				.TextStyle(FAppStyle::Get(), "MessageLog")
-				.IsReadOnly(true)
-				.Marshaller(SyntaxHighlighter)
-			]
-		];
-}
-
-TSharedRef<class SDockTab> FMaterialStats::SpawnTab_HLSLCode(const class FSpawnTabArgs& Args)
-{
-	auto CodeViewUtility =
-		SNew(SVerticalBox)
-		// Copy Button
+	TSharedRef<SVerticalBox> Box = SNew(SVerticalBox)
 		+ SVerticalBox::Slot()
 		.AutoHeight()
 		[
-			SNew(SHorizontalBox)
-			+ SHorizontalBox::Slot()
-			.AutoWidth()
-			.Padding(2.0f, 0.0f)
-			.VAlign(VAlign_Center)
-			.HAlign(HAlign_Left)
-			[
-				SNew(SButton)
-				.Text(LOCTEXT("CopyHLSLButton", "Copy"))
-				.ToolTipText(LOCTEXT("CopyHLSLButtonToolTip", "Copies all HLSL code to the clipboard."))
-				.ContentPadding(3)
-				.OnClicked_Lambda
-				(
-					[Code = &HLSLCode]()
-					{
-						FPlatformApplicationMisc::ClipboardCopy(**Code);
-						return FReply::Handled();
-					}
-				)
-			]
+			SNew(SVerticalBox)
+				// Copy Button
+				+ SVerticalBox::Slot()
+				.AutoHeight()
+				[
+					SNew(SHorizontalBox)
+						+ SHorizontalBox::Slot()
+						.AutoWidth()
+						.Padding(2.0f, 0.0f)
+						.VAlign(VAlign_Center)
+						.HAlign(HAlign_Left)
+						[
+							SNew(SButton)
+								.Text(LOCTEXT("CopyHLSLButton", "Copy"))
+								.ToolTipText(LOCTEXT("CopyHLSLButtonToolTip", "Copies all HLSL code to the clipboard."))
+								.ContentPadding(3)
+								.OnClicked_Lambda(InCopyCodeCallback)
+						]
+						+ SHorizontalBox::Slot()
+						.AutoWidth()
+						.Padding(2.0f, 0.0f)
+						.VAlign(VAlign_Center)
+						.HAlign(HAlign_Center)
+						[
+							SNew(STextBlock)
+								.Text(TextLabel)
+						]
+				]
+				// Separator
+				+ SVerticalBox::Slot()
+				.FillHeight(1)
+				[
+					SNew(SSeparator)
+				]
 		]
-		// Separator
 		+ SVerticalBox::Slot()
 		.FillHeight(1)
 		[
-			SNew(SSeparator)
+			SNew(SScrollBox)
+				+ SScrollBox::Slot().Padding(5)
+				[
+					SNew(SSplitter)
+						.Orientation(EOrientation::Orient_Horizontal)
+						+ SSplitter::Slot()
+						.Resizable(false)
+						.SizeRule(SSplitter::SizeToContent)
+						[
+							SNew(STextBlock)
+								.Text_Lambda([InShaderCodeCallback]()
+									{
+										return GenerateLineNoTextForCodeView(InShaderCodeCallback().ToString());
+									})
+								.TextStyle(FAppStyle::Get(), "MessageLog")
+								.ColorAndOpacity(FLinearColor::Gray)
+						]
+						+ SSplitter::Slot()
+						[
+							SNew(SMultiLineEditableText)
+								.Text_Lambda([InShaderCodeCallback]() { return InShaderCodeCallback(); })
+								.TextStyle(FAppStyle::Get(), "MessageLog")
+								.IsReadOnly(true)
+								.Marshaller(SyntaxHighlighter)
+						]
+				]
 		];
 
+	Box->SetVisibility(InVisibilityAttribute);
+
+	return Box;
+}
+
+TSharedRef<class SDockTab> FMaterialStats::SpawnTab_HLSLCode(const class FSpawnTabArgs& Args)
+{
 	TSharedRef<SDockTab> SpawnedTab = SNew(SDockTab)
 		.Label(LOCTEXT("HLSLCodeTitle", "HLSL Code"))
 		[
 			SNew(SVerticalBox)
 			+ SVerticalBox::Slot()
-			.AutoHeight()
-			[
-				CodeViewUtility
-			]
-			+ SVerticalBox::Slot()
 			.FillHeight(1)
 			[
-				BuildShaderCodeWidget([Code = &HLSLCode]() { return FText::FromString(*Code); })
+				SNew(SSplitter)
+				.Orientation(EOrientation::Orient_Vertical)
+				+SSplitter::Slot()
+				[
+					SNew(SSplitter)
+					.Orientation(EOrientation::Orient_Horizontal)
+					+SSplitter::Slot()
+					[
+						BuildShaderCodeWidget(LOCTEXT("CurrentTranslatorLabel", "HLSL"),
+							[Code = &HLSLCode]() { return FText::FromString(*Code); },
+							[Code = &HLSLCode]() {
+								FPlatformApplicationMisc::ClipboardCopy(**Code);
+								return FReply::Handled();
+							},
+							MakeAttributeLambda([Code = &HLSLCode]() { return (*Code).Len() == 0 ? EVisibility::Collapsed : EVisibility::Visible; })
+						)
+					]
+					+SSplitter::Slot()
+					[
+						BuildShaderCodeWidget(LOCTEXT("PreviousTranslatorLabel", "HLSL"),
+							[Code = &PreviousTranslatorHLSLCode]() { return FText::FromString(*Code); },
+							[Code = &PreviousTranslatorHLSLCode]() {
+								FPlatformApplicationMisc::ClipboardCopy(**Code);
+								return FReply::Handled();
+							},
+							MakeAttributeLambda([Code = &PreviousTranslatorHLSLCode]() { return (*Code).Len() == 0 ? EVisibility::Collapsed : EVisibility::Visible; })
+						)
+					]
+				]
+				+SSplitter::Slot()
+				[
+					BuildShaderCodeWidget(LOCTEXT("MaterialIRLabel", "Intermediate Representation"),
+						[Code = &MaterialIR]() { return FText::FromString(*Code); },
+						[Code = &MaterialIR]() {
+							FPlatformApplicationMisc::ClipboardCopy(**Code);
+							return FReply::Handled();
+						},
+						MakeAttributeLambda([Code = &MaterialIR]() { return (*Code).Len() == 0 ? EVisibility::Collapsed : EVisibility::Visible; })
+					)
+				]
 			]
 		];
 
@@ -1149,32 +1197,6 @@ TSharedRef<SDockTab> FMaterialStats::SpawnTab_ShaderCode(const FSpawnTabArgs& Ar
 					SNew(SHorizontalBox)
 					+ SHorizontalBox::Slot()
 					.AutoWidth()
-					.Padding(2.0f, 0.0f)
-					.VAlign(VAlign_Center)
-					.HAlign(HAlign_Left)
-					[
-						SNew(SButton)
-						.Text(LOCTEXT("CopyShaderCodeButton", "Copy"))
-						.ToolTipText(LOCTEXT("CopyShaderCodeButtonToolTip", "Copies all shader code to the clipboard."))
-						.ContentPadding(3)
-						.OnClicked_Lambda
-						(
-							[MaterialStats = TWeakPtr<FMaterialStats>(SharedThis(this)), PlatformID, QualityLevel, InstanceIndex]()
-							{
-								auto StatsPtr = MaterialStats.Pin();
-								if (StatsPtr.IsValid())
-								{
-									FText ShaderCode = StatsPtr->GetShaderCode(PlatformID, QualityLevel, InstanceIndex);
-									FPlatformApplicationMisc::ClipboardCopy(*ShaderCode.ToString());
-									return FReply::Handled();
-								}
-
-								return FReply::Unhandled();
-							}
-						)
-					]
-					+ SHorizontalBox::Slot()
-					.AutoWidth()
 					.Padding(2.0f, 2.0f)
 					.VAlign(VAlign_Center)
 					.HAlign(HAlign_Left)
@@ -1287,7 +1309,8 @@ void FMaterialStats::BuildViewShaderCodeMenus()
 				TSharedRef<SScrollBox> CodeScrollBox = SNew(SScrollBox)
 					+ SScrollBox::Slot().Padding(5)
 					[
-						BuildShaderCodeWidget([MaterialStats = TWeakPtr<FMaterialStats>(SharedThis(this)), PlatformID, QualityLevel]()
+						BuildShaderCodeWidget(LOCTEXT("ShaderCodeLabel", "Shader Code"), 
+						[MaterialStats = TWeakPtr<FMaterialStats>(SharedThis(this)), PlatformID, QualityLevel]()
 						{
 							auto StatsPtr = MaterialStats.Pin();
 							if (StatsPtr.IsValid())
@@ -1295,7 +1318,19 @@ void FMaterialStats::BuildViewShaderCodeMenus()
 								return StatsPtr->GetShaderCode(PlatformID, QualityLevel, InstanceIndex);
 							}
 							return FText::FromString(TEXT("Error reading shader code!"));
-						})
+						},
+						[MaterialStats = TWeakPtr<FMaterialStats>(SharedThis(this)), PlatformID, QualityLevel]() {
+							auto StatsPtr = MaterialStats.Pin();
+							if (StatsPtr.IsValid())
+							{
+								FText ShaderCode = StatsPtr->GetShaderCode(PlatformID, QualityLevel, InstanceIndex);
+								FPlatformApplicationMisc::ClipboardCopy(*ShaderCode.ToString());
+								return FReply::Handled();
+							}
+
+							return FReply::Unhandled();
+						},
+						MakeAttributeLambda([]() { return EVisibility::Visible; }))
 					];
 
 				PlatformPtr->GetPlatformData(QualityLevel).CodeScrollBox = CodeScrollBox;
@@ -1550,53 +1585,55 @@ void FMaterialStats::ComputeGridWarnings()
 
 void FMaterialStats::ExtractHLSLCode()
 {
-#define MARKTAG TEXT("/*MARK_")
-#define MARKTAGLEN 7
-
 	HLSLCode = TEXT("");
+	PreviousTranslatorHLSLCode = TEXT("");
+	MaterialIR = TEXT("");
 
  	if (!HLSLTab.IsValid())
 	{
 		return;
 	}
 
-	FString MarkupSource;
-	if (MaterialInterface != nullptr && MaterialInterface->GetMaterialResource(GMaxRHIFeatureLevel)->GetMaterialExpressionSource(MarkupSource))
+	FString FullSource;
+	if (MaterialInterface != nullptr && MaterialInterface->GetMaterialResource(GMaxRHIFeatureLevel)->GetMaterialExpressionSource(FullSource))
 	{
-		// Remove line-feeds and leave just CRs so the character counts match the selection ranges.
-		MarkupSource.ReplaceInline(TEXT("\r"), TEXT(""));
-
-		// Improve formatting: Convert tab to 4 spaces since STextBlock (currently) doesn't show tab characters
-		MarkupSource.ReplaceInline(TEXT("\t"), TEXT("    "));
+		if (const FMaterialInsights* Insights = MaterialInterface->MaterialInsight.Get())
+		{
+			MaterialIR = Insights->IRString;
 
-		// Extract highlight ranges from markup tags
+			TMap<FString, FString> SortedLegacy_ShaderString = Insights->Legacy_ShaderStringParameters; 
+			SortedLegacy_ShaderString.KeySort([](const FString& A, const FString& B) { return A < B; });
 
-		// Make a copy so we can insert null terminators.
-		TCHAR* MarkupSourceCopy = new TCHAR[MarkupSource.Len() + 1];
-		FCString::Strncpy(MarkupSourceCopy, *MarkupSource, MarkupSource.Len() + 1);
+			TMap<FString, FString> SortedNew_ShaderString = Insights->New_ShaderStringParameters; 
+			SortedNew_ShaderString.KeySort([](const FString& A, const FString& B) { return A < B; });
 
-		TCHAR* Ptr = MarkupSourceCopy;
-		while (Ptr && *Ptr != '\0')
-		{
-			TCHAR* NextTag = FCString::Strstr(Ptr, MARKTAG);
-			if (!NextTag)
+			for (const TPair<FString, FString>& Pair : SortedLegacy_ShaderString)
 			{
-				// No more tags, so we're done!
-				HLSLCode += Ptr;
-				break;
+				const FString& Key = Pair.Key;
+				const FString& Value = Pair.Value;
+				PreviousTranslatorHLSLCode += FString::Printf(TEXT("// %s\n"), *Key);
+				PreviousTranslatorHLSLCode += FString::Printf(TEXT("%s\n"), *Value);
 			}
 
-			// Copy the text up to the tag.
-			*NextTag = '\0';
-			HLSLCode += Ptr;
+			for (const TPair<FString, FString>& Pair : SortedNew_ShaderString)
+			{
+				const FString& Key = Pair.Key;
+				const FString& Value = Pair.Value;
+				HLSLCode += FString::Printf(TEXT("// %s\n"), *Key);
+				HLSLCode += FString::Printf(TEXT("%s\n"), *Value);
+			}
 
-			// Advance past the markup tag to see what type it is (beginning or end)
-			NextTag += MARKTAGLEN;
-			int32 TagNumber = FCString::Atoi(NextTag + 1);
-			Ptr = FCString::Strstr(NextTag, TEXT("*/")) + 2;
+			if (MaterialInterface->IsUsingNewHLSLGenerator())
+			{
+				HLSLCode += TEXT("// FULL SOURCE BELOW\n");
+				HLSLCode += FullSource;
+			}
+			else
+			{
+				PreviousTranslatorHLSLCode += TEXT("// FULL SOURCE BELOW\n");
+				PreviousTranslatorHLSLCode += FullSource;
+			}
 		}
-
-		delete[] MarkupSourceCopy;
 	}
 }
 
diff --git a/Engine/Source/Editor/MaterialEditor/Private/MaterialStats.h b/Engine/Source/Editor/MaterialEditor/Private/MaterialStats.h
--- a/Engine/Source/Editor/MaterialEditor/Private/MaterialStats.h
+++ b/Engine/Source/Editor/MaterialEditor/Private/MaterialStats.h
@@ -6,6 +6,8 @@
 #include "MaterialStatsCommon.h"
 #include "UObject/GCObject.h"
 #include "Preferences/MaterialStatsOptions.h"
+#include "Input/Reply.h"
+#include "Layout/Visibility.h"
 
 class FHLSLSyntaxHighlighterMarshaller;
 
@@ -340,6 +342,12 @@ class FMaterialStats : public FGCObject, public TSharedFromThis<FMaterialStats>
 	/** Cached HLSL code for analyzed material */
 	FString HLSLCode;
 
+	/** If the new translator is enabled this will contain the source for the old translator.  Useful for comparing. */
+	FString PreviousTranslatorHLSLCode;
+
+	/** Cached IR from the new translator if enabled */
+	FString MaterialIR;
+
 	TObjectPtr<class UMaterialStatsOptions >Options = nullptr;
 
 	/** Pointer to Material Editor or to Material Instance Editor set by Initialize() function */
@@ -372,7 +380,7 @@ private:
 	TSharedRef<class SDockTab> SpawnTab_ShaderCode(const class FSpawnTabArgs& Args, const EShaderPlatform PlatformID, const EMaterialQualityLevel::Type QualityType, const int32 InstanceIndex);
 	TSharedRef<class SDockTab> SpawnTab_HLSLCode(const class FSpawnTabArgs& Args);
 
-	TSharedRef<SScrollBox> BuildShaderCodeWidget(TFunction<FText(void)>&& InShaderCodeCallback);
+	TSharedRef<SWidget> BuildShaderCodeWidget(FText TextLabel, TFunction<FText(void)>&& InShaderCodeCallback, TFunction<FReply(void)>&& InCopyCodeCallback, TAttribute<EVisibility> InVisibilityAttribute);
 
 	/** utility function used to build names for the shader viewing tabs  */
 	static FName MakeTabName(const EPlatformCategoryType PlatformType, const EShaderPlatform ShaderPlatformType, const EMaterialQualityLevel::Type QualityLevel, const int32 InstanceIndex);
diff --git a/Engine/Source/Runtime/Engine/Private/Materials/HLSLMaterialTranslator.cpp b/Engine/Source/Runtime/Engine/Private/Materials/HLSLMaterialTranslator.cpp
--- a/Engine/Source/Runtime/Engine/Private/Materials/HLSLMaterialTranslator.cpp
+++ b/Engine/Source/Runtime/Engine/Private/Materials/HLSLMaterialTranslator.cpp
@@ -39,6 +39,7 @@
 #include "Materials/MaterialExternalCodeRegistry.h"
 #include "Materials/MaterialSourceTemplate.h"
 #include "Materials/SubstrateMaterial.h"
+#include "Materials/MaterialInsights.h"
 #include "StringTemplate.h"
 #include "ParameterCollection.h"
 #include "RenderUtils.h"
@@ -3069,6 +3070,9 @@ FString FHLSLMaterialTranslator::GetMaterialShaderCode()
 	Resolver.SetParameterMap(&MaterialSourceTemplateParams);
 	MaterialSourceTemplateParams.Add({ TEXT("line_number"), FString::Printf(TEXT("%u"), MaterialTemplateLineNumberNew) });
 
+	// Copy generate code chunks to the insights object.
+	Material->GetMaterialInterface()->MaterialInsight->Legacy_ShaderStringParameters = MaterialSourceTemplateParams;
+
 	return Resolver.Finalize();
 }
 
diff --git a/Engine/Source/Runtime/Engine/Private/Materials/MaterialIRDebug.cpp b/Engine/Source/Runtime/Engine/Private/Materials/MaterialIRDebug.cpp
--- a/Engine/Source/Runtime/Engine/Private/Materials/MaterialIRDebug.cpp
+++ b/Engine/Source/Runtime/Engine/Private/Materials/MaterialIRDebug.cpp
@@ -405,7 +405,7 @@ struct FDebugDumpIRState
 	}
 };
 
-void DebugDumpIR(FStringView MaterialName, const FMaterialIRModule& Module)
+FString DebugDumpIR(FStringView MaterialName, const FMaterialIRModule& Module)
 {
 	FDebugDumpIRState State{};
 	State.Out.Append(TEXT("; Material IR module dump.\n"));
@@ -419,9 +419,7 @@ void DebugDumpIR(FStringView MaterialName, const FMaterialIRModule& Module)
 	// Print referenced material parameters recap if any
 	State.DumpReferencedParameters();
 
-	// Save the dump to file
-	FString FilePath = FPaths::Combine(FPaths::ProjectSavedDir(), "Materials", TEXT("IRDump.txt"));
-	FFileHelper::SaveStringToFile(State.Out, *FilePath);
+	return State.Out;
 }
 
 } // namespace MIR
diff --git a/Engine/Source/Runtime/Engine/Private/Materials/MaterialIRModuleBuilder.cpp b/Engine/Source/Runtime/Engine/Private/Materials/MaterialIRModuleBuilder.cpp
--- a/Engine/Source/Runtime/Engine/Private/Materials/MaterialIRModuleBuilder.cpp
+++ b/Engine/Source/Runtime/Engine/Private/Materials/MaterialIRModuleBuilder.cpp
@@ -21,6 +21,7 @@
 #include "Materials/MaterialExpression.h"
 #include "Materials/MaterialInsights.h"
 #include "Engine/Texture.h"
+#include "Misc/FileHelper.h"
 
 static TAutoConsoleVariable<int32> CVarMaterialIRDebugDumpLevel(
 	TEXT("r.Material.Translator.DebugDump"),
@@ -536,11 +537,20 @@ bool FMaterialIRModuleBuilder::Build(FMaterialIRModule* TargetModule)
 	Impl.Step_LinkInstructions();
 	Impl.Step_Finalize();
 
+	check(Material->MaterialInsight.IsValid());
+	Material->MaterialInsight.Get()->IRString = MIR::DebugDumpIR(Material->GetFullName(), *TargetModule);
+
 	// Dump debugging information if requested 
 	switch (CVarMaterialIRDebugDumpLevel.GetValueOnGameThread())
 	{
 		case 2: MIR::DebugDumpIRUseGraph(*TargetModule); // fallthrough
-		case 1: MIR::DebugDumpIR(Material->GetFullName(), *TargetModule); // fallthrough
+		case 1:
+		{
+			// Save the dump to file
+			FString FilePath = FPaths::Combine(FPaths::ProjectSavedDir(), "Materials", TEXT("IRDump.txt"));
+			FFileHelper::SaveStringToFile(Material->MaterialInsight.Get()->IRString, *FilePath);
+			// fallthrough
+		}
 	}
 
 	return true;
diff --git a/Engine/Source/Runtime/Engine/Private/Materials/MaterialInsights.cpp b/Engine/Source/Runtime/Engine/Private/Materials/MaterialInsights.cpp
--- a/Engine/Source/Runtime/Engine/Private/Materials/MaterialInsights.cpp
+++ b/Engine/Source/Runtime/Engine/Private/Materials/MaterialInsights.cpp
@@ -9,5 +9,11 @@ FMaterialInsights::FMaterialInsights()
 void FMaterialInsights::Empty()
 {
 	ConnectionInsights.Empty();
+	UniformParameterAllocationInsights.Empty();
+	IRString.Empty();
+	Legacy_ShaderStringParameters.Empty();
+	LegacyHLSLCode.Empty();
+	New_ShaderStringParameters.Empty();
+	NewHLSLCode.Empty();
 }
 
diff --git a/Engine/Source/Runtime/Engine/Private/Materials/MaterialShared.cpp b/Engine/Source/Runtime/Engine/Private/Materials/MaterialShared.cpp
--- a/Engine/Source/Runtime/Engine/Private/Materials/MaterialShared.cpp
+++ b/Engine/Source/Runtime/Engine/Private/Materials/MaterialShared.cpp
@@ -3313,8 +3313,6 @@ bool FMaterial::Translate_New(const FMaterialShaderMapId& InShaderMapId,
 
 	FMaterialIRModule Module;
 
-	GetMaterialInterface()->MaterialInsight.Reset(new FMaterialInsights);
-
 	// Build the material
 	FMaterialIRModuleBuilder Builder = {
 		.Material = GetMaterialInterface()->GetMaterial(),
@@ -3363,6 +3361,8 @@ bool FMaterial::Translate_New(const FMaterialShaderMapId& InShaderMapId,
 	// Emit uniform data debug information to the end of the generated shader
 	EmitDebugInfoComment(*GetMaterialInterface()->MaterialInsight, MaterialShaderCode);
 
+	GetMaterialInterface()->MaterialInsight->New_ShaderStringParameters = ShaderStringParameters;
+
 	OutMaterialEnvironment->IncludeVirtualPathToContentsMap.Add(TEXT("/Engine/Generated/Material.ush"), MoveTemp(MaterialShaderCode));
 	
 	return true;
@@ -3375,6 +3375,9 @@ bool FMaterial::Translate(const FMaterialShaderMapId& InShaderMapId,
 	FMaterialCompilationOutput& OutCompilationOutput,
 	TRefCountPtr<FSharedShaderCompilerEnvironment>& OutMaterialEnvironment)
 {
+	// Not all Insight data will be filled out by both translators.
+	GetMaterialInterface()->MaterialInsight.Reset(new FMaterialInsights);
+
 	if (InShaderMapId.bUsingNewHLSLGenerator)
 	{
 		return Translate_New(InShaderMapId, InStaticParameters, InPlatform, InTargetPlatform, OutCompilationOutput, OutMaterialEnvironment);
@@ -4385,13 +4388,7 @@ void FMaterial::GetExpressionIncludesHash(EShaderPlatform Platform, FSHAHash& Ou
 }
 #endif // WITH_EDITOR
 
-/**
- * Get user source code for the material, with a list of code snippets to highlight representing the code for each MaterialExpression
- * @param OutSource - generated source code
- * @param OutHighlightMap - source code highlight list
- * @return - true on Success
- */
-bool FMaterial::GetMaterialExpressionSource( FString& OutSource )
+bool FMaterial::GetMaterialExpressionSource(FString& OutSource)
 {
 #if WITH_EDITORONLY_DATA
 	const EShaderPlatform ShaderPlatform = GMaxRHIShaderPlatform;
@@ -4413,6 +4410,17 @@ bool FMaterial::GetMaterialExpressionSource( FString& OutSource )
 		{
 			OutSource = MoveTemp(*Source);
 
+			// If we've succesfully translated using the new translator, let's compile w/ the old translator so we have source to compare against.
+			if (ShaderMapId.bUsingNewHLSLGenerator)
+			{
+				ShaderMapId.bUsingNewHLSLGenerator = false;
+
+				// Translate_Legacy will fill out FMaterialInsights::Legacy_ShaderStringParameters
+				FMaterialCompilationOutput LegacyNewCompilationOutput;
+				TRefCountPtr<FSharedShaderCompilerEnvironment> LegacyMaterialEnvironment;
+				Translate_Legacy(ShaderMapId, StaticParameterSet, ShaderPlatform, TargetPlatform, LegacyNewCompilationOutput, LegacyMaterialEnvironment);
+			}
+
 			if (CVarMaterialEdPreshaderDumpToHLSL.GetValueOnGameThread())
 			{
 				OutSource.AppendChar('\n');
diff --git a/Engine/Source/Runtime/Engine/Public/Materials/MaterialIRDebug.h b/Engine/Source/Runtime/Engine/Public/Materials/MaterialIRDebug.h
--- a/Engine/Source/Runtime/Engine/Public/Materials/MaterialIRDebug.h
+++ b/Engine/Source/Runtime/Engine/Public/Materials/MaterialIRDebug.h
@@ -14,8 +14,8 @@ namespace MIR
 // iff it uses that other value as argument (e.g. an "Add" instruction will link to its operand values).
 void DebugDumpIRUseGraph(const FMaterialIRModule& Module);
 
-// Dumps the module to a human readable file in "{ProjectLogDir}/MaterialIRDump.txt"
-void DebugDumpIR(FStringView MaterialName, const FMaterialIRModule& Module);
+// Returns the IR generated by the module as a string.
+FString DebugDumpIR(FStringView MaterialName, const FMaterialIRModule& Module);
 
 } // namespace MIR
 
diff --git a/Engine/Source/Runtime/Engine/Public/Materials/MaterialInsights.h b/Engine/Source/Runtime/Engine/Public/Materials/MaterialInsights.h
--- a/Engine/Source/Runtime/Engine/Public/Materials/MaterialInsights.h
+++ b/Engine/Source/Runtime/Engine/Public/Materials/MaterialInsights.h
@@ -59,6 +59,21 @@ struct FMaterialInsights
 	/// Array of parameter allocation insights.
     TArray<FUniformParameterAllocationInsight> UniformParameterAllocationInsights;
 
+	/// String of the IR after translation.
+	FString IRString;
+
+	/// HLSL template string parameters for the legacy translator
+	TMap<FString, FString> Legacy_ShaderStringParameters;
+
+	/// HLSL full legacy generated source
+	FString LegacyHLSLCode;
+
+	/// HLSL template string parameters for the new translator
+	TMap<FString, FString> New_ShaderStringParameters;
+
+	/// HLSL full new generate source
+	FString NewHLSLCode;
+
     FMaterialInsights();
 
 	/// Clears all insight data.
