Make Material Parameter Collection's parameter name drop down searchable.

- Adds a SSearchBox above the list view.
- As the text changes in the search box filter the array of strings and update the list view.

[RN] minor | Rendering - Materials and Shaders

diff --git a/Engine/Source/Editor/MaterialEditor/Private/MaterialEditorDetailCustomization.cpp b/Engine/Source/Editor/MaterialEditor/Private/MaterialEditorDetailCustomization.cpp
--- a/Engine/Source/Editor/MaterialEditor/Private/MaterialEditorDetailCustomization.cpp
+++ b/Engine/Source/Editor/MaterialEditor/Private/MaterialEditorDetailCustomization.cpp
@@ -4,6 +4,7 @@
 #include "Widgets/Text/STextBlock.h"
 #include "Widgets/Input/SEditableText.h"
 #include "Widgets/Input/SComboButton.h"
+#include "Widgets/Input/SSearchBox.h"
 #include "Widgets/Views/SListView.h"
 #include "Materials/Material.h"
 #include "Materials/MaterialParameterCollection.h"
@@ -636,7 +637,7 @@ bool FMaterialExpressionCollectionParameterDetails::IsParameterNameComboEnabled(
 
 void FMaterialExpressionCollectionParameterDetails::OnCollectionChanged()
 {
-	PopulateParameters();
+	PopulateParameters(TEXT(""));
 }
 
 void FMaterialExpressionCollectionParameterDetails::CustomizeDetails( IDetailLayoutBuilder& DetailLayout )
@@ -658,7 +659,7 @@ void FMaterialExpressionCollectionParameterDetails::CustomizeDetails( IDetailLay
 	ParameterNamePropertyHandle->MarkHiddenByCustomization();
 	CollectionPropertyHandle->MarkHiddenByCustomization();
 
-	PopulateParameters();
+	PopulateParameters(TEXT(""));
 
 	TSharedPtr<SComboButton> NewComboButton;
 	TSharedPtr<SListView<TSharedPtr<FString>>> NewListView;
@@ -688,6 +689,13 @@ void FMaterialExpressionCollectionParameterDetails::CustomizeDetails( IDetailLay
 		[
 			SNew(SVerticalBox)
 			+SVerticalBox::Slot()
+			.Padding(FMargin(4.0f))
+			.AutoHeight()
+			[
+				SAssignNew(ParameterSearchBox, SSearchBox)
+				.OnTextChanged(this, &FMaterialExpressionCollectionParameterDetails::OnFilterTextChanged)
+			]
+			+SVerticalBox::Slot()
 			.AutoHeight()
 			.MaxHeight(400.0f)
 			[
@@ -705,7 +713,7 @@ void FMaterialExpressionCollectionParameterDetails::CustomizeDetails( IDetailLay
 	NewComboButton->SetToolTipText(GetToolTipText());
 }
 
-void FMaterialExpressionCollectionParameterDetails::PopulateParameters()
+void FMaterialExpressionCollectionParameterDetails::PopulateParameters(const FString& TextFilter)
 {
 	UMaterialParameterCollection* Collection = nullptr;
 	
@@ -742,7 +750,11 @@ void FMaterialExpressionCollectionParameterDetails::PopulateParameters()
 
 		for (const FName& Name : NameList)
 		{
-			ParametersSource.Add(MakeShareable(new FString(Name.ToString())));
+			TSharedPtr<FString> ParamNameString = MakeShareable(new FString(Name.ToString()));
+			if (TextFilter.Len() == 0 || ParamNameString->Contains(TextFilter))
+			{
+				ParametersSource.Add(ParamNameString);
+			}
 		}
 	}
 
@@ -770,6 +782,12 @@ void FMaterialExpressionCollectionParameterDetails::OnSelectionChanged( TSharedP
 	}
 }
 
+void FMaterialExpressionCollectionParameterDetails::OnFilterTextChanged(const FText& InFilterText)
+{
+	PopulateParameters(InFilterText.ToString());
+	ParameterListView.Pin()->RequestListRefresh();
+}
+
 TSharedRef<class IDetailCustomization> FMaterialDetailCustomization::MakeInstance()
 {
 	return MakeShareable( new FMaterialDetailCustomization );
diff --git a/Engine/Source/Editor/MaterialEditor/Private/MaterialEditorDetailCustomization.h b/Engine/Source/Editor/MaterialEditor/Private/MaterialEditorDetailCustomization.h
--- a/Engine/Source/Editor/MaterialEditor/Private/MaterialEditorDetailCustomization.h
+++ b/Engine/Source/Editor/MaterialEditor/Private/MaterialEditorDetailCustomization.h
@@ -96,11 +96,12 @@ public:
 
 private:
 	/** Populates the group name that this parameter details uses */
-	void PopulateParameters();
+	void PopulateParameters(const FString& FilterText);
 
 	/** Functions used to make the custom combo button with editable text */
 	TSharedRef< ITableRow > MakeDetailsGroupViewWidget( TSharedPtr<FString> Item, const TSharedRef< STableViewBase >& OwnerTable );
 	void OnSelectionChanged( TSharedPtr<FString> ProposedSelection, ESelectInfo::Type SelectInfo );
+	void OnFilterTextChanged(const FText& InFilterText);
 
 	FText GetToolTipText() const;
 	FText GetParameterNameString() const;
@@ -116,6 +117,8 @@ private:
 
 	TWeakPtr<SComboButton> ParameterComboButton;
 	TWeakPtr<SListView<TSharedPtr<FString>>> ParameterListView;
+	/** Text entry to filter parameters */
+	TSharedPtr<class SSearchBox> ParameterSearchBox;
 
 	/** A list of all parameter names to choose from */
 	TArray<TSharedPtr<FString>> ParametersSource;
