Material Editor: Double clicking on the parameter rows in the parameter tab jumps to and selects parameter nodes.

diff --git a/Engine/Source/Editor/MaterialEditor/Private/MaterialEditor.cpp b/Engine/Source/Editor/MaterialEditor/Private/MaterialEditor.cpp
--- a/Engine/Source/Editor/MaterialEditor/Private/MaterialEditor.cpp
+++ b/Engine/Source/Editor/MaterialEditor/Private/MaterialEditor.cpp
@@ -1782,7 +1782,9 @@ void FMaterialEditor::CreateInternalWidgets()
 
 	MaterialParametersOverviewWidget = SNew(SMaterialParametersOverviewPanel)
 		.InMaterialEditorInstance(MaterialEditorInstance)
-		.InGenerator(Generator);
+		.InGenerator(Generator)
+		.SelectMaterialNode(this, &FMaterialEditor::JumpToNodeByGuid);
+
 	Generator->OnFinishedChangingProperties().AddSP(this, &FMaterialEditor::OnFinishedChangingParametersFromOverview);
 	Generator->OnRowsRefreshed().AddSP(this, &FMaterialEditor::GeneratorRowsRefreshed);
 	MaterialCustomPrimitiveDataWidget = SNew(SMaterialCustomPrimitiveDataPanel, MaterialEditorInstance);
@@ -6291,7 +6293,10 @@ void FMaterialEditor::JumpToExpression(UMaterialExpression* Expression)
 		ExpressionNode = GraphExpression ? ToRawPtr((*GraphExpression)->GraphNode) : nullptr;
 	}
 
-	JumpToNode(ExpressionNode);
+	if (ExpressionNode)
+	{
+		JumpToNode(ExpressionNode);
+	}
 }
 
 void FMaterialEditor::ModifyMaterial()
@@ -8285,6 +8290,101 @@ void FMaterialEditor::SetEditorInstanceForWidgets() const
 	MaterialExpressionsOverviewWidget->SetEditorInstance(MaterialEditorInstance);
 }
 
+void FMaterialEditor::JumpToNodeByGuid(const FGuid& InGuid)
+{
+	if (!InGuid.IsValid() || !IsValid(Material))
+	{
+		return;
+	}
+
+	UMaterialExpression* Expression = Material->FindExpressionByGUID<UMaterialExpression>(InGuid);
+
+	if (!Expression)
+	{
+		return;
+	}
+
+	// A parameter with a valid opened graph node can be jumped to straight away.
+	if (IsValid(Expression->GraphNode))
+	{
+		JumpToExpression(Expression);
+		return;
+	}
+
+	if (!GEngine)
+	{
+		return;
+	}
+
+	UMaterial* ParentMaterial = Cast<UMaterial>(Expression->GetOuter());
+	UMaterialFunction* ParentFunction = Cast<UMaterialFunction>(Expression->GetOuter());
+	UObject* AssetToOpen = nullptr;
+
+	if (ParentMaterial && ParentMaterial != OriginalMaterial)
+	{
+		AssetToOpen = ParentMaterial;
+	}
+	else if (ParentFunction && ParentFunction != MaterialFunction)
+	{
+		AssetToOpen = ParentFunction;
+	}
+
+	if (!AssetToOpen)
+	{
+		return;
+	}
+		
+	UAssetEditorSubsystem* AssetEditorSubsystem = GEditor->GetEditorSubsystem<UAssetEditorSubsystem>();
+
+	if (!AssetEditorSubsystem)
+	{
+		return;
+	}
+
+	IAssetEditorInstance* CurrentInstance = AssetEditorSubsystem->FindEditorForAsset(AssetToOpen, /* Focus */ true);
+
+	if (!CurrentInstance || CurrentInstance->GetEditorName() != GetEditorName())
+	{
+		IMaterialEditorModule& MaterialEditorModule = FModuleManager::LoadModuleChecked<IMaterialEditorModule>("MaterialEditor");
+
+		if (ParentMaterial)
+		{
+			if (TSharedPtr<IMaterialEditor> NewEditor = MaterialEditorModule.CreateMaterialEditor(EToolkitMode::Standalone, nullptr, ParentMaterial))
+			{
+				CurrentInstance = NewEditor.Get();
+			}
+		}
+		else if (ParentFunction)
+		{
+			if (TSharedPtr<IMaterialEditor> NewEditor = MaterialEditorModule.CreateMaterialEditor(EToolkitMode::Standalone, nullptr, ParentFunction))
+			{
+				CurrentInstance = NewEditor.Get();
+			}
+		}
+
+		if (!CurrentInstance)
+		{
+			return;
+		}
+	}
+
+	// We've ensured that this is a valid cast by checking GetEditorName()
+	IMaterialEditor* MaterialAssetEditor = static_cast<IMaterialEditor*>(CurrentInstance);
+	UMaterial* PreviewMaterial = Cast<UMaterial>(MaterialAssetEditor->GetMaterialInterface());
+
+	if (!PreviewMaterial)
+	{
+		return;
+	}
+
+	Expression = PreviewMaterial->FindExpressionByGUID<UMaterialExpression>(InGuid);
+
+	if (Expression)
+	{
+		MaterialAssetEditor->JumpToExpression(Expression);
+	}
+}
+
 void FMaterialEditor::UpdateDetailView()
 {
 	GetDetailView()->InvalidateCachedState();
diff --git a/Engine/Source/Editor/MaterialEditor/Private/MaterialEditor.h b/Engine/Source/Editor/MaterialEditor/Private/MaterialEditor.h
--- a/Engine/Source/Editor/MaterialEditor/Private/MaterialEditor.h
+++ b/Engine/Source/Editor/MaterialEditor/Private/MaterialEditor.h
@@ -979,6 +979,9 @@ private:
 	/* Updates internal material editor widgets */
 	void SetEditorInstanceForWidgets() const;
 
+	/* Jumps to and selects the node corresponding to the given guid, even if it exists in another asset. */
+	void JumpToNodeByGuid(const FGuid& InGuid);
+
 private:
 	TSharedRef<SDockTab> SpawnTab_Preview(const FSpawnTabArgs& Args);
 	TSharedRef<SDockTab> SpawnTab_MaterialProperties(const FSpawnTabArgs& Args);
diff --git a/Engine/Source/Editor/MaterialEditor/Private/MaterialEditorDelegates.h b/Engine/Source/Editor/MaterialEditor/Private/MaterialEditorDelegates.h
new file mode 100644
--- /dev/null
+++ b/Engine/Source/Editor/MaterialEditor/Private/MaterialEditorDelegates.h
@@ -0,0 +1,10 @@
+// Copyright Epic Games, Inc. All Rights Reserved.
+
+#pragma once
+
+#include "Delegates/Delegate.h"
+#include "Delegates/DelegateCombinations.h"
+
+struct FGuid;
+
+DECLARE_DELEGATE_OneParam(FSelectMaterialNode, const FGuid& InNodeGuid);
diff --git a/Engine/Source/Editor/MaterialEditor/Private/SMaterialParametersOverviewWidget.cpp b/Engine/Source/Editor/MaterialEditor/Private/SMaterialParametersOverviewWidget.cpp
--- a/Engine/Source/Editor/MaterialEditor/Private/SMaterialParametersOverviewWidget.cpp
+++ b/Engine/Source/Editor/MaterialEditor/Private/SMaterialParametersOverviewWidget.cpp
@@ -528,13 +528,11 @@ void SMaterialParametersOverviewTreeItem::Construct(const FArguments& InArgs, co
 			];
 	}
 
-
 	this->ChildSlot
 		[
 			WrapperWidget
 		];
 
-
 	STableRow< TSharedPtr<FSortedParamData> >::ConstructInternal(
 		STableRow::FArguments()
 		.Style(FAppStyle::Get(), "DetailsView.TreeView.TableRow")
@@ -548,6 +546,7 @@ void SMaterialParametersOverviewTree::Construct(const FArguments& InArgs)
 	bHasAnyParameters = false;
 	MaterialEditorInstance = InArgs._InMaterialEditorInstance;
 	Owner = InArgs._InOwner;
+	SelectMaterialNodeDelegate = InArgs._SelectMaterialNode;
 	CreateGroupsWidget();
 
 	STreeView<TSharedPtr<FSortedParamData>>::Construct(
@@ -746,6 +745,23 @@ void SMaterialParametersOverviewTree::CreateGroupsWidget()
 }
 
 
+bool SMaterialParametersOverviewTree::Private_OnItemDoubleClicked(TSharedPtr<FSortedParamData> TheItem)
+{
+	if (STreeView<TSharedPtr<FSortedParamData>>::Private_OnItemDoubleClicked(TheItem))
+	{
+		return true;
+	}
+
+	if (!SelectMaterialNodeDelegate.IsBound() || !TheItem.IsValid() || !TheItem->Parameter || 
+		!TheItem->Parameter->ExpressionId.IsValid())
+	{
+		return false;
+	}
+
+	SelectMaterialNodeDelegate.Execute(TheItem->Parameter->ExpressionId);
+
+	return true;
+}
 
 void SMaterialParametersOverviewTree::ShowSubParameters()
 {
@@ -934,7 +950,8 @@ void SMaterialParametersOverviewPanel::Construct(const FArguments& InArgs)
 	NestedTree = SNew(SMaterialParametersOverviewTree)
 		.InMaterialEditorInstance(InArgs._InMaterialEditorInstance)
 		.InOwner(SharedThis(this))
-		.InScrollbar(ExternalScrollbar);
+		.InScrollbar(ExternalScrollbar)
+		.SelectMaterialNode(InArgs._SelectMaterialNode);
 
 	MaterialEditorInstance = InArgs._InMaterialEditorInstance;
 	Refresh();
diff --git a/Engine/Source/Editor/MaterialEditor/Private/SMaterialParametersOverviewWidget.h b/Engine/Source/Editor/MaterialEditor/Private/SMaterialParametersOverviewWidget.h
--- a/Engine/Source/Editor/MaterialEditor/Private/SMaterialParametersOverviewWidget.h
+++ b/Engine/Source/Editor/MaterialEditor/Private/SMaterialParametersOverviewWidget.h
@@ -6,6 +6,7 @@
 #include "DetailColumnSizeData.h"
 #include "Engine/EngineTypes.h"
 #include "IDetailTreeNode.h"
+#include "MaterialEditorDelegates.h"
 #include "MaterialPropertyHelpers.h"
 #include "Materials/Material.h"
 #include "PropertyCustomizationHelpers.h"
@@ -73,6 +74,7 @@ public:
 
 	SLATE_ARGUMENT(UMaterialEditorPreviewParameters*, InMaterialEditorInstance)
 	SLATE_ARGUMENT(TSharedPtr<class IPropertyRowGenerator>, InGenerator)
+	SLATE_EVENT(FSelectMaterialNode, SelectMaterialNode)
 	SLATE_END_ARGS()
 	void Refresh();
 	void Construct(const FArguments& InArgs);
@@ -112,6 +114,7 @@ public:
 	SLATE_ARGUMENT(UMaterialEditorPreviewParameters*, InMaterialEditorInstance)
 	SLATE_ARGUMENT(TSharedPtr<SMaterialParametersOverviewPanel>, InOwner)
 	SLATE_ARGUMENT(TSharedPtr<SScrollBar>, InScrollbar)
+	SLATE_EVENT(FSelectMaterialNode, SelectMaterialNode)
 	SLATE_END_ARGS()
 
 	/** Constructs this widget with InArgs */
@@ -134,6 +137,8 @@ public:
 
 	FDetailColumnSizeData& GetColumnSizeData() { return ColumnSizeData; }
 
+	virtual bool Private_OnItemDoubleClicked(TSharedPtr<FSortedParamData> TheItem) override;
+
 protected:
 
 	void ShowSubParameters();
@@ -149,4 +154,7 @@ private:
 	bool bHasAnyParameters;
 	
 	FDetailColumnSizeData ColumnSizeData;
+
+	/** Called when a node should be selected in the material editor. */
+	FSelectMaterialNode SelectMaterialNodeDelegate;
 };
