Add warning about asset compatibility when saving a Substrate's auto-converted material.

diff --git a/Engine/Source/Editor/MaterialEditor/Private/MaterialEditor.cpp b/Engine/Source/Editor/MaterialEditor/Private/MaterialEditor.cpp
--- a/Engine/Source/Editor/MaterialEditor/Private/MaterialEditor.cpp
+++ b/Engine/Source/Editor/MaterialEditor/Private/MaterialEditor.cpp
@@ -210,6 +210,33 @@ static void InvalidateSubstrateConversionVersion(UMaterial* Material)
 	}
 }
 
+static bool IsSubstrateAutoConvertedAndNotSaved(const UMaterial* Material)
+{
+	if (Substrate::IsSubstrateEnabled())
+	{
+		if (const UMaterialEditorOnlyData* EditorOnly = Material->GetEditorOnlyData())
+		{
+			return EditorOnly->IsSubstrateAutoConverted() && EditorOnly->bIsSubsutrateAutoConvertedAndNotSaved;
+		}
+	}
+	return false;
+}
+
+static bool MarkSubstrateAutoConversionSaved(UMaterial* Material)
+{
+	if (Substrate::IsSubstrateEnabled())
+	{
+		if (UMaterialEditorOnlyData* EditorOnly = Material->GetEditorOnlyData())
+		{
+			if (EditorOnly->IsSubstrateAutoConverted() && EditorOnly->bIsSubsutrateAutoConvertedAndNotSaved)
+			{
+				EditorOnly->bIsSubsutrateAutoConvertedAndNotSaved = false;
+			}
+		}
+	}
+	return false;
+}
+
 ///////////////////////////
 // FMatExpressionPreview //
 ///////////////////////////
@@ -2272,10 +2299,33 @@ void FMaterialEditor::SaveAsset_Execute()
 {
 	UE_LOG(LogMaterialEditor, Log, TEXT("Saving and Compiling material %s"), *GetEditingObjects()[0]->GetName());
 
+	// When Substrate is enabled, warns user about backward compatibility issue when a material is auto-converted 
+	// and is about to be saved
+	const bool bSaveSubstrateAutoConversion = OriginalMaterial && Substrate::IsSubstrateEnabled() && IsSubstrateAutoConvertedAndNotSaved(OriginalMaterial);
+	if (bSaveSubstrateAutoConversion)
+	{
+		// Find out the user wants to do with this dirty material
+		EAppReturnType::Type YesNoCancelReply = FMessageDialog::Open(EAppMsgType::YesNoCancel,
+			FText::Format(
+			NSLOCTEXT("UnrealEd", "Prompt_MaterialEditorSaveAsset", "This material has been automatically converted to be compatible with Substrate.\nSaving this material will make is no longer works when Substrate is disabled.\nWould you like to apply the changes of the modified material to the original material?\n\n{0}\n\n(Selecting 'No' will cause all changes to be lost!)"),
+			FText::FromString(OriginalMaterialObject->GetPathName())));
+
+		// Skip saving the asset, based on user's choice
+		if (YesNoCancelReply != EAppReturnType::Yes)
+		{
+			return;
+		}
+	}
+
 	const bool bUpdateSucceeded = (bMaterialDirty ? UpdateOriginalMaterial() : true);
 
 	if (bUpdateSucceeded)
 	{
+		if (bSaveSubstrateAutoConversion)
+		{
+			MarkSubstrateAutoConversionSaved(OriginalMaterial);
+		}
+
 		IMaterialEditor::SaveAsset_Execute();
 	}
 }
diff --git a/Engine/Source/Runtime/Engine/Private/Materials/Material.cpp b/Engine/Source/Runtime/Engine/Private/Materials/Material.cpp
--- a/Engine/Source/Runtime/Engine/Private/Materials/Material.cpp
+++ b/Engine/Source/Runtime/Engine/Private/Materials/Material.cpp
@@ -4126,6 +4126,7 @@ bool UMaterial::ConvertMaterialToSubstrateMaterial(bool bAllowEmptyMaterialUpdat
 			else // == Converted shader
 			{
 				EditorOnly->SubstrateConversionVersion = GetSubstrateConversionVersion();
+				EditorOnly->bIsSubsutrateAutoConvertedAndNotSaved = true;
 			}
 		}
 	}
@@ -7926,4 +7927,10 @@ void UMaterialEditorOnlyData::ResetSubstrateConversionVersion()
 {
 	SubstrateConversionVersion = GetSubstrateNoConversionVersion();
 }
+
+bool UMaterialEditorOnlyData::IsSubstrateAutoConverted() const
+{
+	return SubstrateConversionVersion > GetSubstrateNoConversionVersion();
+}
+
 #undef LOCTEXT_NAMESPACE
diff --git a/Engine/Source/Runtime/Engine/Public/Materials/Material.h b/Engine/Source/Runtime/Engine/Public/Materials/Material.h
--- a/Engine/Source/Runtime/Engine/Public/Materials/Material.h
+++ b/Engine/Source/Runtime/Engine/Public/Materials/Material.h
@@ -410,8 +410,15 @@ public:
 	UPROPERTY()
 	int32 SubstrateVersion;
 
+	/** Whether the material has been auto - converted at load time, and hasn't been saved. */
+	UPROPERTY()
+	bool bIsSubsutrateAutoConvertedAndNotSaved = false;
+
 	/** Reset the version of the Substrate's auto-conversion applied to this material */
 	ENGINE_API void ResetSubstrateConversionVersion();
+
+	/** Return true if the Substrate's auto-conversion applied to this material */
+	ENGINE_API bool IsSubstrateAutoConverted() const;
 };
 
 /**
