[UMG - 16670440] MaterialEditor Module missing slate image/icon files

diff --git a/Engine/Plugins/Experimental/DynamicMaterial/Source/DynamicMaterialEditor/Private/DynamicMaterialEditorStyle.cpp b/Engine/Plugins/Experimental/DynamicMaterial/Source/DynamicMaterialEditor/Private/DynamicMaterialEditorStyle.cpp
--- a/Engine/Plugins/Experimental/DynamicMaterial/Source/DynamicMaterialEditor/Private/DynamicMaterialEditorStyle.cpp
+++ b/Engine/Plugins/Experimental/DynamicMaterial/Source/DynamicMaterialEditor/Private/DynamicMaterialEditorStyle.cpp
@@ -46,7 +46,7 @@ void FDynamicMaterialEditorStyle::SetupGeneralStyles()
 	TSharedPtr<IPlugin> Plugin = IPluginManager::Get().FindPlugin(UE_PLUGIN_NAME);
 	check(Plugin.IsValid());
 
-	SetContentRoot(FPaths::Combine(Plugin->GetBaseDir(), TEXT("Resources")));
+	SetContentRoot(FPaths::EngineContentDir() / TEXT("Editor/Slate/DynamicMaterial"));
 	SetCoreContentRoot(FPaths::EngineContentDir() / TEXT("Editor/Slate")); // This is the engine's content root.
 
 	////////////////////////////////////////////////////////////////////////////////////////////////////
@@ -458,7 +458,7 @@ void FDynamicMaterialEditorStyle::SetupComponentIcons()
 	Set(TEXT("ClassIcon.DMMaterialValueFloat2"),       new IMAGE_BRUSH("Icons/ValueTypes/Float2", CoreStyleConstants::Icon16x16));
 	Set(TEXT("ClassIcon.DMMaterialValueFloat3RGB"),    new FSlateImageBrush(FPaths::EngineContentDir() + TEXT("Slate/Common/ColorPicker_Mode_16x.png"), CoreStyleConstants::Icon16x16));
 	Set(TEXT("ClassIcon.DMMaterialValueFloat3RPY"),    new CORE_IMAGE_BRUSH("Icons/icon_ClockwiseRotation_16x", CoreStyleConstants::Icon16x16));
-	Set(TEXT("ClassIcon.DMMaterialValueFloat3XYZ"),    new CORE_IMAGE_BRUSH("Icons/Movable_16x", CoreStyleConstants::Icon16x16));
+	Set(TEXT("ClassIcon.DMMaterialValueFloat3XYZ"),    new CORE_IMAGE_BRUSH("Icons/Mobility/Movable_16x", CoreStyleConstants::Icon16x16));
 	Set(TEXT("ClassIcon.DMMaterialValueFloat4"),       new FSlateImageBrush(FPaths::EngineContentDir() + TEXT("Slate/Common/ColorPicker_Mode_16x.png"), CoreStyleConstants::Icon16x16));
 	Set(TEXT("ClassIcon.DMMaterialValueRenderTarget"), new CORE_IMAGE_BRUSH("Icons/AssetIcons/TextureRenderTarget2D_16x", CoreStyleConstants::Icon16x16));
 	Set(TEXT("ClassIcon.DMMaterialValueTexture"),      new IMAGE_BRUSH("Icons/ValueTypes/Texture", CoreStyleConstants::Icon16x16));
diff --git a/Engine/Source/Editor/MaterialEditor/Private/MaterialEditorStyle.cpp b/Engine/Source/Editor/MaterialEditor/Private/MaterialEditorStyle.cpp
--- a/Engine/Source/Editor/MaterialEditor/Private/MaterialEditorStyle.cpp
+++ b/Engine/Source/Editor/MaterialEditor/Private/MaterialEditorStyle.cpp
@@ -80,7 +80,7 @@ TSharedRef<FSlateStyleSet> FSubstrateMaterialEditorStyle::Create()
 {
 	TSharedRef<FSlateStyleSet> Style = MakeShared<FSlateStyleSet>(TEXT("SubstrateMaterial"));
 
-	Style->SetContentRoot(FPaths::EngineContentDir() / TEXT("Editor/Slate"));
+	Style->SetContentRoot(FPaths::EngineContentDir() / TEXT("Editor/Slate/DynamicMaterial"));
 	Style->SetCoreContentRoot(FPaths::EngineContentDir() / TEXT("Slate"));
 	////////////////////////////////////////////////////////////////////////////////////////////////////
 	// Color Styles
@@ -196,6 +196,7 @@ void FSubstrateMaterialEditorStyle::SetupLayerViewStyles(const TSharedRef<FSlate
 
 	constexpr float LayerViewItemCornerRadius = 10.0f;
 	constexpr float LayerViewItemBorderWidth = 1.0f;
+	constexpr float DropZoneMargin = 0.25f;
 
 	const FLinearColor LayerViewItem1FillColor = FStyleColors::Dropdown.GetSpecifiedColor();
 	const FLinearColor LayerViewItem2FillColor = FStyleColors::DropdownOutline.GetSpecifiedColor();
@@ -238,7 +239,16 @@ void FSubstrateMaterialEditorStyle::SetupLayerViewStyles(const TSharedRef<FSlate
 		LayerViewItem2FillColor, LayerViewItemCornerRadius,
 		LayerItemSelectBorderColor, LayerViewItemBorderWidth));
 
-	const float DropZoneMargin = 0.25f; 
+	Style->Set("LayerView.DropIndicatorOnto", new FSlateBoxBrush(
+		FPaths::EngineContentDir() + TEXT("Editor/Slate/Common/DropZoneIndicator_Onto.png"),
+		FMargin(4.0f / 16.0f), Style->GetColor(TEXT("Color.Select.Hover"))));
+	Style->Set("LayerView.DropIndicatorAbove", new FSlateBorderBrush(
+		FPaths::EngineContentDir() + TEXT("Editor/Slate/Common/LayersDropZoneDashed_Above.png"), 
+		FMargin(DropZoneMargin, DropZoneMargin, 0.f, 0.f), Style->GetColor(TEXT("Color.Select.Hover"))));
+	Style->Set("LayerView.DropIndicatorBelow", new FSlateBorderBrush(
+		FPaths::EngineContentDir() + TEXT("Editor/Slate/Common/LayersDropZoneDashed_Below.png"), 
+		FMargin(DropZoneMargin, 0, 0, DropZoneMargin), Style->GetColor(TEXT("Color.Select.Hover"))));
+
 	Style->Set("LayerView.Row1", FTableRowStyle()
 		.SetTextColor(FStyleColors::Foreground)
 		.SetSelectedTextColor(FStyleColors::ForegroundHover)
@@ -251,9 +261,9 @@ void FSubstrateMaterialEditorStyle::SetupLayerViewStyles(const TSharedRef<FSlate
 		.SetActiveHoveredBrush(*Style->GetBrush(TEXT("LayerView.Row1.ActiveHoveredBrush")))
 		.SetInactiveBrush(*Style->GetBrush(TEXT("LayerView.Row1.InactiveBrush")))
 		.SetInactiveHoveredBrush(*Style->GetBrush(TEXT("LayerView.Row1.InactiveHoveredBrush")))
-		.SetDropIndicator_Onto(BOX_BRUSH("Common/DropZoneIndicator_Onto", FMargin(4.0f / 16.0f), Style->GetColor(TEXT("Color.Select.Hover"))))
-		.SetDropIndicator_Above(BORDER_BRUSH("Common/LayersDropZoneDashed_Above", FMargin(DropZoneMargin, DropZoneMargin, 0.f, 0.f), Style->GetColor(TEXT("Color.Select.Hover"))))
-		.SetDropIndicator_Below(BORDER_BRUSH("Common/LayersDropZoneDashed_Below", FMargin(DropZoneMargin, 0, 0, DropZoneMargin), Style->GetColor(TEXT("Color.Select.Hover"))))
+		.SetDropIndicator_Onto(*Style->GetBrush(TEXT("LayerView.DropIndicatorOnto")))
+		.SetDropIndicator_Above(*Style->GetBrush(TEXT("LayerView.DropIndicatorAbove")))
+		.SetDropIndicator_Below(*Style->GetBrush(TEXT("LayerView.DropIndicatorBelow")))
 	);
 	Style->Set("LayerView.Row2", FTableRowStyle()
 		.SetTextColor(FStyleColors::Foreground)
@@ -267,9 +277,9 @@ void FSubstrateMaterialEditorStyle::SetupLayerViewStyles(const TSharedRef<FSlate
 		.SetActiveHoveredBrush(*Style->GetBrush(TEXT("LayerView.Row2.ActiveHoveredBrush")))
 		.SetInactiveBrush(*Style->GetBrush(TEXT("LayerView.Row2.InactiveBrush")))
 		.SetInactiveHoveredBrush(*Style->GetBrush(TEXT("LayerView.Row2.InactiveHoveredBrush")))
-		.SetDropIndicator_Onto(BOX_BRUSH("Common/DropZoneIndicator_Onto", FMargin(4.0f / 16.0f), Style->GetColor(TEXT("Color.Select.Hover"))))
-		.SetDropIndicator_Above(BORDER_BRUSH("Common/LayersDropZoneDashed_Above", FMargin(DropZoneMargin, DropZoneMargin, 0.f, 0.f), Style->GetColor(TEXT("Color.Select.Hover"))))
-		.SetDropIndicator_Below(BORDER_BRUSH("Common/LayersDropZoneDashed_Below", FMargin(DropZoneMargin, 0, 0, DropZoneMargin), Style->GetColor(TEXT("Color.Select.Hover"))))
+		.SetDropIndicator_Onto(*Style->GetBrush(TEXT("LayerView.DropIndicatorOnto")))
+		.SetDropIndicator_Above(*Style->GetBrush(TEXT("LayerView.DropIndicatorAbove")))
+		.SetDropIndicator_Below(*Style->GetBrush(TEXT("LayerView.DropIndicatorBelow")))
 	);
 	Style->Set("LayerView.AddIcon", new IMAGE_BRUSH("Icons/EditorIcons/LayerAdd", Icon16x16));
 	Style->Set("LayerView.DuplicateIcon", new IMAGE_BRUSH("Icons/EditorIcons/Duplicate_40x", Icon40x40));
