Fix for topology tree previews not matching up correctly in the new layering system. Reverted the editor resolve code from the previous crash fix as this also fixes that issue.

diff --git a/Engine/Source/Editor/MaterialEditor/Private/SMaterialEditorTopologyWidget.cpp b/Engine/Source/Editor/MaterialEditor/Private/SMaterialEditorTopologyWidget.cpp
--- a/Engine/Source/Editor/MaterialEditor/Private/SMaterialEditorTopologyWidget.cpp
+++ b/Engine/Source/Editor/MaterialEditor/Private/SMaterialEditorTopologyWidget.cpp
@@ -4,7 +4,6 @@
 
 #include "EditorWidgetsModule.h"
 #include "AssetRegistry/AssetRegistryModule.h"
-#include "Materials/MaterialInterface.h"
 #include "MaterialEditor.h"
 #include "MaterialInstanceEditor.h"
 #include "SubstrateDefinitions.h"
@@ -116,10 +115,9 @@ void SMaterialEditorTopologyWidget::Tick(const FGeometry& AllottedGeometry, cons
 	if (MaterialEditorPtr.IsValid())
 	{
 		TSharedPtr<IMaterialEditor> MaterialEditorPinned = MaterialEditorPtr.Pin();
-		UMaterialInterface* MaterialInterface = MaterialEditorPinned->GetMaterialInterface();
-		FMaterialResource* MaterialResource = MaterialInterface ? MaterialInterface->GetMaterialResource(GMaxRHIFeatureLevel) : nullptr;
+		FMaterialResource* MaterialResource = MaterialEditorPinned->GetMaterialInterface()->GetMaterialResource(GMaxRHIFeatureLevel);
 
-		if (MaterialInterface && MaterialInterface->IsA<UMaterial>())
+		if (Cast<UMaterial>(MaterialEditorPinned->GetMaterialInterface()))
 		{
 			TSharedPtr<FMaterialEditor> MaterialEditor = StaticCastSharedPtr<FMaterialEditor>(MaterialEditorPinned);
 			UMaterial* MaterialForStats = MaterialEditor->bStatsFromPreviewMaterial ? MaterialEditor->Material : MaterialEditor->OriginalMaterial;
@@ -128,7 +126,7 @@ void SMaterialEditorTopologyWidget::Tick(const FGeometry& AllottedGeometry, cons
 
 		SubstrateMaterialDescription = FText::FromString(FString(TEXT("SubstrateMaterialDescription")));
 		
-		if (MaterialInterface && MaterialResource)
+		if (MaterialResource)
 		{
 			FString MaterialDescription;
 			FString MaterialBudget;
@@ -145,13 +143,8 @@ void SMaterialEditorTopologyWidget::Tick(const FGeometry& AllottedGeometry, cons
 				// Now generate a visual representation of the material from the topology tree of operators.
 				if (CompilationOutput.RootOperatorIndex >= 0)
 				{
-					FMaterialLayersFunctions LayersFunctions;
-					if (MaterialInterface->GetMaterialLayers(LayersFunctions))
-					{
-						MaterialInterface->SyncLayersRuntimeGraphCache(&LayersFunctions);
-					}
-					MaterialBox->SetContent(ProcessOperatorAsThumbnails(CompilationOutput, &LayersFunctions));
-					MaterialInterface->SyncLayersRuntimeGraphCache(nullptr);
+					const FMaterialLayersFunctions* LayersFunctions = MaterialResource->GetMaterialLayers();
+					MaterialBox->SetContent(ProcessOperatorAsThumbnails(CompilationOutput, LayersFunctions));
 				}
 				else
 				{
diff --git a/Engine/Source/Editor/MaterialEditor/Private/SMaterialSubstrateTree.cpp b/Engine/Source/Editor/MaterialEditor/Private/SMaterialSubstrateTree.cpp
--- a/Engine/Source/Editor/MaterialEditor/Private/SMaterialSubstrateTree.cpp
+++ b/Engine/Source/Editor/MaterialEditor/Private/SMaterialSubstrateTree.cpp
@@ -1310,33 +1310,44 @@ void SMaterialSubstrateTree::RecursiveCreateWidgets(FRecursiveCreateWidgetsConte
 		{
 			UObject* AssetObject = nullptr;
 			ChildProperty->ParameterHandle->GetValue(AssetObject);
-			if (AssetObject)
+			if(Substrate::IsMaterialLayeringSupportEnabled())
 			{
-				if (MaterialEditorInstance->StoredLayerPreviews[Payload.Layer] == nullptr)
-				{
-					MaterialEditorInstance->StoredLayerPreviews[Payload.Layer] = (NewObject<UMaterialInstanceConstant>(MaterialEditorInstance, NAME_None));
-				}
-				UMaterialInterface* EditedMaterial = Cast<UMaterialFunctionInterface>(AssetObject)->GetPreviewMaterial();
-				if (MaterialEditorInstance->StoredLayerPreviews[Payload.Layer] && MaterialEditorInstance->StoredLayerPreviews[Payload.Layer]->Parent != EditedMaterial)
+				if (AssetObject)
 				{
-					MaterialEditorInstance->StoredLayerPreviews[Payload.Layer]->SetParentEditorOnly(EditedMaterial);
-				}
-				UMaterialInterface* MaterialInterface = MaterialEditorInstance->GetMaterialInterface();
-				const FMaterialResource* MaterialResource = MaterialInterface->GetMaterialResource(ERHIFeatureLevel::SM6);
-				const FMaterialLayersFunctions* LayersFunctions = MaterialResource->GetMaterialLayers();
-				
-				if(Substrate::IsMaterialLayeringSupportEnabled())
-				{
-					TSharedPtr<FMaterialLayersFunctionsRuntimeGraphCache> LayerTreeCache = LayersFunctions->RuntimeGraphCache;
-					if (LayerTreeCache)
+					UMaterialInterface* MaterialInterface = MaterialEditorInstance->GetMaterialInterface();
+					FMaterialLayersFunctions LayersFunctions;
+					if (MaterialInterface && MaterialInterface->GetMaterialLayers(LayersFunctions))
 					{
-						LayerTreeCache->NodePreviewMaterials[InNodeId] = (UMaterial*)EditedMaterial;
+						MaterialInterface->SyncLayersRuntimeGraphCache(&LayersFunctions);
+						if (MaterialEditorInstance->StoredLayerPreviews[Payload.Layer] == nullptr)
+						{
+							MaterialEditorInstance->StoredLayerPreviews[Payload.Layer] = (NewObject<UMaterialInstanceConstant>(MaterialEditorInstance, NAME_None));
+						}
+						UMaterialInterface* EditedMaterial = Cast<UMaterialFunctionInterface>(AssetObject)->GetPreviewMaterial();
+						if (MaterialEditorInstance->StoredLayerPreviews[Payload.Layer] && MaterialEditorInstance->StoredLayerPreviews[Payload.Layer]->Parent != EditedMaterial)
+						{
+							MaterialEditorInstance->StoredLayerPreviews[Payload.Layer]->SetParentEditorOnly(EditedMaterial);
+						}
+
+						TSharedPtr<FMaterialLayersFunctionsRuntimeGraphCache> LayerTreeCache = LayersFunctions.RuntimeGraphCache;
+						if (LayerTreeCache)
+						{
+							if(FMaterialResource* MaterialResource = MaterialInterface->GetMaterialResource(ERHIFeatureLevel::SM6))
+							{
+								const FMaterialLayersFunctions* LayersFunctions_Original = MaterialResource->GetMaterialLayers();
+								TSharedPtr<FMaterialLayersFunctionsRuntimeGraphCache> OriginalTreeCache = LayersFunctions_Original ? LayersFunctions_Original->RuntimeGraphCache : nullptr;
+								if (OriginalTreeCache)
+								{
+									LayerTreeCache->NodePreviewMaterials = OriginalTreeCache->NodePreviewMaterials;									
+								}
+								MaterialResource->FeedbackMaterialLayersInstancedGraphFromCompilation(&LayersFunctions);
+							}
+							LayerTreeCache->NodePreviewMaterials[InNodeId] = (UMaterial*)EditedMaterial;
+						}						
+						MaterialInterface->SyncLayersRuntimeGraphCache(nullptr);
 					}
 				}
-			}
-			else if(Substrate::IsMaterialLayeringSupportEnabled())
-			{
-				if(UMaterialInterface* MaterialInterface = MaterialEditorInstance->GetMaterialInterface())
+				else if(UMaterialInterface* MaterialInterface = MaterialEditorInstance->GetMaterialInterface())
 				{
 					FMaterialLayersFunctions LayersFunctions;
 					if(MaterialInterface->GetMaterialLayers(LayersFunctions))
@@ -1364,6 +1375,18 @@ void SMaterialSubstrateTree::RecursiveCreateWidgets(FRecursiveCreateWidgetsConte
 					}
 				}
 			}
+			else if (AssetObject)
+			{
+				if (MaterialEditorInstance->StoredLayerPreviews[Payload.Layer] == nullptr)
+				{
+					MaterialEditorInstance->StoredLayerPreviews[Payload.Layer] = (NewObject<UMaterialInstanceConstant>(MaterialEditorInstance, NAME_None));
+				}
+				UMaterialInterface* EditedMaterial = Cast<UMaterialFunctionInterface>(AssetObject)->GetPreviewMaterial();
+				if (MaterialEditorInstance->StoredLayerPreviews[Payload.Layer] && MaterialEditorInstance->StoredLayerPreviews[Payload.Layer]->Parent != EditedMaterial)
+				{
+					MaterialEditorInstance->StoredLayerPreviews[Payload.Layer]->SetParentEditorOnly(EditedMaterial);
+				}
+			}
 		}
 
 		StackProperty->Children.Add(ChildProperty);
