Automatically invalidate material instances when you set their static switches by default.

- Adds `bUpdateMaterialInstance` to `SetMaterialInstanceStaticSwitchParameterValue` which defaults to true

[RN] minor | Rendering - Materials and Shaders

diff --git a/Engine/Source/Editor/MaterialEditor/Private/MaterialEditingLibrary.cpp b/Engine/Source/Editor/MaterialEditor/Private/MaterialEditingLibrary.cpp
--- a/Engine/Source/Editor/MaterialEditor/Private/MaterialEditingLibrary.cpp
+++ b/Engine/Source/Editor/MaterialEditor/Private/MaterialEditingLibrary.cpp
@@ -1246,7 +1246,7 @@ bool UMaterialEditingLibrary::GetMaterialInstanceStaticSwitchParameterValue(UMat
 	return bResult;
 }
 
-bool UMaterialEditingLibrary::SetMaterialInstanceStaticSwitchParameterValue(UMaterialInstanceConstant* Instance, FName ParameterName, bool Value, EMaterialParameterAssociation Association)
+bool UMaterialEditingLibrary::SetMaterialInstanceStaticSwitchParameterValue(UMaterialInstanceConstant* Instance, FName ParameterName, bool Value, EMaterialParameterAssociation Association, bool bUpdateMaterialInstance)
 {
 	bool bResult = false;
 	if (Instance)
@@ -1257,6 +1257,11 @@ bool UMaterialEditingLibrary::SetMaterialInstanceStaticSwitchParameterValue(UMat
 		// If there's ever a cleaner and more isolated way of populating MaterialLayersParameters, we should do that instead.
 		UMaterialEditorInstanceConstant* MaterialEditorInstance = NewObject<UMaterialEditorInstanceConstant>(GetTransientPackage(), NAME_None, RF_Transactional);
 		MaterialEditorInstance->SetSourceInstance(Instance);
+
+		if (bUpdateMaterialInstance)
+		{
+			UMaterialEditingLibrary::UpdateMaterialInstance(Instance);
+		}
 	}
 	return bResult;
 }
diff --git a/Engine/Source/Editor/MaterialEditor/Public/MaterialEditingLibrary.h b/Engine/Source/Editor/MaterialEditor/Public/MaterialEditingLibrary.h
--- a/Engine/Source/Editor/MaterialEditor/Public/MaterialEditingLibrary.h
+++ b/Engine/Source/Editor/MaterialEditor/Public/MaterialEditingLibrary.h
@@ -340,7 +340,7 @@ public:
 
 	/** Set the static switch parameter value for a Material Instance */
 	UFUNCTION(BlueprintCallable, Category = "MaterialEditing")
-	static bool SetMaterialInstanceStaticSwitchParameterValue(UMaterialInstanceConstant* Instance, FName ParameterName, bool Value, EMaterialParameterAssociation Association = EMaterialParameterAssociation::GlobalParameter);
+	static bool SetMaterialInstanceStaticSwitchParameterValue(UMaterialInstanceConstant* Instance, FName ParameterName, bool Value, EMaterialParameterAssociation Association = EMaterialParameterAssociation::GlobalParameter, bool bUpdateMaterialInstance = true);
 
 	/** Called after making modifications to a Material Instance to recompile shaders etc. */
 	UFUNCTION(BlueprintCallable, Category = "MaterialEditing")
