Fixed undo/redo with material function editing.

This issue was caused by the Transactional flag missing in the auxiliary materials created to edit the MF. This was in turn preventing the expression collections to get updated when an expressions was deleted, followed by undo/redo operations.

This fix should also make copy/paste more robust.

[RN] Fixed bug that prevented material function editing undo/redo/copy/paste operations to work reliably. This caused data invalidations, crashes and various unexpected behaviors, like sudden missing connections upon function reopening.
#tests That copy-paste combined with undo-redo works as intended in both Materials and Material Functions without crashes or incorrect behavior.

diff --git a/Engine/Source/Editor/MaterialEditor/Private/MaterialEditor.cpp b/Engine/Source/Editor/MaterialEditor/Private/MaterialEditor.cpp
--- a/Engine/Source/Editor/MaterialEditor/Private/MaterialEditor.cpp
+++ b/Engine/Source/Editor/MaterialEditor/Private/MaterialEditor.cpp
@@ -574,7 +574,7 @@ void FMaterialEditor::InitEditorForMaterialFunction(UMaterialFunction* InMateria
 	ExpressionPreviewMaterial = NULL;
 
 	// Create a temporary material to preview the material function
-	Material = NewObject<UMaterial>(); 
+	Material = NewObject<UMaterial>(GetTransientPackage(), NAME_None, RF_Transactional);
 	{
 		FArchiveUObject DummyArchive;
 		// Hack: serialize the new material with an archive that does nothing so that its material resources are created
@@ -6020,7 +6020,7 @@ void FMaterialEditor::SetPreviewExpression(UMaterialExpression* NewPreviewExpres
 		if (ExpressionPreviewMaterial == nullptr)
 		{
 			// Create the expression preview material if it hasnt already been created
-			ExpressionPreviewMaterial = NewObject<UPreviewMaterial>(GetTransientPackage(), NAME_None, RF_Public);
+			ExpressionPreviewMaterial = NewObject<UPreviewMaterial>(GetTransientPackage(), NAME_None, RF_Public | RF_Transactional);
 			ExpressionPreviewMaterial->bIsPreviewMaterial = true;
 			ExpressionPreviewMaterial->bEnableNewHLSLGenerator = Material->IsUsingNewHLSLGenerator();
 			if (Material->IsUIMaterial())
@@ -6351,6 +6351,8 @@ void FMaterialEditor::JumpToExpression(UMaterialExpression* Expression)
 
 void FMaterialEditor::ModifyMaterial()
 {
+	Material->MaterialGraph->Modify();
+	
 	Material->Modify();
 	Material->GetEditorOnlyData()->Modify();
 
@@ -6717,7 +6719,6 @@ void FMaterialEditor::PasteNodesHereFromBuffer(const FVector2D& Location, const
 	const FScopedTransaction Transaction( NSLOCTEXT("UnrealEd", "MaterialEditorPaste", "Material Editor: Paste") );
 	Material->MaterialGraph->Modify();
 	ModifyMaterial();
-	Material->GetEditorOnlyData()->Modify();
 
 	UMaterialGraph* ExpressionGraph = Graph ? ToRawPtr(CastChecked<UMaterialGraph>(const_cast<UEdGraph*>(Graph))) : ToRawPtr(Material->MaterialGraph);
 	ExpressionGraph->Modify();
@@ -7451,18 +7452,6 @@ void FMaterialEditor::PostUndo(bool bSuccess)
 {
 	if (bSuccess)
 	{	
-		// Due to the order of invocation of callbacks when undo/redo is performed, it may occur that some expression pointers in
-		// the material expression collection are null. Sanitize the list here for safety.
-		if (MaterialFunction)
-		{
-			MaterialFunction->GetExpressionCollection().RemoveExpression(nullptr);
-			Material->AssignExpressionCollection(MaterialFunction->GetExpressionCollection());
-		}
-		else
-		{
-			Material->GetExpressionCollection().RemoveExpression(nullptr);
-		}
-
 		if (TSharedPtr<SGraphEditor> FocusedGraphEd = FocusedGraphEdPtr.Pin())
 		{
 			FocusedGraphEd->ClearSelectionSet();
