When a new material expression is created be sure to grid align the coordinates of the node.

- I'm leveraging the existing function `SnapToGrid(SNodePanel::GetSnapGridSize())` to perform the alignment.  This means I must have the `GraphNode` on the MaterialExpression initialized.
- For this to happen I had to call this function _after_ `ExpressionGraph->AddExpression(NewExpression, bAutoSelect);` or `ExpressionGraph->AddComment(NewComment, true);` were called since they create the graph node.
- I also needed to keep the MaterialExpressionEditorX and Graph Node state in sync.
- There are some other common functions that are called when "new" nodes are created where I could do alignment, but these were also called when instantiating _existing_ node graphs, and I didn't want to modify existing layouts.

[RN] minor | Rendering - Materials and Shaders

diff --git a/Engine/Source/Editor/MaterialEditor/Private/MaterialEditor.cpp b/Engine/Source/Editor/MaterialEditor/Private/MaterialEditor.cpp
--- a/Engine/Source/Editor/MaterialEditor/Private/MaterialEditor.cpp
+++ b/Engine/Source/Editor/MaterialEditor/Private/MaterialEditor.cpp
@@ -6195,6 +6195,10 @@ UMaterialExpression* FMaterialEditor::CreateNewMaterialExpression(UClass* NewExp
 		{
 			ExpressionGraph->AddExpression(NewExpression, bAutoSelect);
 
+			NewExpression->GraphNode->SnapToGrid(SNodePanel::GetSnapGridSize());
+			NewExpression->MaterialExpressionEditorX = NewExpression->GraphNode->NodePosX;
+			NewExpression->MaterialExpressionEditorY = NewExpression->GraphNode->NodePosY;
+
 			// Select the new node.
 			if ( bAutoSelect && FocusedGraphEd )
 			{
@@ -6324,6 +6328,10 @@ UMaterialExpressionComment* FMaterialEditor::CreateNewMaterialExpressionComment(
 
 		ExpressionGraph->AddComment(NewComment, true);
 
+		NewComment->GraphNode->SnapToGrid(SNodePanel::GetSnapGridSize());
+		NewComment->MaterialExpressionEditorX = NewComment->GraphNode->NodePosX;
+		NewComment->MaterialExpressionEditorY = NewComment->GraphNode->NodePosY;
+
 		// Select the new comment.
 		if (FocusedGraphEd)
 		{
