[Material Stats] Use more distinctive names for material statistic panes to better distinguish between new and old translator output.

#rnx

diff --git a/Engine/Source/Editor/MaterialEditor/Private/MaterialStats.cpp b/Engine/Source/Editor/MaterialEditor/Private/MaterialStats.cpp
--- a/Engine/Source/Editor/MaterialEditor/Private/MaterialStats.cpp
+++ b/Engine/Source/Editor/MaterialEditor/Private/MaterialStats.cpp
@@ -1093,7 +1093,7 @@ TSharedRef<class SDockTab> FMaterialStats::SpawnTab_HLSLCode(const class FSpawnT
 					.Orientation(EOrientation::Orient_Horizontal)
 					+SSplitter::Slot()
 					[
-						BuildShaderCodeWidget(LOCTEXT("CurrentTranslatorLabel", "HLSL"),
+						BuildShaderCodeWidget(LOCTEXT("CurrentTranslatorLabel", "New Translator HLSL Code"),
 							[Code = &HLSLCode]() { return FText::FromString(*Code); },
 							[Code = &HLSLCode]() {
 								FPlatformApplicationMisc::ClipboardCopy(**Code);
@@ -1104,7 +1104,7 @@ TSharedRef<class SDockTab> FMaterialStats::SpawnTab_HLSLCode(const class FSpawnT
 					]
 					+SSplitter::Slot()
 					[
-						BuildShaderCodeWidget(LOCTEXT("PreviousTranslatorLabel", "HLSL"),
+						BuildShaderCodeWidget(LOCTEXT("PreviousTranslatorLabel", "HLSL Code"),
 							[Code = &PreviousTranslatorHLSLCode]() { return FText::FromString(*Code); },
 							[Code = &PreviousTranslatorHLSLCode]() {
 								FPlatformApplicationMisc::ClipboardCopy(**Code);
