Fix bug where thumbnails weren't updating properly in the Material Layer Tool

- All these thumbnails were marked as real-time and would update most of the time.
- If you hovered over a thumbnail, and the mouse left the thumbnail, realtime would get toggled off.  When this happened we would no longer receive thumbnail updates, and only get an update when hovering over the thumbnail. (see `OnMouseLeave` in `AssetThumbnail.cpp` for where realtime is turned off)
- `bAllowRealTimeOnHovered` on the thumbnail controls whether we turn off realtime on mouse leave.
- To fix this issue I'm turning this off on thumbnails in the layer tool.
- Add debug names to the transient material instances used for thumbnail previews for debugging purposes ("Layer_1_MyMaterialLayer")

TODO: The thumbnails should probably just render on demand...not sure how many realtime thumbnails we can support

#rnx

diff --git a/Engine/Source/Editor/MaterialEditor/Private/SMaterialSubstrateTree.cpp b/Engine/Source/Editor/MaterialEditor/Private/SMaterialSubstrateTree.cpp
--- a/Engine/Source/Editor/MaterialEditor/Private/SMaterialSubstrateTree.cpp
+++ b/Engine/Source/Editor/MaterialEditor/Private/SMaterialSubstrateTree.cpp
@@ -1321,7 +1321,7 @@ void SMaterialSubstrateTree::RecursiveCreateWidgets(FRecursiveCreateWidgetsConte
 						MaterialInterface->SyncLayersRuntimeGraphCache(&LayersFunctions);
 						if (MaterialEditorInstance->StoredLayerPreviews[Payload.Layer] == nullptr)
 						{
-							MaterialEditorInstance->StoredLayerPreviews[Payload.Layer] = (NewObject<UMaterialInstanceConstant>(MaterialEditorInstance, NAME_None));
+							MaterialEditorInstance->StoredLayerPreviews[Payload.Layer] = NewObject<UMaterialInstanceConstant>(MaterialEditorInstance, FName(*FString::Printf(TEXT("Layer_%d_%s"), Payload.Layer, *AssetObject->GetName())));
 						}
 						UMaterialInterface* EditedMaterial = Cast<UMaterialFunctionInterface>(AssetObject)->GetPreviewMaterial();
 						if (MaterialEditorInstance->StoredLayerPreviews[Payload.Layer] && MaterialEditorInstance->StoredLayerPreviews[Payload.Layer]->Parent != EditedMaterial)
@@ -1360,7 +1360,7 @@ void SMaterialSubstrateTree::RecursiveCreateWidgets(FRecursiveCreateWidgetsConte
 							{
 								if (MaterialEditorInstance->StoredLayerPreviews[Payload.Layer] == nullptr)
 								{
-									MaterialEditorInstance->StoredLayerPreviews[Payload.Layer] = (NewObject<UMaterialInstanceConstant>(MaterialEditorInstance, NAME_None));
+									MaterialEditorInstance->StoredLayerPreviews[Payload.Layer] = NewObject<UMaterialInstanceConstant>(MaterialEditorInstance, FName(*FString::Printf(TEXT("Layer_%d_%s"), Payload.Layer, *Call->MaterialFunction->GetName())));
 								}
 
 								UMaterial* PreviewMaterial = LayersFunctions.GetRuntimeNodePreviewMaterial(InNodeId);
@@ -1379,7 +1379,7 @@ void SMaterialSubstrateTree::RecursiveCreateWidgets(FRecursiveCreateWidgetsConte
 			{
 				if (MaterialEditorInstance->StoredLayerPreviews[Payload.Layer] == nullptr)
 				{
-					MaterialEditorInstance->StoredLayerPreviews[Payload.Layer] = (NewObject<UMaterialInstanceConstant>(MaterialEditorInstance, NAME_None));
+					MaterialEditorInstance->StoredLayerPreviews[Payload.Layer] = NewObject<UMaterialInstanceConstant>(MaterialEditorInstance, FName(*FString::Printf(TEXT("Layer_%d_%s"), Payload.Layer, *AssetObject->GetName())));
 				}
 				UMaterialInterface* EditedMaterial = Cast<UMaterialFunctionInterface>(AssetObject)->GetPreviewMaterial();
 				if (MaterialEditorInstance->StoredLayerPreviews[Payload.Layer] && MaterialEditorInstance->StoredLayerPreviews[Payload.Layer]->Parent != EditedMaterial)
@@ -1411,7 +1411,7 @@ void SMaterialSubstrateTree::RecursiveCreateWidgets(FRecursiveCreateWidgetsConte
 			{
 				if (MaterialEditorInstance->StoredBlendPreviews[Payload.Blend] == nullptr)
 				{
-					MaterialEditorInstance->StoredBlendPreviews[Payload.Blend] = (NewObject<UMaterialInstanceConstant>(MaterialEditorInstance, NAME_None));
+					MaterialEditorInstance->StoredBlendPreviews[Payload.Blend] = NewObject<UMaterialInstanceConstant>(MaterialEditorInstance, FName(*FString::Printf(TEXT("Blend_%d_%s"), Payload.Blend, *AssetObject->GetName())));
 				}
 				UMaterialInterface* EditedMaterial = Cast<UMaterialFunctionInterface>(AssetObject)->GetPreviewMaterial();
 				if (MaterialEditorInstance->StoredBlendPreviews[Payload.Blend] && MaterialEditorInstance->StoredBlendPreviews[Payload.Blend]->Parent != EditedMaterial)
@@ -1613,7 +1613,12 @@ TSharedRef<SWidget> SMaterialSubstrateTree::CreateThumbnailWidget(EMaterialParam
 		ThumbnailObject = MaterialEditorInstance->StoredBlendPreviews[InIndex];
 	}
 	const TSharedPtr<FAssetThumbnail> AssetThumbnail = MakeShareable(new FAssetThumbnail(ThumbnailObject, InThumbnailSize, InThumbnailSize, GetTreeThumbnailPool()));
-	TSharedRef<SWidget> ThumbnailWidget = AssetThumbnail->MakeThumbnailWidget();
+
+	FAssetThumbnailConfig ThumbnailConfig;
+	// disable realtime on hovered since these will always be realtime.  MouseLeave events turn _off_ realtime updates which isn't the behavior we want.
+	ThumbnailConfig.bAllowRealTimeOnHovered = false; 
+
+	TSharedRef<SWidget> ThumbnailWidget = AssetThumbnail->MakeThumbnailWidget(ThumbnailConfig);
 	AssetThumbnail->SetRealTime(true);
 	ThumbnailWidget->SetOnMouseDoubleClick(FPointerEventHandler::CreateSP(this, &SMaterialSubstrateTree::OnThumbnailDoubleClick, InAssociation, InIndex));
 	return ThumbnailWidget;
