Substrate - UI fix up.

#rnx
#tests PC editor

diff --git a/Engine/Source/Editor/MaterialEditor/Private/SMaterialEditorSubstrateWidget.cpp b/Engine/Source/Editor/MaterialEditor/Private/SMaterialEditorSubstrateWidget.cpp
--- a/Engine/Source/Editor/MaterialEditor/Private/SMaterialEditorSubstrateWidget.cpp
+++ b/Engine/Source/Editor/MaterialEditor/Private/SMaterialEditorSubstrateWidget.cpp
@@ -429,8 +429,8 @@ static FString SubstrateMaterialTypeToString(uint8 InMaterialType)
 	{
 		case SUBSTRATE_MATERIAL_TYPE_SIMPLE:  			return FString::Printf(TEXT("SIMPLE  (Diffuse, albedo, roughness)\r\n"));
 		case SUBSTRATE_MATERIAL_TYPE_SINGLE:  			return FString::Printf(TEXT("SINGLE  (BSDF all features except anisotropy)\r\n"));
-		case SUBSTRATE_MATERIAL_TYPE_COMPLEX:  			return FString::Printf(TEXT("COMPLEX (Anisotropy, multi-slabs)\r\n"));
-		case SUBSTRATE_MATERIAL_TYPE_COMPLEX_SPECIAL:  	return FString::Printf(TEXT("COMPLEX SPECIAL (Glints, specular profile)\r\n"));
+		case SUBSTRATE_MATERIAL_TYPE_COMPLEX:  			return FString::Printf(TEXT("COMPLEX (Anisotropy, specular profile, multi-slabs)\r\n"));
+		case SUBSTRATE_MATERIAL_TYPE_COMPLEX_SPECIAL:  	return FString::Printf(TEXT("COMPLEX SPECIAL (Glints)\r\n"));
 		default: 										return FString::Printf(TEXT("UNKOWN => ERROR!\r\n"));
 	}
 }
