Add virtual function IsAllowedIn to UMaterialExpression, which allows expressions to provide additional restrictions for what material or material functions they are allowed to be created in.

diff --git a/Engine/Source/Editor/MaterialEditor/Private/MaterialEditor.cpp b/Engine/Source/Editor/MaterialEditor/Private/MaterialEditor.cpp
--- a/Engine/Source/Editor/MaterialEditor/Private/MaterialEditor.cpp
+++ b/Engine/Source/Editor/MaterialEditor/Private/MaterialEditor.cpp
@@ -6158,12 +6158,12 @@ void FMaterialEditor::CloseDocumentTab(const UObject* DocumentID)
 
 UMaterialExpression* FMaterialEditor::CreateNewMaterialExpression(UClass* NewExpressionClass, const FVector2D& NodePos, bool bAutoSelect, bool bAutoAssignResource, const UEdGraph* Graph)
 {
-	check( NewExpressionClass->IsChildOf(UMaterialExpression::StaticClass()) );
+	UMaterialExpression* DefaultExpression = CastChecked<UMaterialExpression>(NewExpressionClass->GetDefaultObject());
 	TSharedPtr<SGraphEditor> FocusedGraphEd = FocusedGraphEdPtr.Pin();
 	UMaterialGraph* ExpressionGraph = Graph ? ToRawPtr(CastChecked<UMaterialGraph>(const_cast<UEdGraph*>(Graph))) : ToRawPtr(Material->MaterialGraph); 
 	ExpressionGraph->Modify();
 
-	if (!IsAllowedExpressionType(NewExpressionClass, MaterialFunction != NULL))
+	if (!DefaultExpression->IsAllowedIn(OriginalMaterialObject))
 	{
 		UE_LOG(LogMaterialEditor, Warning, TEXT("Trying to create a disallowed material expression of type %s."), *NewExpressionClass->GetName());
 		return NULL;
diff --git a/Engine/Source/Editor/MaterialEditor/Private/MaterialEditorUtilities.cpp b/Engine/Source/Editor/MaterialEditor/Private/MaterialEditorUtilities.cpp
--- a/Engine/Source/Editor/MaterialEditor/Private/MaterialEditorUtilities.cpp
+++ b/Engine/Source/Editor/MaterialEditor/Private/MaterialEditorUtilities.cpp
@@ -29,6 +29,7 @@
 #include "Materials/MaterialExpressionCustomOutput.h"
 #include "Materials/MaterialExpressionMaterialAttributeLayers.h"
 #include "Materials/MaterialExpressionRerouteBase.h"
+#include "Materials/MaterialFunction.h"
 #include "Materials/MaterialInstanceConstant.h"
 
 #include "DebugViewModeHelpers.h"
@@ -184,30 +185,36 @@ int32 FMaterialEditorUtilities::GetNumberOfSelectedNodes(const class UEdGraph* G
 }
 
 void FMaterialEditorUtilities::GetMaterialExpressionActions(FGraphActionMenuBuilder& ActionMenuBuilder, bool bMaterialFunction)
+{
+	UObject* MaterialOrFunction = bMaterialFunction ? UMaterialFunction::StaticClass()->GetDefaultObject() : UMaterial::StaticClass()->GetDefaultObject();
+	GetMaterialExpressionActions(ActionMenuBuilder, MaterialOrFunction);
+}
+
+void FMaterialEditorUtilities::GetMaterialExpressionActions(FGraphActionMenuBuilder& ActionMenuBuilder, const UObject* MaterialOrFunction)
 {
 	bool bUseUnsortedMenus = false;
 	MaterialExpressionClasses* ExpressionClasses = MaterialExpressionClasses::Get();
 
 	if (bUseUnsortedMenus)
 	{
-		AddMaterialExpressionCategory(ActionMenuBuilder, FText::GetEmpty(), &ExpressionClasses->AllExpressionClasses, bMaterialFunction);
+		AddMaterialExpressionCategory(ActionMenuBuilder, FText::GetEmpty(), &ExpressionClasses->AllExpressionClasses, MaterialOrFunction);
 	}
 	else
 	{
 		// Add Favourite expressions as a category
 		const FText FavouritesCategory = LOCTEXT("FavoritesMenu", "Favorites");
-		AddMaterialExpressionCategory(ActionMenuBuilder, FavouritesCategory, &ExpressionClasses->FavoriteExpressionClasses, bMaterialFunction);
+		AddMaterialExpressionCategory(ActionMenuBuilder, FavouritesCategory, &ExpressionClasses->FavoriteExpressionClasses, MaterialOrFunction);
 
 		// Add each category to the menu
 		for (int32 CategoryIndex = 0; CategoryIndex < ExpressionClasses->CategorizedExpressionClasses.Num(); ++CategoryIndex)
 		{
 			FCategorizedMaterialExpressionNode* CategoryNode = &(ExpressionClasses->CategorizedExpressionClasses[CategoryIndex]);
-			AddMaterialExpressionCategory(ActionMenuBuilder, CategoryNode->CategoryName, &CategoryNode->MaterialExpressions, bMaterialFunction);
+			AddMaterialExpressionCategory(ActionMenuBuilder, CategoryNode->CategoryName, &CategoryNode->MaterialExpressions, MaterialOrFunction);
 		}
 
 		if (ExpressionClasses->UnassignedExpressionClasses.Num() > 0)
 		{
-			AddMaterialExpressionCategory(ActionMenuBuilder, FText::GetEmpty(), &ExpressionClasses->UnassignedExpressionClasses, bMaterialFunction);
+			AddMaterialExpressionCategory(ActionMenuBuilder, FText::GetEmpty(), &ExpressionClasses->UnassignedExpressionClasses, MaterialOrFunction);
 		}
 	}
 }
@@ -670,6 +677,12 @@ void FMaterialEditorUtilities::BringFocusAttentionOnObject(const UObject* Object
 }
 
 void FMaterialEditorUtilities::AddMaterialExpressionCategory(FGraphActionMenuBuilder& ActionMenuBuilder, FText CategoryName, TArray<struct FMaterialExpression>* MaterialExpressions, bool bMaterialFunction)
+{
+	UObject* MaterialOrFunction = bMaterialFunction ? UMaterialFunction::StaticClass()->GetDefaultObject() : UMaterial::StaticClass()->GetDefaultObject();
+	AddMaterialExpressionCategory(ActionMenuBuilder, CategoryName, MaterialExpressions, MaterialOrFunction);
+}
+
+void FMaterialEditorUtilities::AddMaterialExpressionCategory(FGraphActionMenuBuilder& ActionMenuBuilder, FText CategoryName, TArray<struct FMaterialExpression>* MaterialExpressions, const UObject* MaterialOrFunction)
 {
 	// Get type of dragged pin
 	uint32 FromPinType = 0;
@@ -678,10 +691,13 @@ void FMaterialEditorUtilities::AddMaterialExpressionCategory(FGraphActionMenuBui
 		FromPinType = UMaterialGraphSchema::GetMaterialValueType(ActionMenuBuilder.FromPin);
 	}
 
+	bool bMaterialFunction = Cast<UMaterialFunction>(MaterialOrFunction) != nullptr;
+
 	for (int32 Index = 0; Index < MaterialExpressions->Num(); ++Index)
 	{
 		const FMaterialExpression& MaterialExpression = (*MaterialExpressions)[Index];
-		if (IsAllowedExpressionType(MaterialExpression.MaterialClass, bMaterialFunction) && MaterialExpression.MaterialClass != UMaterialExpressionComposite::StaticClass())
+		UMaterialExpression* DefaultExpression = CastChecked<UMaterialExpression>(MaterialExpression.MaterialClass->GetDefaultObject());
+		if (DefaultExpression->IsAllowedIn(MaterialOrFunction) && MaterialExpression.MaterialClass != UMaterialExpressionComposite::StaticClass())
 		{
 			if (!ActionMenuBuilder.FromPin || HasCompatibleConnection(MaterialExpression.MaterialClass, FromPinType, ActionMenuBuilder.FromPin->Direction, bMaterialFunction))
 			{
diff --git a/Engine/Source/Editor/MaterialEditor/Private/SMaterialPalette.cpp b/Engine/Source/Editor/MaterialEditor/Private/SMaterialPalette.cpp
--- a/Engine/Source/Editor/MaterialEditor/Private/SMaterialPalette.cpp
+++ b/Engine/Source/Editor/MaterialEditor/Private/SMaterialPalette.cpp
@@ -192,7 +192,7 @@ void SMaterialPalette::CollectAllActions(FGraphActionListBuilderBase& OutAllActi
 	FGraphActionMenuBuilder ActionMenuBuilder;
 
 	// Determine all possible actions
-	Schema->GetPaletteActions(ActionMenuBuilder, GetFilterCategoryName(), MaterialEditorPtr.Pin()->MaterialFunction != NULL);
+	Schema->GetPaletteActions(ActionMenuBuilder, GetFilterCategoryName(), MaterialEditorPtr.Pin()->OriginalMaterialObject);
 
 	//@TODO: Avoid this copy
 	OutAllActions.Append(ActionMenuBuilder);
diff --git a/Engine/Source/Editor/MaterialEditor/Public/MaterialEditorUtilities.h b/Engine/Source/Editor/MaterialEditor/Public/MaterialEditorUtilities.h
--- a/Engine/Source/Editor/MaterialEditor/Public/MaterialEditorUtilities.h
+++ b/Engine/Source/Editor/MaterialEditor/Public/MaterialEditorUtilities.h
@@ -153,8 +153,17 @@ public:
 	 * @param [in,out]	ActionMenuBuilder	The output menu builder.
 	 * @param bMaterialFunction				Whether we're dealing with a Material Function
 	 */
+	UE_DEPRECATED(5.8, "Should call version that provides a UMaterial or UMaterialFunction")
 	static void GetMaterialExpressionActions(FGraphActionMenuBuilder& ActionMenuBuilder, bool bMaterialFunction);
 
+	/**
+	 * Get all actions for placing material expressions
+	 *
+	 * @param [in,out]	ActionMenuBuilder	The output menu builder.
+	 * @param MaterialOrFunction			A UMaterial or UMaterialFunction object
+	 */
+	static void GetMaterialExpressionActions(FGraphActionMenuBuilder& ActionMenuBuilder, const UObject* MaterialOrFunction);
+
 	/**
 	 * Check whether an expression is in the favourites list
 	 *
@@ -268,8 +277,21 @@ private:
 	 * @param	MaterialExpressions	List of Material Expressions in the category.
 	 * @param	bMaterialFunction	Whether we are building for a material function.
 	 */
+	UE_DEPRECATED(5.8, "Should call version that provides a UMaterial or UMaterialFunction")
 	static void AddMaterialExpressionCategory(FGraphActionMenuBuilder& ActionMenuBuilder, FText CategoryName, TArray<struct FMaterialExpression>* MaterialExpressions, bool bMaterialFunction);
 
+	/**
+	* Adds a category of Material Expressions to an action builder
+	*
+	* @param	ActionMenuBuilder	The builder to add to.
+	* @param	CategoryName		The name of the category.
+	* @param	MaterialExpressions	List of Material Expressions in the category.
+	* @param	MaterialOrFunction	UMaterial or UMaterialFunction object
+	*/
+	static void AddMaterialExpressionCategory(FGraphActionMenuBuilder& ActionMenuBuilder, FText CategoryName, TArray<struct FMaterialExpression>* MaterialExpressions, const UObject* MaterialOrFunction);
+
+public:
+
 	/**
 	 * Checks whether a Material Expression class has any connections that are compatible with a type/direction
 	 *
diff --git a/Engine/Source/Editor/UnrealEd/Classes/MaterialGraph/MaterialGraph.h b/Engine/Source/Editor/UnrealEd/Classes/MaterialGraph/MaterialGraph.h
--- a/Engine/Source/Editor/UnrealEd/Classes/MaterialGraph/MaterialGraph.h
+++ b/Engine/Source/Editor/UnrealEd/Classes/MaterialGraph/MaterialGraph.h
@@ -257,6 +257,8 @@ public:
 
 	UNREALED_API void UpdatePinTypes();
 
+	inline UObject* GetMaterialOrFunction() const { return MaterialFunction ? (UObject*)MaterialFunction : Material; }
+
 private:
 	/**
 	 * Remove all Nodes from the graph
diff --git a/Engine/Source/Editor/UnrealEd/Classes/MaterialGraph/MaterialGraphSchema.h b/Engine/Source/Editor/UnrealEd/Classes/MaterialGraph/MaterialGraphSchema.h
--- a/Engine/Source/Editor/UnrealEd/Classes/MaterialGraph/MaterialGraphSchema.h
+++ b/Engine/Source/Editor/UnrealEd/Classes/MaterialGraph/MaterialGraphSchema.h
@@ -246,7 +246,9 @@ class UMaterialGraphSchema : public UEdGraphSchema
 	 */
 	void OnConnectToMaterial(UEdGraphPin* InGraphPin, int32 ConnIndex);
 
+	UE_DEPRECATED(5.8, "Should call version that provides a UMaterial or UMaterialFunction")
 	UNREALED_API void GetPaletteActions(FGraphActionMenuBuilder& ActionMenuBuilder, const FString& CategoryName, bool bMaterialFunction) const;
+	UNREALED_API void GetPaletteActions(FGraphActionMenuBuilder& ActionMenuBuilder, const FString& CategoryName, const UObject* MaterialOrFunction) const;
 
 	/** Check whether connecting these pins would cause a loop */
 	bool ConnectionCausesLoop(const UEdGraphPin* InputPin, const UEdGraphPin* OutputPin) const;
diff --git a/Engine/Source/Editor/UnrealEd/Private/Commandlets/DumpMaterialExpressionsCommandlet.cpp b/Engine/Source/Editor/UnrealEd/Private/Commandlets/DumpMaterialExpressionsCommandlet.cpp
--- a/Engine/Source/Editor/UnrealEd/Private/Commandlets/DumpMaterialExpressionsCommandlet.cpp
+++ b/Engine/Source/Editor/UnrealEd/Private/Commandlets/DumpMaterialExpressionsCommandlet.cpp
@@ -14,6 +14,7 @@
 #include "Internationalization/Text.h"
 #include "Logging/LogCategory.h"
 #include "Logging/LogMacros.h"
+#include "Materials/Material.h"
 #include "Materials/MaterialExpression.h"
 #include "Materials/MaterialExpressionComment.h"
 #include "Materials/MaterialExpressionComposite.h"
@@ -108,16 +109,12 @@ void WriteOutMaterialExpressions(FDiagnosticTableWriterCSV& CsvFile)
 				// See class exclusions in:
 				//    MaterialExpressionClasses::InitMaterialExpressionClasses()
 				//    FMaterialEditorUtilities::AddMaterialExpressionCategory() 
-				//    MaterialExpressions.cpp IsAllowedExpressionType()
-				bool bShowInCreateMenu = !(bClassDeprecated
-					|| Class == UMaterialExpressionComment::StaticClass()
-					|| Class == UMaterialExpressionMaterialLayerOutput::StaticClass()
-					|| Class == UMaterialExpressionParameter::StaticClass()
-					|| Class == UMaterialExpressionNamedRerouteUsage::StaticClass()
-					|| Class == UMaterialExpressionPinBase::StaticClass()
-					|| Class == UMaterialExpressionFunctionInput::StaticClass()
-					|| Class == UMaterialExpressionFunctionOutput::StaticClass()
-					|| Class == UMaterialExpressionComposite::StaticClass());
+				//    and IsAllowedIn in UMaterialExpression::IsAllowedIn and overridden methods
+				bool bShowInCreateMenu = !bClassDeprecated
+						&& DefaultExpression->IsAllowedIn(UMaterial::StaticClass()->GetDefaultObject())
+						&& Class != UMaterialExpressionMaterialLayerOutput::StaticClass()
+						&& Class != UMaterialExpressionNamedRerouteUsage::StaticClass()
+						&& Class != UMaterialExpressionComposite::StaticClass();
 
 				const bool bCollapseCategories = Class->HasAnyClassFlags(CLASS_CollapseCategories);
 				const bool bHideCategories = Class->HasMetaData(TEXT("HideCategories"));
diff --git a/Engine/Source/Editor/UnrealEd/Private/MaterialGraphNode.cpp b/Engine/Source/Editor/UnrealEd/Private/MaterialGraphNode.cpp
--- a/Engine/Source/Editor/UnrealEd/Private/MaterialGraphNode.cpp
+++ b/Engine/Source/Editor/UnrealEd/Private/MaterialGraphNode.cpp
@@ -160,7 +160,7 @@ bool UMaterialGraphNode::CanPasteHere(const UEdGraph* TargetGraph) const
 				bIsValidFunctionExpression = false;
 			}
 
-			if (bIsValidFunctionExpression && MaterialExpression && IsAllowedExpressionType(MaterialExpression->GetClass(), MaterialGraph->MaterialFunction != NULL))
+			if (bIsValidFunctionExpression && MaterialExpression && MaterialExpression->IsAllowedIn(MaterialGraph->GetMaterialOrFunction()))
 			{
 				return true;
 			}
diff --git a/Engine/Source/Editor/UnrealEd/Private/MaterialGraphSchema.cpp b/Engine/Source/Editor/UnrealEd/Private/MaterialGraphSchema.cpp
--- a/Engine/Source/Editor/UnrealEd/Private/MaterialGraphSchema.cpp
+++ b/Engine/Source/Editor/UnrealEd/Private/MaterialGraphSchema.cpp
@@ -384,10 +384,16 @@ void UMaterialGraphSchema::OnConnectToMaterial(class UEdGraphPin* InGraphPin, in
 }
 
 void UMaterialGraphSchema::GetPaletteActions(FGraphActionMenuBuilder& ActionMenuBuilder, const FString& CategoryName, bool bMaterialFunction) const
+{
+	UObject* MaterialOrFunction = bMaterialFunction ? UMaterialFunction::StaticClass()->GetDefaultObject() : UMaterial::StaticClass()->GetDefaultObject();
+	return GetPaletteActions(ActionMenuBuilder, CategoryName, MaterialOrFunction);
+}
+
+void UMaterialGraphSchema::GetPaletteActions(FGraphActionMenuBuilder& ActionMenuBuilder, const FString& CategoryName, const UObject* MaterialOrFunction) const
 {
 	if (CategoryName != TEXT("Functions"))
 	{
-		FMaterialEditorUtilities::GetMaterialExpressionActions(ActionMenuBuilder, bMaterialFunction);
+		FMaterialEditorUtilities::GetMaterialExpressionActions(ActionMenuBuilder, MaterialOrFunction);
 		GetCommentAction(ActionMenuBuilder);
 	}
 	if (CategoryName != TEXT("Expressions"))
@@ -483,7 +489,7 @@ void UMaterialGraphSchema::GetGraphContextActions(FGraphContextMenuBuilder& Cont
 	Super::GetGraphContextActions(ContextMenuBuilder);
 
 	// Get the Context Actions from Material Editor Module
-	FMaterialEditorUtilities::GetMaterialExpressionActions(ContextMenuBuilder, MaterialGraph->MaterialFunction != NULL);
+	FMaterialEditorUtilities::GetMaterialExpressionActions(ContextMenuBuilder, MaterialGraph->GetMaterialOrFunction());
 
 	// Get the Material Functions as well
 	GetMaterialFunctionActions(ContextMenuBuilder);
diff --git a/Engine/Source/Runtime/Engine/Private/Materials/MaterialExpressions.cpp b/Engine/Source/Runtime/Engine/Private/Materials/MaterialExpressions.cpp
--- a/Engine/Source/Runtime/Engine/Private/Materials/MaterialExpressions.cpp
+++ b/Engine/Source/Runtime/Engine/Private/Materials/MaterialExpressions.cpp
@@ -337,32 +337,11 @@ static const TCHAR* CPD_UI_ErrorMessage = TEXT("Custom Primitive Data can't be u
 #endif // #if WITH_EDITOR
 
 /** Returns whether the given expression class is allowed. */
-bool IsAllowedExpressionType(const UClass* const Class, const bool bMaterialFunction)
+bool IsAllowedExpressionType(const UClass* Class, const bool bMaterialFunction)
 {
-	// Custom HLSL expressions are not allowed for client generated materials in certain UE editor configuration
-	if (!IsExpressionClassPermitted(Class))
-	{
-		return false;
-	}
-
-	static const auto AllowTextureArrayAssetCreationVar = IConsoleManager::Get().FindTConsoleVariableDataInt(TEXT("r.AllowTexture2DArrayCreation"));
-
-	// Exclude comments from the expression list, as well as the base parameter expression, as it should not be used directly
-	const bool bSharedAllowed = Class != UMaterialExpressionComment::StaticClass() 
-		&& Class != UMaterialExpressionPinBase::StaticClass()
-		&& Class != UMaterialExpressionParameter::StaticClass()
-		&& (Class != UMaterialExpressionTextureSampleParameter2DArray::StaticClass() || AllowTextureArrayAssetCreationVar->GetValueOnGameThread() != 0);
-
-	if (bMaterialFunction)
-	{
-		return bSharedAllowed;
-	}
-	else
-	{
-		return bSharedAllowed
-			&& Class != UMaterialExpressionFunctionInput::StaticClass()
-			&& Class != UMaterialExpressionFunctionOutput::StaticClass();
-	}
+	UMaterialExpression* DefaultExpression = CastChecked<UMaterialExpression>(Class->GetDefaultObject());
+	UObject* MaterialOrFunction = bMaterialFunction ? UMaterialFunction::StaticClass()->GetDefaultObject() : UMaterial::StaticClass()->GetDefaultObject() ;
+	return DefaultExpression->IsAllowedIn(MaterialOrFunction);
 }
 
 /** Parses a string into multiple lines, for use with tooltips. */
@@ -825,7 +804,7 @@ void UMaterialExpression::CopyMaterialExpressions(const TArray<UMaterialExpressi
 			bIsValidFunctionExpression = false;
 		}
 
-		if (bIsValidFunctionExpression && IsAllowedExpressionType(SrcExpression->GetClass(), EditFunction != nullptr))
+		if (bIsValidFunctionExpression && SrcExpression->IsAllowedIn(ExpressionOuter))
 		{
 			UMaterialExpression* NewExpression = Cast<UMaterialExpression>(StaticDuplicateObject( SrcExpression, ExpressionOuter, NAME_None, RF_Transactional ));
 			NewExpression->Material = Material;
@@ -2345,6 +2324,12 @@ FSubstrateOperator* UMaterialExpression::SubstrateGenerateMaterialTopologyTree(c
 
 #endif // WITH_EDITOR
 
+bool UMaterialExpression::IsAllowedIn(const UObject* MaterialOrFunction) const
+{
+	// Custom HLSL expressions are not allowed for client generated materials in certain UE editor configuration
+	return IsExpressionClassPermitted(GetClass());
+}
+
 #if WITH_EDITOR
 void UMaterialExpressionTextureBase::PostEditChangeProperty(FPropertyChangedEvent& PropertyChangedEvent)
 {
@@ -4422,6 +4407,17 @@ bool UMaterialExpressionTextureSampleParameter2DArray::TextureIsValid(UTexture*
 }
 #endif // WITH_EDITOR
 
+bool UMaterialExpressionTextureSampleParameter2DArray::IsAllowedIn(const UObject* MaterialOrFunction) const
+{
+	static const auto AllowTextureArrayAssetCreationVar = IConsoleManager::Get().FindTConsoleVariableDataInt(TEXT("r.AllowTexture2DArrayCreation"));
+	if (AllowTextureArrayAssetCreationVar->GetValueOnGameThread() == 0)
+	{
+		return false;
+	}
+
+	return Super::IsAllowedIn(MaterialOrFunction);
+}
+
 const TCHAR* UMaterialExpressionTextureSampleParameter2DArray::GetRequirements()
 {
 	return TEXT("Requires Texture2DArray");
@@ -18134,6 +18130,10 @@ uint32 UMaterialExpressionFunctionInput::GetOutputType(int32 OutputIndex)
 }
 #endif // WITH_EDITOR
 
+bool UMaterialExpressionFunctionInput::IsAllowedIn(const UObject* MaterialOrFunction) const
+{
+	return MaterialOrFunction && MaterialOrFunction->IsA<UMaterialFunction>() && Super::IsAllowedIn(MaterialOrFunction);
+}
 
 ///////////////////////////////////////////////////////////////////////////////
 // UMaterialExpressionFunctionOutput
@@ -18337,9 +18337,13 @@ FSubstrateOperator* UMaterialExpressionFunctionOutput::SubstrateGenerateMaterial
 	}
 	return nullptr;
 }
-
 #endif // WITH_EDITOR
 
+bool UMaterialExpressionFunctionOutput::IsAllowedIn(const UObject* MaterialOrFunction) const
+{
+	return MaterialOrFunction && MaterialOrFunction->IsA<UMaterialFunction>() && Super::IsAllowedIn(MaterialOrFunction);
+}
+
 ///////////////////////////////////////////////////////////////////////////////
 // UMaterialExpressionMaterialLayerOutput
 ///////////////////////////////////////////////////////////////////////////////
diff --git a/Engine/Source/Runtime/Engine/Public/MaterialShared.h b/Engine/Source/Runtime/Engine/Public/MaterialShared.h
--- a/Engine/Source/Runtime/Engine/Public/MaterialShared.h
+++ b/Engine/Source/Runtime/Engine/Public/MaterialShared.h
@@ -1941,7 +1941,8 @@ private:
 };
 
 /** Returns whether the given expression class is allowed. */
-extern ENGINE_API bool IsAllowedExpressionType(const UClass* const Class, const bool bMaterialFunction);
+UE_DEPRECATED(5.8, "Should call UMaterialExpression::IsAllowedIn on the Class' DefaultObject")
+extern ENGINE_API bool IsAllowedExpressionType(const UClass* Class, const bool bMaterialFunction);
 
 /** Parses a string into multiple lines, for use with tooltips. */
 extern ENGINE_API void ConvertToMultilineToolTip(const FString& InToolTip, const int32 TargetLineLength, TArray<FString>& OutToolTip);
diff --git a/Engine/Source/Runtime/Engine/Public/Materials/MaterialExpression.h b/Engine/Source/Runtime/Engine/Public/Materials/MaterialExpression.h
--- a/Engine/Source/Runtime/Engine/Public/Materials/MaterialExpression.h
+++ b/Engine/Source/Runtime/Engine/Public/Materials/MaterialExpression.h
@@ -280,13 +280,22 @@ class UMaterialExpression : public UObject
 	ENGINE_API virtual void Build(MIR::FEmitter& Emitter);
 
 	/* Material expressions can add shader tags that to a material.  Adding a shader tag helps to identify materials
-	   that use certain material expressions, and can be used to limit shader generation to only materials that have
-	   use that material expression*/
+	 *  that use certain material expressions, and can be used to limit shader generation to only materials that have
+	 *  use that material expression
+	 */
 	ENGINE_API virtual void GetShaderTags(TArray<FName>& ShaderTagsOut) {}
 
 	ENGINE_API virtual TSharedPtr<class SGraphNodeMaterialBase> CreateCustomGraphNodeWidget() { return nullptr; }
 #endif // WITH_EDITOR
 
+	/**
+	* Returns whether the creation of this expression should be allowed in a given material or material function
+	*
+	* @param MaterialOrFunction - The material or function for this expression to test against
+	* @return True if this expression should be created in this material or function, false if it should not.
+	*/
+	ENGINE_API virtual bool IsAllowedIn(const UObject* MaterialOrFunction) const;
+
 	/**
 	* Fill the array with all textures dependence that should trig a recompile of the material.
 	*/
diff --git a/Engine/Source/Runtime/Engine/Public/Materials/MaterialExpressionComment.h b/Engine/Source/Runtime/Engine/Public/Materials/MaterialExpressionComment.h
--- a/Engine/Source/Runtime/Engine/Public/Materials/MaterialExpressionComment.h
+++ b/Engine/Source/Runtime/Engine/Public/Materials/MaterialExpressionComment.h
@@ -58,6 +58,7 @@ public:
 	virtual void GetCaption(TArray<FString>& OutCaptions) const override;
 	virtual bool MatchesSearchQuery( const TCHAR* SearchQuery ) override;
 #endif // WITH_EDITOR
+	virtual bool IsAllowedIn(const UObject* MaterialOrFunction) const override { return false; }
 	//~ End UMaterialExpression Interface
 };
 
diff --git a/Engine/Source/Runtime/Engine/Public/Materials/MaterialExpressionFunctionInput.h b/Engine/Source/Runtime/Engine/Public/Materials/MaterialExpressionFunctionInput.h
--- a/Engine/Source/Runtime/Engine/Public/Materials/MaterialExpressionFunctionInput.h
+++ b/Engine/Source/Runtime/Engine/Public/Materials/MaterialExpressionFunctionInput.h
@@ -112,6 +112,7 @@ class UMaterialExpressionFunctionInput : public UMaterialExpression
 	virtual uint32 GetInputType(int32 InputIndex) override;
 	virtual uint32 GetOutputType(int32 OutputIndex) override;
 #endif // WITH_EDITOR
+	virtual bool IsAllowedIn(const UObject* MaterialOrFunction) const override;
 	//~ End UMaterialExpression Interface
 
 	/** Generate the Id for this input. */
diff --git a/Engine/Source/Runtime/Engine/Public/Materials/MaterialExpressionFunctionOutput.h b/Engine/Source/Runtime/Engine/Public/Materials/MaterialExpressionFunctionOutput.h
--- a/Engine/Source/Runtime/Engine/Public/Materials/MaterialExpressionFunctionOutput.h
+++ b/Engine/Source/Runtime/Engine/Public/Materials/MaterialExpressionFunctionOutput.h
@@ -69,6 +69,7 @@ class UMaterialExpressionFunctionOutput : public UMaterialExpression
 	virtual void GatherSubstrateMaterialInfo(FSubstrateMaterialInfo& SubstrateMaterialInfo, int32 OutputIndex) override;
 	virtual FSubstrateOperator* SubstrateGenerateMaterialTopologyTree(class FMaterialCompiler* Compiler, class UMaterialExpression* Parent, int32 OutputIndex) override;
 #endif // WITH_EDITOR
+	virtual bool IsAllowedIn(const UObject* MaterialOrFunction) const override;
 	//~ End UMaterialExpression Interface
 
 
diff --git a/Engine/Source/Runtime/Engine/Public/Materials/MaterialExpressionParameter.h b/Engine/Source/Runtime/Engine/Public/Materials/MaterialExpressionParameter.h
--- a/Engine/Source/Runtime/Engine/Public/Materials/MaterialExpressionParameter.h
+++ b/Engine/Source/Runtime/Engine/Public/Materials/MaterialExpressionParameter.h
@@ -61,6 +61,7 @@ class UMaterialExpressionParameter : public UMaterialExpression
 
 	virtual void Build(MIR::FEmitter& Emitter) override;
 #endif
+	virtual bool IsAllowedIn(const UObject* MaterialOrFunction) const override { return false; }
 	//~ End UMaterialExpression Interface
 
 	virtual FGuid& GetParameterExpressionId() override
diff --git a/Engine/Source/Runtime/Engine/Public/Materials/MaterialExpressionPinBase.h b/Engine/Source/Runtime/Engine/Public/Materials/MaterialExpressionPinBase.h
--- a/Engine/Source/Runtime/Engine/Public/Materials/MaterialExpressionPinBase.h
+++ b/Engine/Source/Runtime/Engine/Public/Materials/MaterialExpressionPinBase.h
@@ -86,8 +86,12 @@ class UMaterialExpressionPinBase : public UMaterialExpression
 	virtual void ConnectExpression(FExpressionInput* Input, int32 OutputIndex) override;
 
 	virtual void GetCaption(TArray<FString>& OutCaptions) const override;
+#endif // WITH_EDITOR
+
+	virtual bool IsAllowedIn(const UObject* MaterialOrFunction) const override { return false; }
 	//~ End UMaterialExpression Interface
 
+#if WITH_EDITOR
 	//~ Begin UObject Interface.
 	virtual bool Modify(bool bAlwaysMarkDirty = true) override;
 	//~ End UObject Interface.
diff --git a/Engine/Source/Runtime/Engine/Public/Materials/MaterialExpressionTextureSampleParameter2DArray.h b/Engine/Source/Runtime/Engine/Public/Materials/MaterialExpressionTextureSampleParameter2DArray.h
--- a/Engine/Source/Runtime/Engine/Public/Materials/MaterialExpressionTextureSampleParameter2DArray.h
+++ b/Engine/Source/Runtime/Engine/Public/Materials/MaterialExpressionTextureSampleParameter2DArray.h
@@ -24,6 +24,7 @@ public:
 #if WITH_EDITOR
 	virtual bool TextureIsValid( UTexture* InTexture, FString& OutMessage) override;
 #endif
+	virtual bool IsAllowedIn(const UObject* MaterialOrFunction) const override;
 	virtual const TCHAR* GetRequirements();
 	// End UMaterialExpressionTextureSampleParameter Interface
 };
\ No newline at end of file
