Fixes bug triggered by copy pasting then deleting collapsed nodes, causing Input and Output pin base expressions to linger in the material.

diff --git a/Engine/Source/Editor/MaterialEditor/Private/MaterialEditor.cpp b/Engine/Source/Editor/MaterialEditor/Private/MaterialEditor.cpp
--- a/Engine/Source/Editor/MaterialEditor/Private/MaterialEditor.cpp
+++ b/Engine/Source/Editor/MaterialEditor/Private/MaterialEditor.cpp
@@ -6565,13 +6565,13 @@ void FMaterialEditor::DeleteSelectedDuplicatableNodes()
 	}
 }
 
-void FMaterialEditor::DeleteNodesInternal(const TArray<class UEdGraphNode*>& NodesToDelete, bool& bHaveExpressionsToDelete, bool& bPreviewExpressionDeleted)
+void FMaterialEditor::DeleteNodesInternal(const TArray<class UEdGraphNode*>& NodesToDelete, bool& bHaveExpressionsToDelete, bool& bPreviewExpressionDeleted, bool bForceDeletion)
 {
 	ModifyMaterial();
 
 	for (int32 Index = 0; Index < NodesToDelete.Num(); ++Index)
 	{
-		if (NodesToDelete[Index]->CanUserDeleteNode())
+		if (bForceDeletion || NodesToDelete[Index]->CanUserDeleteNode())
 		{
 			if (UMaterialGraphNode* GraphNode = Cast<UMaterialGraphNode>(NodesToDelete[Index]))
 			{
@@ -6606,7 +6606,7 @@ void FMaterialEditor::DeleteNodesInternal(const TArray<class UEdGraphNode*>& Nod
 
 					// Remove all subgraph nodes, note that composites remove their subgraph in DestroyNode
 					const TArray<UEdGraphNode*> SubgraphNodesToDelete = CompositeNode->BoundGraph->Nodes;
-					DeleteNodesInternal(SubgraphNodesToDelete, bHaveExpressionsToDelete, bPreviewExpressionDeleted);
+					DeleteNodesInternal(SubgraphNodesToDelete, bHaveExpressionsToDelete, bPreviewExpressionDeleted, true);
 				}
 				if(MaterialExpression)
 				{
diff --git a/Engine/Source/Editor/MaterialEditor/Private/MaterialEditor.h b/Engine/Source/Editor/MaterialEditor/Private/MaterialEditor.h
--- a/Engine/Source/Editor/MaterialEditor/Private/MaterialEditor.h
+++ b/Engine/Source/Editor/MaterialEditor/Private/MaterialEditor.h
@@ -563,7 +563,7 @@ protected:
 	/** Delete only the currently selected nodes that can be duplicated */
 	void DeleteSelectedDuplicatableNodes();
 	/** Recursively deletes nodes if needed */
-	void DeleteNodesInternal(const TArray<class UEdGraphNode*>& NodesToDelete, bool& bHaveExpressionsToDelete, bool& bPreviewExpressionDeleted);
+	void DeleteNodesInternal(const TArray<class UEdGraphNode*>& NodesToDelete, bool& bHaveExpressionsToDelete, bool& bPreviewExpressionDeleted, bool bForceDeletion = false);
 
 	/** Copy the currently selected nodes */
 	void CopySelectedNodes();
