Fixed a bug where MaterialFunctionInstance failed to gather all parameters from the material function.

[RN] Fixed a bug where MaterialFunctionInstance failed to gather all parameters from the material function.

diff --git a/Engine/Source/Editor/MaterialEditor/Private/MaterialEditorUtilities.cpp b/Engine/Source/Editor/MaterialEditor/Private/MaterialEditorUtilities.cpp
--- a/Engine/Source/Editor/MaterialEditor/Private/MaterialEditorUtilities.cpp
+++ b/Engine/Source/Editor/MaterialEditor/Private/MaterialEditorUtilities.cpp
@@ -287,6 +287,13 @@ void FMaterialEditorUtilities::GetVisibleMaterialParameters(const UMaterial* Mat
 	{
 		GetVisibleMaterialParametersFromExpression(FMaterialExpressionKey(Expression, 0), MaterialInstance, VisibleExpressions, FunctionStack);
 	}
+
+	TArray<UMaterialExpressionFunctionOutput*> FunctionOutputExpressions;
+	Material->GetAllFunctionOutputExpressions(FunctionOutputExpressions);
+	for (UMaterialExpressionFunctionOutput* Expression : FunctionOutputExpressions)
+	{
+		GetVisibleMaterialParametersFromExpression(FMaterialExpressionKey(Expression, 0), MaterialInstance, VisibleExpressions, FunctionStack);
+	}
 }
 
 bool FMaterialEditorUtilities::GetStaticSwitchExpressionValue(UMaterialInstance* MaterialInstance, UMaterialExpression* SwitchValueExpression, bool& bOutValue, FGuid& OutExpressionID, TArray<FGetVisibleMaterialParametersFunctionState*>& FunctionStack)
