Debug text in the MaterialInstanceEditor viewport should respect GAreScreenMessagesEnabled

- Now the debug text in the instance editor viewport can be toggled off with: (TOGGLEALLSCREENMESSAGES, TOGGLESCREENMESSAGES, CAPTUREMODE) or the cmd line option (NOSCREENMESSAGES)

#tests Open Material Instance in Lyra and run TOGGLEALLSCREENMESSAGES command

diff --git a/Engine/Source/Editor/MaterialEditor/Private/MaterialInstanceEditor.cpp b/Engine/Source/Editor/MaterialEditor/Private/MaterialInstanceEditor.cpp
--- a/Engine/Source/Editor/MaterialEditor/Private/MaterialInstanceEditor.cpp
+++ b/Engine/Source/Editor/MaterialEditor/Private/MaterialInstanceEditor.cpp
@@ -1485,7 +1485,7 @@ void FMaterialInstanceEditor::DrawMessages( FViewport* Viewport, FCanvas* Canvas
 		if ( BaseMaterial && MaterialResource )
 		{
 			const bool bGeneratedNewShaders = MaterialEditorInstance->SourceInstance->bHasStaticPermutationResource;
-			const bool bAllowOldMaterialStats = true;
+			const bool bAllowOldMaterialStats = GAreScreenMessagesEnabled;
 			FMaterialEditor::DrawMaterialInfoStrings( Canvas, BaseMaterial, MaterialResource, MaterialResource->GetCompileErrors(), DrawPositionY, bAllowOldMaterialStats, bGeneratedNewShaders );
 		}
 
