Remove on screen shader stat messages in the Material Instance Editor since this data was redundant with the existing Platform Stats docking tab.

- Still draw any on screen messages the material instance editor generates
- Still draw any error messages on screen if there are errors in the material
- Respect the disable on screen messages bool

#rnx

diff --git a/Engine/Source/Editor/MaterialEditor/Private/MaterialEditor.cpp b/Engine/Source/Editor/MaterialEditor/Private/MaterialEditor.cpp
--- a/Engine/Source/Editor/MaterialEditor/Private/MaterialEditor.cpp
+++ b/Engine/Source/Editor/MaterialEditor/Private/MaterialEditor.cpp
@@ -2486,147 +2486,6 @@ bool FMaterialEditor::MatchesContext(const FTransactionContext& InContext, const
 	return false;
 }
 
-void FMaterialEditor::DrawMaterialInfoStrings(
-	FCanvas* Canvas, 
-	const UMaterial* Material, 
-	const FMaterialResource* MaterialResource, 
-	const TArray<FString>& CompileErrors, 
-	int32 &DrawPositionY,
-	bool bDrawInstructions,
-	bool bGeneratedNewShaders)
-{
-	check(Material && MaterialResource);
-
-	ERHIFeatureLevel::Type FeatureLevel = MaterialResource->GetFeatureLevel();
-	FString FeatureLevelName;
-	GetFeatureLevelName(FeatureLevel,FeatureLevelName);
-
-	// The font to use when displaying info strings
-	UFont* FontToUse = GEngine->GetTinyFont();
-
-	constexpr int32 DrawPositionX = 5;
-	constexpr int32 SpacingBetweenLines = 13;
-
-	if (bDrawInstructions)
-	{
-		// Display any errors and messages in the upper left corner of the viewport.
-		TArray<FMaterialStatsUtils::FShaderInstructionsInfo> Descriptions;
-		FMaterialStatsUtils::GetRepresentativeInstructionCounts(Descriptions, MaterialResource);
-
-		for (int32 InstructionIndex = 0; InstructionIndex < Descriptions.Num(); InstructionIndex++)
-		{
-			FString InstructionCountString = FString::Printf(TEXT("%s: %u instructions\nStats: %s"), *Descriptions[InstructionIndex].ShaderDescription, Descriptions[InstructionIndex].InstructionCount, *Descriptions[InstructionIndex].ShaderStatisticsString);
-			DrawPositionY += Canvas->DrawShadowedString(DrawPositionX, DrawPositionY, *InstructionCountString, FontToUse, FLinearColor(1, 1, 0));
-		}
-
-		// Display the number of texture samplers and samplers used by the material.
-		const int32 SamplersUsed = MaterialResource->GetSamplerUsage();
-
-		if (SamplersUsed >= 0)
-		{
-			int32 MaxSamplers = GetExpectedFeatureLevelMaxTextureSamplers(MaterialResource->GetFeatureLevel());
-
-			Canvas->DrawShadowedString(
-				DrawPositionX,
-				DrawPositionY,
-				*FString::Printf(TEXT("%s samplers: %u/%u"), FeatureLevel <= ERHIFeatureLevel::ES3_1 ? TEXT("Mobile texture") : TEXT("Texture"), SamplersUsed, MaxSamplers),
-				FontToUse,
-				SamplersUsed > MaxSamplers ? FLinearColor(1,0,0) : FLinearColor(1,1,0)
-				);
-			DrawPositionY += SpacingBetweenLines;
-		}
-
-		uint32 NumVSTextureSamples = 0, NumPSTextureSamples = 0;
-		MaterialResource->GetEstimatedNumTextureSamples(NumVSTextureSamples, NumPSTextureSamples);
-
-		if (NumVSTextureSamples > 0 || NumPSTextureSamples > 0)
-		{
-			Canvas->DrawShadowedString(
-				DrawPositionX,
-				DrawPositionY,
-				*FString::Printf(TEXT("Texture Lookups (Est.): VS(%u), PS(%u)"), NumVSTextureSamples, NumPSTextureSamples),
-				FontToUse,
-				FLinearColor(1,1,0)
-				);
-			DrawPositionY += SpacingBetweenLines;
-		}
-
-		uint32 NumVirtualTextureLookups = MaterialResource->GetEstimatedNumVirtualTextureLookups();
-		if (NumVirtualTextureLookups > 0)
-		{
-			Canvas->DrawShadowedString(
-				DrawPositionX,
-				DrawPositionY,
-				*FString::Printf(TEXT("Virtual Texture Lookups (Est.): %u"), NumVirtualTextureLookups),
-				FontToUse,
-				FLinearColor(1, 1, 0)
-			);
-			DrawPositionY += SpacingBetweenLines;
-		}
-
-		const uint32 NumVirtualTextureStacks = MaterialResource->GetNumVirtualTextureStacks();
-		if (NumVirtualTextureStacks > 0)
-		{
-			Canvas->DrawShadowedString(
-				DrawPositionX,
-				DrawPositionY,
-				*FString::Printf(TEXT("Virtual Texture Stacks: %u"), NumVirtualTextureStacks),
-				FontToUse,
-				FLinearColor(1, 1, 0)
-			);
-			DrawPositionY += SpacingBetweenLines;
-		}
-
-		FMaterialResource::FLWCUsagesArray LWCFuncUsagesVS;
-		FMaterialResource::FLWCUsagesArray LWCFuncUsagesPS;
-		FMaterialResource::FLWCUsagesArray LWCFuncUsagesCS;
-		MaterialResource->GetEstimatedLWCFuncUsages(LWCFuncUsagesVS, LWCFuncUsagesPS, LWCFuncUsagesCS);
-		for (int KindIndex = 0; KindIndex < (int)ELWCFunctionKind::Max; ++KindIndex)
-		{
-			int Usages = LWCFuncUsagesVS[KindIndex] + LWCFuncUsagesPS[KindIndex] + LWCFuncUsagesCS[KindIndex];
-			if (Usages > 0)
-			{
-				Canvas->DrawShadowedString(
-					DrawPositionX,
-					DrawPositionY,
-					*FString::Printf(
-						TEXT("LWC %s usages (Est.): %u (VS), %u (PS), %u (CS)"), 
-						*UEnum::GetDisplayValueAsText((ELWCFunctionKind)KindIndex).ToString(), 
-						LWCFuncUsagesVS[KindIndex],
-						LWCFuncUsagesPS[KindIndex],
-						LWCFuncUsagesCS[KindIndex]),
-					FontToUse,
-					FLinearColor(1,1,0)
-				);
-				DrawPositionY += SpacingBetweenLines;
-			}
-		}
-
-		if (bGeneratedNewShaders)
-		{
-			int32 NumShaders = 0;
-			int32 NumPipelines = 0;
-			if(FMaterialShaderMap* ShaderMap = MaterialResource->GetGameThreadShaderMap())
-			{
-				ShaderMap->CountNumShaders(NumShaders, NumPipelines);
-			}
-
-			if (NumShaders)
-			{
-				FString ShaderCountString = FString::Printf(TEXT("Num shaders added: %i"), NumShaders);
-				Canvas->DrawShadowedString(DrawPositionX, DrawPositionY, *ShaderCountString, FontToUse, FLinearColor(1, 0.8, 0));
-				DrawPositionY += SpacingBetweenLines;
-			}
-		}
-	}
-
-	for(int32 ErrorIndex = 0; ErrorIndex < CompileErrors.Num(); ErrorIndex++)
-	{
-		Canvas->DrawShadowedString(DrawPositionX, DrawPositionY, *FString::Printf(TEXT("[%s] %s"), *FeatureLevelName, *CompileErrors[ErrorIndex]), FontToUse, FLinearColor(1, 0, 0));
-		DrawPositionY += SpacingBetweenLines;
-	}
-}
-
 void FMaterialEditor::DrawMessages( FViewport* InViewport, FCanvas* Canvas )
 {
 	if( PreviewExpression != NULL )
diff --git a/Engine/Source/Editor/MaterialEditor/Private/MaterialEditor.h b/Engine/Source/Editor/MaterialEditor/Private/MaterialEditor.h
--- a/Engine/Source/Editor/MaterialEditor/Private/MaterialEditor.h
+++ b/Engine/Source/Editor/MaterialEditor/Private/MaterialEditor.h
@@ -278,18 +278,6 @@ public:
 	/** The material instance applied to the preview mesh. */
 	virtual UMaterialInterface* GetMaterialInterface() const override;
 	
-	/**
-	 * Draws material info strings such as instruction count and current errors onto the canvas.
-	 */
-	static void DrawMaterialInfoStrings(
-		FCanvas* Canvas,
-		const UMaterial* Material,
-		const FMaterialResource* MaterialResource,
-		const TArray<FString>& CompileErrors,
-		int32 &DrawPositionY,
-		bool bDrawInstructions,
-		bool bGeneratedNewShaders = false);
-	
 	/**
 	 * Draws messages on the specified viewport and canvas.
 	 */
diff --git a/Engine/Source/Editor/MaterialEditor/Private/MaterialInstanceEditor.cpp b/Engine/Source/Editor/MaterialEditor/Private/MaterialInstanceEditor.cpp
--- a/Engine/Source/Editor/MaterialEditor/Private/MaterialInstanceEditor.cpp
+++ b/Engine/Source/Editor/MaterialEditor/Private/MaterialInstanceEditor.cpp
@@ -231,6 +231,8 @@ protected:
 //////////////////////////////////////////////////////////////////////////
 // FMaterialInstanceEditor
 
+static constexpr int32 SpacingBetweenLines = 13;
+
 void FMaterialInstanceEditor::RegisterTabSpawners(const TSharedRef<class FTabManager>& InTabManager)
 {
 	WorkspaceMenuCategory = InTabManager->AddLocalWorkspaceMenuCategory(LOCTEXT("WorkspaceMenu_MaterialInstanceEditor", "Material Instance Editor"));
@@ -1477,19 +1479,28 @@ void FMaterialInstanceEditor::RebuildMaterialInstanceEditor()
 void FMaterialInstanceEditor::DrawMessages( FViewport* Viewport, FCanvas* Canvas )
 {
 	Canvas->PushAbsoluteTransform(FMatrix::Identity);
-	if ( MaterialEditorInstance->Parent && MaterialEditorInstance->SourceInstance )
+	if ( MaterialEditorInstance->Parent && MaterialEditorInstance->SourceInstance && GAreScreenMessagesEnabled )
 	{
 		const FMaterialResource* MaterialResource = MaterialEditorInstance->SourceInstance->GetMaterialResource(GMaxRHIFeatureLevel);
 		UMaterial* BaseMaterial = MaterialEditorInstance->SourceInstance->GetMaterial();
 		int32 DrawPositionY = 50;
 		if ( BaseMaterial && MaterialResource )
 		{
-			const bool bGeneratedNewShaders = MaterialEditorInstance->SourceInstance->bHasStaticPermutationResource;
-			const bool bAllowOldMaterialStats = GAreScreenMessagesEnabled;
-			FMaterialEditor::DrawMaterialInfoStrings( Canvas, BaseMaterial, MaterialResource, MaterialResource->GetCompileErrors(), DrawPositionY, bAllowOldMaterialStats, bGeneratedNewShaders );
+			constexpr int32 DrawPositionX = 5;
+
+			ERHIFeatureLevel::Type FeatureLevel = MaterialResource->GetFeatureLevel();
+			FString FeatureLevelName;
+			GetFeatureLevelName(FeatureLevel,FeatureLevelName);
+
+			const TArray<FString>& CompileErrors = MaterialResource->GetCompileErrors();
+			for(int32 ErrorIndex = 0; ErrorIndex < CompileErrors.Num(); ErrorIndex++)
+			{
+				Canvas->DrawShadowedString(DrawPositionX, DrawPositionY, *FString::Printf(TEXT("[%s] %s"), *FeatureLevelName, *CompileErrors[ErrorIndex]), GEngine->GetTinyFont(), FLinearColor(1, 0, 0));
+				DrawPositionY += SpacingBetweenLines;
+			}
 		}
 
-		DrawSamplerWarningStrings( Canvas, DrawPositionY );
+		DrawOnScreenMessages( Canvas, DrawPositionY );
 	}
 	Canvas->PopTransform();
 }
@@ -1630,9 +1641,8 @@ void FMaterialInstanceEditor::RefreshOnScreenMessages()
  * @param Canvas - The canvas on which to draw.
  * @param DrawPositionY - The Y position at which to draw. Upon return contains the Y value following the last line of text drawn.
  */
-void FMaterialInstanceEditor::DrawSamplerWarningStrings(FCanvas* Canvas, int32& DrawPositionY)
+void FMaterialInstanceEditor::DrawOnScreenMessages(FCanvas* Canvas, int32& DrawPositionY)
 {
-	const int32 SpacingBetweenLines = 13;
 	UFont* FontToUse = GEngine->GetTinyFont();
 	for (const FOnScreenMessage& Message : OnScreenMessages)
 	{
diff --git a/Engine/Source/Editor/MaterialEditor/Private/MaterialInstanceEditor.h b/Engine/Source/Editor/MaterialEditor/Private/MaterialInstanceEditor.h
--- a/Engine/Source/Editor/MaterialEditor/Private/MaterialInstanceEditor.h
+++ b/Engine/Source/Editor/MaterialEditor/Private/MaterialInstanceEditor.h
@@ -95,11 +95,11 @@ public:
 	virtual void DrawMessages( FViewport* Viewport, FCanvas* Canvas ) override;
 
 	/**
-	 * Draws sampler/texture mismatch warning strings.
+	 * Draws generic on screen messages.
 	 * @param Canvas - The canvas on which to draw.
 	 * @param DrawPositionY - The Y position at which to draw. Upon return contains the Y value following the last line of text drawn.
 	 */
-	void DrawSamplerWarningStrings(FCanvas* Canvas, int32& DrawPositionY);
+	void DrawOnScreenMessages(FCanvas* Canvas, int32& DrawPositionY);
 
 	/** Passes instructions to the preview viewport */
 	bool SetPreviewAsset(UObject* InAsset);
