Pull Material Editor Preivew Mesh icons into a group to use less space.

diff --git a/Engine/Source/Editor/MaterialEditor/Private/MaterialEditorViewportToolbarSections.cpp b/Engine/Source/Editor/MaterialEditor/Private/MaterialEditorViewportToolbarSections.cpp
--- a/Engine/Source/Editor/MaterialEditor/Private/MaterialEditorViewportToolbarSections.cpp
+++ b/Engine/Source/Editor/MaterialEditor/Private/MaterialEditorViewportToolbarSections.cpp
@@ -98,6 +98,10 @@ void UE::MaterialEditor::ExtendPreviewSceneSettingsSubmenu(FName InSubmenuName)
 	FToolMenuSection& PreviewMeshSection = Submenu->FindOrAddSection(
 		"AssetViewerPreviewMeshSection", LOCTEXT("AssetViewerPreviewMeshSectionLabel", "Preview Mesh Options")
 	);
+	
+	FToolMenuEntryToolBarData ToolbarData;
+	ToolbarData.BlockGroupName = "PreviewMeshOptions";
+	ToolbarData.ResizeParams.ClippingPriority = 2000;
 
 	for (const TSharedPtr<FUICommandInfo>& Command : {
 			 FMaterialEditorCommands::Get().SetSpherePreview,
@@ -108,7 +112,7 @@ void UE::MaterialEditor::ExtendPreviewSceneSettingsSubmenu(FName InSubmenuName)
 	{
 		FToolMenuEntry& Entry = PreviewMeshSection.AddMenuEntry(Command);
 		Entry.SetShowInToolbarTopLevel(true);
-		Entry.ToolBarData.ResizeParams.ClippingPriority = 2000;
+		Entry.ToolBarData = ToolbarData;
 	}
 
 	FToolMenuEntry& MeshFromSelectionEntry =
@@ -116,7 +120,7 @@ void UE::MaterialEditor::ExtendPreviewSceneSettingsSubmenu(FName InSubmenuName)
 	MeshFromSelectionEntry.Label = LOCTEXT("SetPreviewMeshFromSelectionLabel", "Static Mesh in Content Browser");
 	MeshFromSelectionEntry.Icon = FSlateIconFinder::FindIconForClass(UStaticMesh::StaticClass());
 	MeshFromSelectionEntry.SetShowInToolbarTopLevel(true);
-	MeshFromSelectionEntry.ToolBarData.ResizeParams.ClippingPriority = 2000;
+	MeshFromSelectionEntry.ToolBarData = ToolbarData;
 }
 
 #undef LOCTEXT_NAMESPACE
