Moving all material editor related code from GraphNode module into the MaterialEditor module.

- GraphNode module now includes MaterialEditor module as a public dependency so any client code that is dependent on the public headers will still compile.
- Made a connection drawing policy factory to hide material editor implementation details.
- Implement UMaterialGraphNode_Comment::CreateVisualWidget so the graph node is responsible for making the swidget and not GraphEditor code
- Cleanup code in NodeFactory.cpp - everything should be delegated to factories or material editor specific code.
- Change dll export to MATERIALEDITOR_API instead of GRAPHEDITOR_API
- Engine module Editor builds must include/link MaterialEditor in addition to GraphEditor

[RN] major | Rendering - Materials and Shaders

diff --git a/Engine/Source/Editor/GraphEditor/GraphEditor.Build.cs b/Engine/Source/Editor/GraphEditor/GraphEditor.Build.cs
--- a/Engine/Source/Editor/GraphEditor/GraphEditor.Build.cs
+++ b/Engine/Source/Editor/GraphEditor/GraphEditor.Build.cs
@@ -13,6 +13,7 @@ public class GraphEditor : ModuleRules
             new string[] {                
 				"ClassViewer",
 				"StructViewer",
+				"MaterialEditor"
 			}
         );
 
diff --git a/Engine/Source/Editor/GraphEditor/Private/MaterialGraphConnectionDrawingPolicy.cpp b/Engine/Source/Editor/GraphEditor/Private/MaterialGraphConnectionDrawingPolicy.cpp
deleted file mode 100644
diff --git a/Engine/Source/Editor/GraphEditor/Private/MaterialGraphConnectionDrawingPolicy.h b/Engine/Source/Editor/GraphEditor/Private/MaterialGraphConnectionDrawingPolicy.h
deleted file mode 100644
diff --git a/Engine/Source/Editor/GraphEditor/Private/MaterialNodes/SGraphNodeMaterialBase.cpp b/Engine/Source/Editor/GraphEditor/Private/MaterialNodes/SGraphNodeMaterialBase.cpp
deleted file mode 100644
diff --git a/Engine/Source/Editor/GraphEditor/Private/MaterialNodes/SGraphNodeMaterialComment.cpp b/Engine/Source/Editor/GraphEditor/Private/MaterialNodes/SGraphNodeMaterialComment.cpp
deleted file mode 100644
diff --git a/Engine/Source/Editor/GraphEditor/Private/MaterialNodes/SGraphNodeMaterialComment.h b/Engine/Source/Editor/GraphEditor/Private/MaterialNodes/SGraphNodeMaterialComment.h
deleted file mode 100644
diff --git a/Engine/Source/Editor/GraphEditor/Private/MaterialNodes/SGraphNodeMaterialComposite.cpp b/Engine/Source/Editor/GraphEditor/Private/MaterialNodes/SGraphNodeMaterialComposite.cpp
deleted file mode 100644
diff --git a/Engine/Source/Editor/GraphEditor/Private/MaterialNodes/SGraphNodeMaterialConvert.cpp b/Engine/Source/Editor/GraphEditor/Private/MaterialNodes/SGraphNodeMaterialConvert.cpp
deleted file mode 100644
diff --git a/Engine/Source/Editor/GraphEditor/Private/MaterialNodes/SGraphNodeMaterialCustom.cpp b/Engine/Source/Editor/GraphEditor/Private/MaterialNodes/SGraphNodeMaterialCustom.cpp
deleted file mode 100644
diff --git a/Engine/Source/Editor/GraphEditor/Private/MaterialNodes/SGraphNodeMaterialResult.cpp b/Engine/Source/Editor/GraphEditor/Private/MaterialNodes/SGraphNodeMaterialResult.cpp
deleted file mode 100644
diff --git a/Engine/Source/Editor/GraphEditor/Private/MaterialNodes/SGraphSubstrateMaterial.cpp b/Engine/Source/Editor/GraphEditor/Private/MaterialNodes/SGraphSubstrateMaterial.cpp
deleted file mode 100644
diff --git a/Engine/Source/Editor/GraphEditor/Private/MaterialPins/SGraphPinMaterialInput.cpp b/Engine/Source/Editor/GraphEditor/Private/MaterialPins/SGraphPinMaterialInput.cpp
deleted file mode 100644
diff --git a/Engine/Source/Editor/GraphEditor/Private/NodeFactory.cpp b/Engine/Source/Editor/GraphEditor/Private/NodeFactory.cpp
--- a/Engine/Source/Editor/GraphEditor/Private/NodeFactory.cpp
+++ b/Engine/Source/Editor/GraphEditor/Private/NodeFactory.cpp
@@ -62,18 +62,6 @@
 #include "KismetPins/SGraphPinText.h"
 #include "KismetPins/SGraphPinVector.h"
 #include "KismetPins/SGraphPinVector2D.h"
-#include "MaterialGraph/MaterialGraphNode.h"
-#include "MaterialGraph/MaterialGraphNode_Base.h"
-#include "MaterialGraph/MaterialGraphNode_Comment.h"
-#include "MaterialGraph/MaterialGraphNode_Composite.h"
-#include "MaterialGraph/MaterialGraphNode_Root.h"
-#include "MaterialGraph/MaterialGraphSchema.h"
-#include "MaterialGraphConnectionDrawingPolicy.h"
-#include "MaterialGraphNode_Knot.h"
-#include "MaterialNodes/SGraphNodeMaterialBase.h"
-#include "MaterialNodes/SGraphNodeMaterialComment.h"
-#include "MaterialNodes/SGraphNodeMaterialComposite.h"
-#include "MaterialNodes/SGraphNodeMaterialResult.h"
 #include "Math/Rotator.h"
 #include "Math/UnrealMathSSE.h"
 #include "Math/Vector2D.h"
@@ -196,14 +184,7 @@ TSharedPtr<SGraphNode> FNodeFactory::CreateNodeWidget(UEdGraphNode* InNode)
 	}
 	else if (UEdGraphNode_Comment* CommentNode = Cast<UEdGraphNode_Comment>(InNode))
 	{
-		if (UMaterialGraphNode_Comment* MaterialCommentNode = Cast<UMaterialGraphNode_Comment>(InNode))
-		{
-			return SNew(SGraphNodeMaterialComment, MaterialCommentNode);
-		}
-		else
-		{
-			return SNew(SGraphNodeComment, CommentNode);
-		}
+		return SNew(SGraphNodeComment, CommentNode);
 	}
 	else
 	{
@@ -392,10 +373,6 @@ FConnectionDrawingPolicy* FNodeFactory::CreateConnectionPolicy(const UEdGraphSch
 		{
 			ConnectionDrawingPolicy = new FKismetConnectionDrawingPolicy(InBackLayerID, InFrontLayerID, ZoomFactor, InClippingRect, InDrawElements, InGraphObj);
 		}
-		else if (Schema->IsA(UMaterialGraphSchema::StaticClass()))
-		{
-			ConnectionDrawingPolicy = new FMaterialGraphConnectionDrawingPolicy(InBackLayerID, InFrontLayerID, ZoomFactor, InClippingRect, InDrawElements, InGraphObj);
-		}
 		else
 		{
 			ConnectionDrawingPolicy = new FConnectionDrawingPolicy(InBackLayerID, InFrontLayerID, ZoomFactor, InClippingRect, InDrawElements);
diff --git a/Engine/Source/Editor/GraphEditor/Public/MaterialNodes/SGraphNodeMaterialBase.h b/Engine/Source/Editor/GraphEditor/Public/MaterialNodes/SGraphNodeMaterialBase.h
deleted file mode 100644
diff --git a/Engine/Source/Editor/GraphEditor/Public/MaterialNodes/SGraphNodeMaterialComposite.h b/Engine/Source/Editor/GraphEditor/Public/MaterialNodes/SGraphNodeMaterialComposite.h
deleted file mode 100644
diff --git a/Engine/Source/Editor/GraphEditor/Public/MaterialNodes/SGraphNodeMaterialConvert.h b/Engine/Source/Editor/GraphEditor/Public/MaterialNodes/SGraphNodeMaterialConvert.h
deleted file mode 100644
diff --git a/Engine/Source/Editor/GraphEditor/Public/MaterialNodes/SGraphNodeMaterialCustom.h b/Engine/Source/Editor/GraphEditor/Public/MaterialNodes/SGraphNodeMaterialCustom.h
deleted file mode 100644
diff --git a/Engine/Source/Editor/GraphEditor/Public/MaterialNodes/SGraphNodeMaterialResult.h b/Engine/Source/Editor/GraphEditor/Public/MaterialNodes/SGraphNodeMaterialResult.h
deleted file mode 100644
diff --git a/Engine/Source/Editor/GraphEditor/Public/MaterialPins/SGraphPinMaterialInput.h b/Engine/Source/Editor/GraphEditor/Public/MaterialPins/SGraphPinMaterialInput.h
deleted file mode 100644
diff --git a/Engine/Source/Editor/GraphEditor/Public/SGraphSubstrateMaterial.h b/Engine/Source/Editor/GraphEditor/Public/SGraphSubstrateMaterial.h
deleted file mode 100644
diff --git a/Engine/Source/Editor/MaterialEditor/Private/MaterialEditorModule.cpp b/Engine/Source/Editor/MaterialEditor/Private/MaterialEditorModule.cpp
--- a/Engine/Source/Editor/MaterialEditor/Private/MaterialEditorModule.cpp
+++ b/Engine/Source/Editor/MaterialEditor/Private/MaterialEditorModule.cpp
@@ -25,6 +25,7 @@
 #include "MaterialEditorGraphPanelPinFactory.h"
 #include "SubstrateMaterialEditorStyle.h"
 #include "MaterialEditorDetailCustomization.h"
+#include "MaterialGraphConnectionDrawingPolicy.h"
 
 const FName MaterialEditorAppIdentifier = FName(TEXT("MaterialEditorApp"));
 const FName MaterialInstanceEditorAppIdentifier = FName(TEXT("MaterialInstanceEditorApp"));
@@ -85,6 +86,9 @@ public:
 		GraphPanelPinFactory = MakeShared<FMaterialEditorGraphPanelPinFactory>();
 		FEdGraphUtilities::RegisterVisualPinFactory(GraphPanelPinFactory);
 
+		GraphConnectionFactory = MakeShared<FMaterialEditorGraphPinConnectionFactory>();
+		FEdGraphUtilities::RegisterVisualPinConnectionFactory(GraphConnectionFactory);
+
 		FSubstrateMaterialEditorStyle::Initialize();
 	}
 
@@ -162,6 +166,7 @@ private:
 	TSharedPtr<FExtensibilityManager> ToolBarExtensibilityManager;
 	FDelegateHandle ContentBrowserAssetExtenderDelegateHandle;
 	TSharedPtr<FMaterialEditorGraphPanelPinFactory> GraphPanelPinFactory;
+	TSharedPtr<FGraphPanelPinConnectionFactory> GraphConnectionFactory;
 };
 
 IMPLEMENT_MODULE( FMaterialEditorModule, MaterialEditor );
diff --git a/Engine/Source/Editor/MaterialEditor/Private/MaterialGraphConnectionDrawingPolicy.cpp b/Engine/Source/Editor/MaterialEditor/Private/MaterialGraphConnectionDrawingPolicy.cpp
new file mode 100644
--- /dev/null
+++ b/Engine/Source/Editor/MaterialEditor/Private/MaterialGraphConnectionDrawingPolicy.cpp
@@ -0,0 +1,300 @@
+// Copyright Epic Games, Inc. All Rights Reserved.
+
+#include "MaterialGraphConnectionDrawingPolicy.h"
+
+#include "Containers/Array.h"
+#include "Containers/Set.h"
+#include "EdGraph/EdGraph.h"
+#include "EdGraph/EdGraphPin.h"
+#include "GraphEditorSettings.h"
+#include "Layout/ArrangedWidget.h"
+#include "MaterialGraph/MaterialGraph.h"
+#include "MaterialGraph/MaterialGraphSchema.h"
+#include "MaterialGraphNode_Knot.h"
+#include "Math/Color.h"
+#include "SGraphPin.h"
+#include "Shader/ShaderTypes.h"
+#include "Templates/Casts.h"
+#include "Templates/SharedPointer.h"
+#include "UObject/NameTypes.h"
+#include "SGraphSubstrateMaterial.h"
+#include "Widgets/SToolTip.h"
+#include "Widgets/SBoxPanel.h"
+#include "Widgets/Text/STextBlock.h"
+
+class FSlateRect;
+
+/////////////////////////////////////////////////////
+// FMaterialGraphConnectionDrawingPolicy
+
+FMaterialGraphConnectionDrawingPolicy::FMaterialGraphConnectionDrawingPolicy(int32 InBackLayerID, int32 InFrontLayerID, float ZoomFactor, const FSlateRect& InClippingRect, FSlateWindowElementList& InDrawElements, UEdGraph* InGraphObj)
+	: FConnectionDrawingPolicy(InBackLayerID, InFrontLayerID, ZoomFactor, InClippingRect, InDrawElements)
+	, MaterialGraph(CastChecked<UMaterialGraph>(InGraphObj))
+	, MaterialGraphSchema(CastChecked<UMaterialGraphSchema>(InGraphObj->GetSchema()))
+{
+	// Don't want to draw ending arrowheads
+	ArrowImage = nullptr;
+	ArrowRadius = FVector2f::ZeroVector;
+
+	// Still need to be able to perceive the graph while dragging connectors, esp over comment boxes
+	HoverDeemphasisDarkFraction = 0.4f;
+}
+
+bool FMaterialGraphConnectionDrawingPolicy::FindPinCenter(UEdGraphPin* Pin, FVector2f& OutCenter) const
+{
+	if (const TSharedPtr<SGraphPin>* pPinWidget = PinToPinWidgetMap.Find(Pin))
+	{
+		if (FArrangedWidget* pPinEntry = PinGeometries->Find((*pPinWidget).ToSharedRef()))
+		{
+			OutCenter = FGeometryHelper::CenterOf(pPinEntry->Geometry);
+			return true;
+		}
+	}
+
+	return false;
+}
+
+bool FMaterialGraphConnectionDrawingPolicy::GetAverageConnectedPosition(class UMaterialGraphNode_Knot* Knot, EEdGraphPinDirection Direction, FVector2f& OutPos) const
+{
+	FVector2f Result = FVector2f::ZeroVector;
+	int32 ResultCount = 0;
+
+	UEdGraphPin* Pin = (Direction == EGPD_Input) ? Knot->GetInputPin() : Knot->GetOutputPin();
+	for (UEdGraphPin* LinkedPin : Pin->LinkedTo)
+	{
+		FVector2f CenterPoint;
+		if (FindPinCenter(LinkedPin, /*out*/ CenterPoint))
+		{
+			Result += CenterPoint;
+			ResultCount++;
+		}
+	}
+
+	if (ResultCount > 0)
+	{
+		OutPos = Result * (1.0f / static_cast<float>(ResultCount));
+		return true;
+	}
+	else
+	{
+		return false;
+	}
+}
+
+bool FMaterialGraphConnectionDrawingPolicy::ShouldChangeTangentForKnot(UMaterialGraphNode_Knot* Knot)
+{
+	if (bool* pResult = KnotToReversedDirectionMap.Find(Knot))
+	{
+		return *pResult;
+	}
+	else
+	{
+		bool bPinReversed = false;
+
+		FVector2f AverageLeftPin;
+		FVector2f AverageRightPin;
+		FVector2f CenterPin;
+		bool bCenterValid = FindPinCenter(Knot->GetOutputPin(), /*out*/ CenterPin);
+		bool bLeftValid = GetAverageConnectedPosition(Knot, EGPD_Input, /*out*/ AverageLeftPin);
+		bool bRightValid = GetAverageConnectedPosition(Knot, EGPD_Output, /*out*/ AverageRightPin);
+
+		if (bLeftValid && bRightValid)
+		{
+			bPinReversed = AverageRightPin.X < AverageLeftPin.X;
+		}
+		else if (bCenterValid)
+		{
+			if (bLeftValid)
+			{
+				bPinReversed = CenterPin.X < AverageLeftPin.X;
+			}
+			else if (bRightValid)
+			{
+				bPinReversed = AverageRightPin.X < CenterPin.X;
+			}
+		}
+
+		KnotToReversedDirectionMap.Add(Knot, bPinReversed);
+
+		return bPinReversed;
+	}
+}
+
+void FMaterialGraphConnectionDrawingPolicy::DetermineWiringStyle(UEdGraphPin* OutputPin, UEdGraphPin* InputPin, /*inout*/ FConnectionParams& Params)
+{
+	Params.AssociatedPin1 = OutputPin;
+	Params.AssociatedPin2 = InputPin;
+	Params.WireColor = MaterialGraphSchema->ActivePinColor;
+
+	if (Substrate::IsSubstrateEnabled() && (FSubstrateWidget::HasOutputSubstrateType(OutputPin) || FSubstrateWidget::HasInputSubstrateType(InputPin) || FSubstrateWidget::HasInputSubstrateType(OutputPin)))
+	{
+		Params.WireColor = FSubstrateWidget::GetConnectionColor();
+	}
+
+	UE::Shader::EValueType InputType = UE::Shader::EValueType::Void;
+	UE::Shader::EValueType OutputType = UE::Shader::EValueType::Void;
+	bool bInactivePin = false;
+	bool bExecPin = false;
+	bool bValueTypePin = false;
+
+	// Have to consider both pins as the input will be an 'output' when previewing a connection
+	if (OutputPin)
+	{
+		if (!MaterialGraph->IsInputActive(OutputPin))
+		{
+			bInactivePin = true;
+		}
+
+		if (OutputPin->PinType.PinCategory == UMaterialGraphSchema::PC_Exec)
+		{
+			bExecPin = true;
+		}
+		else if (OutputPin->PinType.PinCategory == UMaterialGraphSchema::PC_Void)
+		{
+			bValueTypePin = true;
+		}
+		else if (OutputPin->PinType.PinCategory == UMaterialGraphSchema::PC_ValueType)
+		{
+			OutputType = UE::Shader::FindValueType(OutputPin->PinType.PinSubCategory);
+			bValueTypePin = true;
+		}
+
+		UEdGraphNode* OutputNode = OutputPin->GetOwningNode();
+		if (UMaterialGraphNode_Knot* OutputKnotNode = Cast<UMaterialGraphNode_Knot>(OutputNode))
+		{
+			if (ShouldChangeTangentForKnot(OutputKnotNode))
+			{
+				Params.StartDirection = EGPD_Input;
+			}
+		}
+		else if (!OutputNode->IsNodeEnabled() || OutputNode->IsDisplayAsDisabledForced() || OutputNode->IsNodeUnrelated())
+		{
+			bInactivePin = true;
+		}
+	}
+	if (InputPin)
+	{
+		if (!MaterialGraph->IsInputActive(InputPin))
+		{
+			bInactivePin = true;
+		}
+
+		if (InputPin->PinType.PinCategory == UMaterialGraphSchema::PC_Exec)
+		{
+			bExecPin = true;
+		}
+		else if (InputPin->PinType.PinCategory == UMaterialGraphSchema::PC_Void)
+		{
+			bValueTypePin = true;
+		}
+		else if (InputPin->PinType.PinCategory == UMaterialGraphSchema::PC_ValueType)
+		{
+			InputType = UE::Shader::FindValueType(InputPin->PinType.PinSubCategory);
+			bValueTypePin = true;
+		}
+
+		UEdGraphNode* InputNode = InputPin->GetOwningNode();
+		if (UMaterialGraphNode_Knot* InputKnotNode = Cast<UMaterialGraphNode_Knot>(InputNode))
+		{
+			if (ShouldChangeTangentForKnot(InputKnotNode))
+			{
+				Params.EndDirection = EGPD_Output;
+			}
+		}
+		else if (!InputNode->IsNodeEnabled() || InputNode->IsDisplayAsDisabledForced() || InputNode->IsNodeUnrelated())
+		{
+			bInactivePin = true;
+		}
+	}
+
+	if (bInactivePin)
+	{
+		Params.WireColor = MaterialGraphSchema->InactivePinColor;
+	}
+	else if (bExecPin)
+	{
+		Params.WireColor = Settings->ExecutionPinTypeColor;
+		Params.WireThickness = Settings->DefaultExecutionWireThickness;
+	}
+	else if (bValueTypePin)
+	{
+		// TODO - how to handle InputPin being nullptr, when does this occur?
+		if (InputPin)
+		{
+			Params.WireColor = MaterialGraphSchema->GetPinTypeColor(InputPin->PinType);
+		}
+	}
+
+	const bool bDeemphasizeUnhoveredPins = HoveredPins.Num() > 0;
+	if (bDeemphasizeUnhoveredPins)
+	{
+		ApplyHoverDeemphasis(OutputPin, InputPin, /*inout*/ Params.WireThickness, /*inout*/ Params.WireColor);
+	}
+}
+
+TSharedPtr<IToolTip> FMaterialGraphConnectionDrawingPolicy::GetConnectionToolTip(const SGraphPanel& GraphPanel,	const FGraphSplineOverlapResult& OverlapData) const
+{
+	TSharedPtr<SGraphPin> Pin1Widget;
+	TSharedPtr<SGraphPin> Pin2Widget;
+	OverlapData.GetPinWidgets(GraphPanel, Pin1Widget, Pin2Widget);
+
+	if (!Pin1Widget || !Pin2Widget)
+	{
+		return FConnectionDrawingPolicy::GetConnectionToolTip(GraphPanel, OverlapData);
+	}
+
+	const FText LeftText = FText::Format(NSLOCTEXT("Unreal", "PinConnectionTooltipLeft", "<< {0}"), GetNodePinInfo(Pin1Widget));
+	const FText RightText = FText::Format(NSLOCTEXT("Unreal", "PinConnectionTooltipRight", "{0} >>"), GetNodePinInfo(Pin2Widget));
+
+	return SNew(SToolTip)
+		[
+			SNew(SVerticalBox)
+			+ SVerticalBox::Slot()
+			.HAlign(HAlign_Left)
+			.AutoHeight()
+			.Padding(0.0f, 0.0f, 0.0f, 5.0f)
+			[
+				SNew(STextBlock)
+					.Margin(FMargin(0.0f, 0.0f, 4.0f, 0.0f))
+					.Justification(ETextJustify::Left)
+					.Text(LeftText)
+			]
+			+ SVerticalBox::Slot()
+			.HAlign(HAlign_Right)
+			.AutoHeight()
+			.Padding(0.0f, 5.0f, 0.0f, 0.0f)
+			[
+				SNew(STextBlock)
+					.Margin(FMargin(4.0f, 0.0f, 0.0f, 0.0f))
+					.Justification(ETextJustify::Right)
+					.Text(RightText)
+			]
+		];
+}
+
+FText FMaterialGraphConnectionDrawingPolicy::GetNodePinInfo(const TSharedPtr<SGraphPin>& PinWidget) const
+{
+	const UEdGraphPin* PinObj = PinWidget->GetPinObj();
+	const UEdGraphNode* EdNode = PinObj->GetOwningNode();
+	 
+	FString NodeTitle = EdNode->GetNodeTitle(ENodeTitleType::ListView).ToString();
+	NodeTitle.RemoveFromStart(TEXT("Material Expression "));
+	
+	if (EdNode->GetCanRenameNode())
+	{
+		FText NodeEditableName = EdNode->GetNodeTitle(ENodeTitleType::EditableTitle);
+		NodeTitle = NodeTitle + TEXT(" (") + NodeEditableName.ToString() + TEXT(")");
+	}
+	const FText PinName = PinObj->GetDisplayName().IsEmptyOrWhitespace() ? FText::FromName(PinObj->PinName) : PinObj->GetDisplayName();
+	
+	return FText::Format(NSLOCTEXT("Unreal", "PinConnectionTooltipPartial", "{0}\r\n{1}"), FText::FromString(NodeTitle), /*NodeType,*/ PinName);
+}
+
+FConnectionDrawingPolicy* FMaterialEditorGraphPinConnectionFactory::CreateConnectionPolicy(const UEdGraphSchema* Schema, int32 InBackLayerID, int32 InFrontLayerID, float ZoomFactor, const FSlateRect& InClippingRect, FSlateWindowElementList& InDrawElements, UEdGraph* InGraphObj) const
+{
+	if (Schema->IsA(UMaterialGraphSchema::StaticClass()))
+	{
+		return new FMaterialGraphConnectionDrawingPolicy(InBackLayerID, InFrontLayerID, ZoomFactor, InClippingRect, InDrawElements, InGraphObj);
+	}
+	return nullptr;
+}
diff --git a/Engine/Source/Editor/MaterialEditor/Private/MaterialGraphConnectionDrawingPolicy.h b/Engine/Source/Editor/MaterialEditor/Private/MaterialGraphConnectionDrawingPolicy.h
new file mode 100644
--- /dev/null
+++ b/Engine/Source/Editor/MaterialEditor/Private/MaterialGraphConnectionDrawingPolicy.h
@@ -0,0 +1,53 @@
+// Copyright Epic Games, Inc. All Rights Reserved.
+
+#pragma once
+
+#include "ConnectionDrawingPolicy.h"
+#include "Containers/Map.h"
+#include "EdGraph/EdGraphNode.h"
+#include "EdGraphUtilities.h"
+#include "HAL/Platform.h"
+#include "Math/Vector2D.h"
+
+class FSlateRect;
+class FSlateWindowElementList;
+class UEdGraph;
+class UEdGraphPin;
+class UMaterialGraph;
+class UMaterialGraphSchema;
+
+/////////////////////////////////////////////////////
+// FMaterialGraphConnectionDrawingPolicy
+
+// This class draws the connections for an UEdGraph using a MaterialGraph schema
+class FMaterialGraphConnectionDrawingPolicy : public FConnectionDrawingPolicy
+{
+protected:
+	UMaterialGraph* MaterialGraph;
+	const UMaterialGraphSchema* MaterialGraphSchema;
+
+	// Each time a knot is encountered, input geometry is compared to output geometry to see if the pins on the knot need to be reversed
+	TMap<class UMaterialGraphNode_Knot*, bool> KnotToReversedDirectionMap;
+
+	bool ShouldChangeTangentForKnot(class UMaterialGraphNode_Knot* Knot);
+	bool GetAverageConnectedPosition(class UMaterialGraphNode_Knot* Knot, EEdGraphPinDirection Direction, FVector2f& OutPos) const;
+	bool FindPinCenter(UEdGraphPin* Pin, FVector2f& OutCenter) const;
+
+public:
+	FMaterialGraphConnectionDrawingPolicy(int32 InBackLayerID, int32 InFrontLayerID, float ZoomFactor, const FSlateRect& InClippingRect, FSlateWindowElementList& InDrawElements, UEdGraph* InGraphObj);
+
+	virtual TSharedPtr<IToolTip> GetConnectionToolTip(const SGraphPanel& GraphPanel, const FGraphSplineOverlapResult& OverlapData) const override;
+	
+	// FConnectionDrawingPolicy interface
+	virtual void DetermineWiringStyle(UEdGraphPin* OutputPin, UEdGraphPin* InputPin, /*inout*/ FConnectionParams& Params) override;
+	// End of FConnectionDrawingPolicy interface
+
+private:
+	FText GetNodePinInfo(const TSharedPtr<SGraphPin>& PinWidget) const;
+};
+
+struct FMaterialEditorGraphPinConnectionFactory : public FGraphPanelPinConnectionFactory
+{
+public:
+	virtual class FConnectionDrawingPolicy* CreateConnectionPolicy(const class UEdGraphSchema* Schema, int32 InBackLayerID, int32 InFrontLayerID, float ZoomFactor, const class FSlateRect& InClippingRect, class FSlateWindowElementList& InDrawElements, class UEdGraph* InGraphObj) const override;
+};
diff --git a/Engine/Source/Editor/MaterialEditor/Private/MaterialNodes/SGraphNodeMaterialBase.cpp b/Engine/Source/Editor/MaterialEditor/Private/MaterialNodes/SGraphNodeMaterialBase.cpp
new file mode 100644
--- /dev/null
+++ b/Engine/Source/Editor/MaterialEditor/Private/MaterialNodes/SGraphNodeMaterialBase.cpp
@@ -0,0 +1,745 @@
+// Copyright Epic Games, Inc. All Rights Reserved.
+
+#include "MaterialNodes/SGraphNodeMaterialBase.h"
+
+#include "CanvasItem.h"
+#include "CanvasTypes.h"
+#include "Containers/Array.h"
+#include "Containers/EnumAsByte.h"
+#include "Containers/UnrealString.h"
+#include "CoreGlobals.h"
+#include "Delegates/Delegate.h"
+#include "EdGraph/EdGraphNode.h"
+#include "EdGraph/EdGraphPin.h"
+#include "Engine/World.h"
+#include "GenericPlatform/ICursor.h"
+#include "GraphEditor.h"
+#include "GraphEditorSettings.h"
+#include "HAL/PlatformCrt.h"
+#include "Layout/Geometry.h"
+#include "Layout/Margin.h"
+#include "Layout/SlateRect.h"
+#include "MaterialGraph/MaterialGraph.h"
+#include "MaterialGraph/MaterialGraphNode.h"
+#include "MaterialGraph/MaterialGraphSchema.h"
+#include "Materials/MaterialExpression.h"
+#include "Materials/MaterialExpressionBreakMaterialAttributes.h"
+#include "Materials/MaterialExpressionMakeMaterialAttributes.h"
+#include "Materials/MaterialExpressionMaterialFunctionCall.h"
+#include "Materials/MaterialExpressionSubstrate.h"
+#include "Materials/MaterialFunction.h"
+#include "Math/Color.h"
+#include "Math/IntPoint.h"
+#include "Math/IntRect.h"
+#include "Misc/App.h"
+#include "Misc/AssertionMacros.h"
+#include "Misc/Attribute.h"
+#include "Misc/Guid.h"
+#include "Misc/Optional.h"
+#include "RHI.h"
+#include "RHICommandList.h"
+#include "Rendering/DrawElements.h"
+#include "Rendering/SlateRenderer.h"
+#include "RenderingThread.h"
+#include "SGraphPanel.h"
+#include "SGraphPin.h"
+#include "SlotBase.h"
+#include "Styling/AppStyle.h"
+#include "Styling/ISlateStyle.h"
+#include "Templates/Casts.h"
+#include "TutorialMetaData.h"
+#include "Types/SlateEnums.h"
+#include "Types/SlateStructs.h"
+#include "UObject/NameTypes.h"
+#include "UObject/ObjectPtr.h"
+#include "UnrealClient.h"
+#include "Widgets/Images/SImage.h"
+#include "Widgets/Input/SButton.h"
+#include "Widgets/Input/SCheckBox.h"
+#include "Widgets/Layout/SBorder.h"
+#include "Widgets/Layout/SBox.h"
+#include "Widgets/SBoxPanel.h"
+#include "Widgets/SNullWidget.h"
+#include "Widgets/SOverlay.h"
+#include "Widgets/SViewport.h"
+#include "Widgets/Text/STextBlock.h"
+#include "Rendering/SubstrateMaterialShared.h"
+#include "SGraphSubstrateMaterial.h"
+#include "MaterialShared.h"
+#include "RenderGraphBuilder.h"
+
+class FWidgetStyle;
+class SWidget;
+struct FSlateBrush;
+
+static const FName NAME_Pin_NotConnectable("Graph.Pin.Dummy");
+static const FSlateBrush* CacheImg_Pin_NotConnectable = nullptr;
+
+/**
+* Simple representation of the backbuffer that the preview canvas renders to
+* This class may only be accessed from the render thread
+*/
+class FSlateMaterialPreviewRenderTarget : public FRenderTarget
+{
+public:
+	/** FRenderTarget interface */
+	virtual FIntPoint GetSizeXY() const
+	{
+		return ClippingRect.Size();
+	}
+
+	/** Sets the texture that this target renders to */
+	void SetRenderTargetTexture(FRDGTexture* Texture)
+	{
+		RDGTexture = Texture;
+	}
+
+	/** Clears the render target texture */
+	void ClearRenderTargetTexture()
+	{
+		RDGTexture = nullptr;
+	}
+
+	FRDGTextureRef GetRenderTargetTexture(FRDGBuilder&) const override
+	{
+		return RDGTexture;
+	}
+
+	/** Sets the viewport rect for the render target */
+	void SetViewRect( const FIntRect& InViewRect ) 
+	{ 
+		ViewRect = InViewRect;
+	}
+
+	/** Gets the viewport rect for the render target */
+	const FIntRect& GetViewRect() const 
+	{
+		return ViewRect;
+	}
+
+	/** Sets the clipping rect for the render target */
+	void SetClippingRect( const FIntRect& InClippingRect ) 
+	{ 
+		ClippingRect = InClippingRect;
+	}
+
+	/** Gets the clipping rect for the render target */
+	const FIntRect& GetClippingRect() const 
+	{
+		return ClippingRect;
+	}
+
+private:
+	FRDGTexture* RDGTexture = nullptr;
+	FIntRect ViewRect;
+	FIntRect ClippingRect;
+};
+
+/*-----------------------------------------------------------------------------
+   FPreviewViewport
+-----------------------------------------------------------------------------*/
+
+FPreviewViewport::FPreviewViewport(class UMaterialGraphNode* InNode)
+	: MaterialNode(InNode)
+	, PreviewElement( new FPreviewElement )
+{
+	if (MaterialNode)
+	{
+		MaterialNode->InvalidatePreviewMaterialDelegate.BindRaw(this, &FPreviewViewport::UpdatePreviewNodeRenderProxy);
+	}
+}
+
+FPreviewViewport::~FPreviewViewport()
+{
+	if (MaterialNode)
+	{
+		MaterialNode->InvalidatePreviewMaterialDelegate.Unbind();
+	}
+	// Pass the preview element to the render thread so that it's deleted after it's shown for the last time
+	ENQUEUE_RENDER_COMMAND(SafeDeletePreviewElement)(
+		[PreviewElement = PreviewElement](FRHICommandListImmediate& RHICmdList) mutable
+		{
+			PreviewElement.Reset();
+		}
+	);
+}
+
+void FPreviewViewport::OnDrawViewport( const FGeometry& AllottedGeometry, const FSlateRect& MyCullingRect, class FSlateWindowElementList& OutDrawElements, int32 LayerId, const FWidgetStyle& InWidgetStyle, bool bParentEnabled )
+{
+	FSlateRect SlateCanvasRect = AllottedGeometry.GetLayoutBoundingRect();
+	FSlateRect ClippedCanvasRect = SlateCanvasRect.IntersectionWith(MyCullingRect);
+
+	FIntRect CanvasRect(
+		FMath::TruncToInt( FMath::Max(0.0f, SlateCanvasRect.Left) ),
+		FMath::TruncToInt( FMath::Max(0.0f, SlateCanvasRect.Top) ),
+		FMath::TruncToInt( FMath::Max(0.0f, SlateCanvasRect.Right) ), 
+		FMath::TruncToInt( FMath::Max(0.0f, SlateCanvasRect.Bottom) ) );
+
+	FIntRect ClippingRect(
+		FMath::TruncToInt( FMath::Max(0.0f, ClippedCanvasRect.Left) ),
+		FMath::TruncToInt( FMath::Max(0.0f, ClippedCanvasRect.Top) ),
+		FMath::TruncToInt( FMath::Max(0.0f, ClippedCanvasRect.Right) ), 
+		FMath::TruncToInt( FMath::Max(0.0f, ClippedCanvasRect.Bottom) ) );
+
+	bool bIsRealtime = MaterialNode->RealtimeDelegate.IsBound() ? MaterialNode->RealtimeDelegate.Execute() : false;
+
+	if (PreviewElement->BeginRenderingCanvas(CanvasRect, ClippingRect, MaterialNode, bIsRealtime))
+	{
+		// Draw above everything else
+		uint32 PreviewLayer = LayerId+1;
+		FSlateDrawElement::MakeCustom( OutDrawElements, PreviewLayer, PreviewElement );
+	}
+}
+
+FIntPoint FPreviewViewport::GetSize() const
+{
+	return FIntPoint(96,96);
+}
+
+void FPreviewViewport::UpdatePreviewNodeRenderProxy()
+{
+	if (PreviewElement.IsValid())
+	{
+		PreviewElement->UpdateExpressionPreview(MaterialNode);
+	}
+}
+
+/////////////////////////////////////////////////////
+// FPreviewElement
+
+FPreviewElement::FPreviewElement()
+	: RenderTarget(new FSlateMaterialPreviewRenderTarget)
+	, ExpressionPreview(nullptr)
+	, bIsRealtime(false)
+{
+}
+
+FPreviewElement::~FPreviewElement()
+{
+	delete RenderTarget;
+}
+
+bool FPreviewElement::BeginRenderingCanvas( const FIntRect& InCanvasRect, const FIntRect& InClippingRect, UMaterialGraphNode* InGraphNode, bool bInIsRealtime )
+{
+	if(InCanvasRect.Size().X > 0 && InCanvasRect.Size().Y > 0 && InClippingRect.Size().X > 0 && InClippingRect.Size().Y > 0 && InGraphNode != NULL)
+	{
+		/**
+		 * Struct to contain all info that needs to be passed to the render thread
+		 */
+		struct FPreviewRenderInfo
+		{
+			/** Size of the Canvas tile */
+			FIntRect CanvasRect;
+			/** How to clip the canvas tile */
+			FIntRect ClippingRect;
+			/** Render proxy for the expression preview */
+			FMaterialRenderProxy* RenderProxy;
+			/** Whether preview is using realtime values */
+			bool bIsRealtime;
+		};
+
+		FPreviewRenderInfo RenderInfo;
+		RenderInfo.CanvasRect = InCanvasRect;
+		RenderInfo.ClippingRect = InClippingRect;
+		RenderInfo.RenderProxy = InGraphNode->GetExpressionPreview();
+		RenderInfo.bIsRealtime = bInIsRealtime;
+
+		FPreviewElement* PreviewElement = this;
+		ENQUEUE_RENDER_COMMAND(BeginRenderingPreviewCanvas)(
+			[PreviewElement, RenderInfo](FRHICommandListImmediate& RHICmdList)
+			{
+				PreviewElement->RenderTarget->SetViewRect(RenderInfo.CanvasRect);
+				PreviewElement->RenderTarget->SetClippingRect(RenderInfo.ClippingRect);
+				PreviewElement->ExpressionPreview = RenderInfo.RenderProxy;
+				PreviewElement->bIsRealtime = RenderInfo.bIsRealtime;
+			}
+			);
+		return true;
+	}
+
+	return false;
+}
+
+void FPreviewElement::UpdateExpressionPreview(UMaterialGraphNode* MaterialNode)
+{
+	FPreviewElement* PreviewElement = this;
+	FMaterialRenderProxy* InRenderProxy = MaterialNode ? MaterialNode->GetExpressionPreview() : nullptr;
+	ENQUEUE_RENDER_COMMAND(UpdatePreviewNodeRenderProxy)(
+		[PreviewElement, InRenderProxy](FRHICommandListImmediate& RHICmdList)
+		{
+			PreviewElement->ExpressionPreview = InRenderProxy;
+		}
+	);
+}
+
+void FPreviewElement::Draw_RenderThread(FRDGBuilder& GraphBuilder, const FDrawPassInputs& Inputs)
+{
+	if(ExpressionPreview)
+	{
+		RenderTarget->SetRenderTargetTexture(Inputs.OutputTexture);
+		{
+			// Check realtime mode for whether to pass current time to canvas
+			double CurrentTime = bIsRealtime ? (FApp::GetCurrentTime() - GStartTime) : 0.0;
+			float DeltaTime = bIsRealtime ? static_cast<float>(FApp::GetDeltaTime()) : 0.0f;
+
+			FCanvas* Canvas = GraphBuilder.AllocObject<FCanvas>(RenderTarget, nullptr, FGameTime::CreateUndilated(CurrentTime, DeltaTime), GMaxRHIFeatureLevel);
+			{
+				Canvas->SetAllowedModes(0);
+				Canvas->SetRenderTargetRect(RenderTarget->GetViewRect());
+				Canvas->SetRenderTargetScissorRect(RenderTarget->GetClippingRect());
+
+				FCanvasTileItem TileItem(FVector2D::ZeroVector, ExpressionPreview, RenderTarget->GetSizeXY());
+				Canvas->DrawItem(TileItem);
+				Canvas->Flush_RenderThread(GraphBuilder, true);
+			}
+		}
+		RenderTarget->ClearRenderTargetTexture();
+	}
+}
+
+/////////////////////////////////////////////////////
+// SGraphNodeMaterialBase
+
+void SGraphNodeMaterialBase::Construct(const FArguments& InArgs, UMaterialGraphNode* InNode)
+{
+	this->TitleBorderMargin = FMargin(10.f, 5.f, 50.f, 3.f);
+	this->GraphNode = InNode;
+	this->MaterialNode = InNode;
+
+	this->SetCursor(EMouseCursor::CardinalCross);
+
+	this->UpdateGraphNode();
+}
+
+void SGraphNodeMaterialBase::CreatePinWidgets()
+{
+	// Create Pin widgets for each of the pins.
+	for( int32 PinIndex=0; PinIndex < GraphNode->Pins.Num(); ++PinIndex )
+	{
+		UEdGraphPin* CurPin = GraphNode->Pins[PinIndex];
+
+		bool bHideNoConnectionPins = false;
+
+		if (OwnerGraphPanelPtr.IsValid())
+		{
+			bHideNoConnectionPins = OwnerGraphPanelPtr.Pin()->GetPinVisibility() == SGraphEditor::Pin_HideNoConnection;
+		}
+
+		const bool bPinHasConections = CurPin->LinkedTo.Num() > 0;
+
+		bool bPinDesiresToBeHidden = CurPin->bHidden || (bHideNoConnectionPins && !bPinHasConections); 
+
+		UMaterialGraph* MaterialGraph = CastChecked<UMaterialGraph>(GraphNode->GetGraph());
+		if (MaterialNode && MaterialNode->MaterialExpression && MaterialGraph->MaterialFunction == nullptr && !MaterialGraph->MaterialInputs.IsEmpty())
+		{
+			bool bIsAMakeAttrNode = MaterialNode->MaterialExpression->IsA(UMaterialExpressionMakeMaterialAttributes::StaticClass());
+			bool bIsABreakAttrNode = MaterialNode->MaterialExpression->IsA(UMaterialExpressionBreakMaterialAttributes::StaticClass());
+		
+			if ((bIsABreakAttrNode && CurPin->Direction == EGPD_Output) || (bIsAMakeAttrNode && CurPin->Direction == EGPD_Input))
+			{
+				if (CurPin->PinType.PinCategory != UMaterialGraphSchema::PC_Exec)
+				{
+					bPinDesiresToBeHidden |= !MaterialGraph->MaterialInputs[CurPin->SourceIndex].IsVisiblePin(MaterialGraph->Material, true);
+				}
+			}
+		}
+
+		if (!bPinDesiresToBeHidden)
+		{
+			TSharedPtr<SGraphPin> NewPin = CreatePinWidget(CurPin);
+			check(NewPin.IsValid());
+
+			// Assign an custom icon to not connectible pins
+			if (CurPin->bNotConnectable)
+			{
+				if (!CacheImg_Pin_NotConnectable)
+				{
+					CacheImg_Pin_NotConnectable = FAppStyle::Get().GetBrush(NAME_Pin_NotConnectable);
+				}
+				NewPin->SetCustomPinIcon(CacheImg_Pin_NotConnectable, CacheImg_Pin_NotConnectable);
+			}
+
+			// Override pin color for Substrate node
+			if (Substrate::IsSubstrateEnabled())
+			{
+				FSubstrateWidget::GetPinColor(NewPin, MaterialNode);
+			}
+
+			this->AddPin(NewPin.ToSharedRef());
+		}
+	}
+}
+
+void SGraphNodeMaterialBase::MoveTo(const FVector2f& NewPosition, FNodeSet& NodeFilter, bool bMarkDirty)
+{
+	SGraphNode::MoveTo(NewPosition, NodeFilter, bMarkDirty);
+
+	MaterialNode->MaterialExpression->MaterialExpressionEditorX = MaterialNode->NodePosX;
+	MaterialNode->MaterialExpression->MaterialExpressionEditorY = MaterialNode->NodePosY;
+	MaterialNode->MaterialExpression->MarkPackageDirty();
+	MaterialNode->MaterialDirtyDelegate.ExecuteIfBound();
+}
+
+void SGraphNodeMaterialBase::AddPin( const TSharedRef<SGraphPin>& PinToAdd )
+{
+	PinToAdd->SetOwner( SharedThis(this) );
+
+	// Set visibility on advanced view pins
+	const UEdGraphPin* PinObj = PinToAdd->GetPinObj();
+	const bool bAdvancedParameter = (PinObj != nullptr) && PinObj->bAdvancedView;
+	if (bAdvancedParameter)
+	{
+		PinToAdd->SetVisibility(TAttribute<EVisibility>(PinToAdd, &SGraphPin::IsPinVisibleAsAdvanced));
+	}
+
+	if (PinToAdd->GetDirection() == EEdGraphPinDirection::EGPD_Input)
+	{
+		FMargin Padding = Settings->GetInputPinPadding();
+		Padding.Left *= 0.5f;
+		Padding.Right = 0.0f;
+
+		LeftNodeBox->AddSlot()
+		.AutoHeight()
+		.HAlign(HAlign_Left)
+		.VAlign(VAlign_Center)
+		.Padding(Padding)
+		[
+			PinToAdd
+		];
+		InputPins.Add(PinToAdd);
+	}
+	else // Direction == EEdGraphPinDirection::EGPD_Output
+	{
+		FMargin Padding = Settings->GetOutputPinPadding();
+		Padding.Left = 0.0f;
+		Padding.Right *= 0.5f;
+
+		RightNodeBox->AddSlot()
+		.AutoHeight()
+		.HAlign(HAlign_Right)
+		.VAlign(VAlign_Center)
+		.Padding(Padding)
+		[
+			PinToAdd
+		];
+		OutputPins.Add(PinToAdd);
+	}
+}
+
+void SGraphNodeMaterialBase::CreateBelowPinControls(TSharedPtr<SVerticalBox> MainBox)
+{
+	if (GraphNode && MainBox.IsValid())
+	{
+		// Count the number of visible input pins on the left
+		int32 LeftPinCount = 0;
+		if (GraphNode->AdvancedPinDisplay == ENodeAdvancedPins::Hidden)
+		{
+			// Advanced view pins are hidden so exclude them from the pin count
+			for (int32 i = 0; i < InputPins.Num(); ++i)
+			{
+				const UEdGraphPin* PinObj = InputPins[i]->GetPinObj();
+				if (!PinObj->bAdvancedView)
+				{
+					LeftPinCount++;
+				}
+			}
+		}
+		else
+		{
+			LeftPinCount = InputPins.Num();
+		}
+
+		int32 RightPinCount = OutputPins.Num();
+
+		const float NegativeHPad = FMath::Max<float>(-Settings->PaddingTowardsNodeEdge, 0.0f);
+		const float ExtraPad = 0.0f;
+
+		// Place preview widget based on where the least pins are
+		if ((LeftPinCount < RightPinCount) || (RightPinCount == 0))
+		{
+			LeftNodeBox->AddSlot()
+			.Padding(FMargin(NegativeHPad + ExtraPad, 0.0f, 0.0f, 0.0f))
+			.AutoHeight()
+			.HAlign(HAlign_Left)
+			[
+				CreatePreviewWidget()
+			];
+		}
+		else if (LeftPinCount > RightPinCount)
+		{
+			RightNodeBox->AddSlot()
+			.Padding(FMargin(NegativeHPad + ExtraPad, 0.0f, 0.0f, 0.0f))
+			.AutoHeight()
+			.HAlign(HAlign_Right)
+			[
+				CreatePreviewWidget()
+			];
+		}
+		else
+		{
+			MainBox->AddSlot()
+			.Padding(Settings->GetNonPinNodeBodyPadding())
+			.AutoHeight()
+			[
+				SNew(SHorizontalBox)
+				+SHorizontalBox::Slot()
+				.AutoWidth()
+				[
+					CreatePreviewWidget()
+				]
+			];
+		}
+	}
+
+	// Preview of Substrate nodes topology
+	if (Substrate::IsSubstrateEnabled() && MaterialNode)
+	{
+		TArray<FGuid> Guids;
+		if (MaterialNode->MaterialExpression->IsA(UMaterialExpressionSubstrateBSDF::StaticClass()))
+		{
+			const UMaterialExpression* SubstrateExpression = (const UMaterialExpression*)MaterialNode->MaterialExpression;
+			Guids.Add(SubstrateExpression->MaterialExpressionGuid);
+		}
+		else if (MaterialNode->MaterialExpression->IsA(UMaterialExpressionMaterialFunctionCall::StaticClass()))
+		{
+			UMaterialExpressionMaterialFunctionCall* FunctionCall = (UMaterialExpressionMaterialFunctionCall*)MaterialNode->MaterialExpression;
+			const uint32 OutputCount = FunctionCall->FunctionOutputs.Num();
+			for (uint32 OutputIndex = 0; OutputIndex < OutputCount; ++OutputIndex)
+			{
+				if (FunctionCall->IsResultSubstrateMaterial(OutputIndex))
+				{
+					FSubstrateMaterialInfo SubstrateMaterialInfo(true/*bGatherGuids*/);
+					FunctionCall->GatherSubstrateMaterialInfo(SubstrateMaterialInfo, OutputIndex);
+					Guids = SubstrateMaterialInfo.GetGuids();
+					break;
+				}
+			}
+		}
+
+		if (Guids.Num() > 0)
+		{
+			if (const UMaterialExpression* SubstrateExpression = (const UMaterialExpression*)MaterialNode->MaterialExpression)
+			{		
+				if (UMaterial* MaterialForStats = SubstrateExpression->Material)
+				{
+					if (const FMaterialResource* MaterialResource = MaterialForStats->GetMaterialResource(GMaxRHIFeatureLevel))
+					{
+						if (FMaterialShaderMap* ShaderMap = MaterialResource->GetGameThreadShaderMap())
+						{
+							const FSubstrateMaterialCompilationOutput& CompilationOutput = ShaderMap->GetSubstrateMaterialCompilationOutput();
+							MainBox->AddSlot()
+							.Padding(Settings->GetNonPinNodeBodyPadding())
+							.AutoHeight()
+							[
+								SNew(SHorizontalBox)
+								+SHorizontalBox::Slot()
+								.VAlign(VAlign_Center)
+								.HAlign(HAlign_Center)
+								[							
+									FSubstrateWidget::ProcessOperator(CompilationOutput, Guids)
+								]
+							];
+						}
+					}
+	
+				}
+			}
+		}
+	}
+}
+
+void SGraphNodeMaterialBase::SetDefaultTitleAreaWidget(TSharedRef<SOverlay> DefaultTitleAreaWidget)
+{
+	if (!MaterialNode->MaterialExpression->bHidePreviewWindow)
+	{
+		DefaultTitleAreaWidget->AddSlot()
+		.HAlign(HAlign_Right)
+		.VAlign(VAlign_Top)
+		.Padding(FMargin(5))
+		[
+			SNew(SHorizontalBox)
+			+ SHorizontalBox::Slot()
+			.AutoWidth()
+			[
+				SNew(SButton)
+				.ButtonStyle(FAppStyle::Get(), "HoverHintOnly")
+				.OnClicked(this, &SGraphNodeMaterialBase::OnTogglePreviewClicked)
+				.Cursor(EMouseCursor::Default)
+				[
+					SNew(SHorizontalBox)
+					+SHorizontalBox::Slot()
+					.VAlign(VAlign_Center)
+					.HAlign(HAlign_Center)
+					[
+						SNew(SImage)
+						.Image(FAppStyle::GetBrush(TEXT("Icons.Preview")))
+					]
+				]
+			]
+			+ SHorizontalBox::Slot()
+			.AutoWidth()
+			[
+				SNew(SCheckBox)
+				.OnCheckStateChanged( this, &SGraphNodeMaterialBase::OnExpressionPreviewChanged )
+				.IsChecked( IsExpressionPreviewChecked() )
+				.Cursor(EMouseCursor::Default)
+				.Style(FAppStyle::Get(), "Graph.Node.AdvancedView")
+				[
+					SNew(SHorizontalBox)
+					+SHorizontalBox::Slot()
+					.VAlign(VAlign_Center)
+					.HAlign(HAlign_Center)
+					[
+						SNew(SImage)
+						.Image(GetExpressionPreviewArrow())
+					]
+				]
+			]
+		];
+	}
+}
+
+TSharedRef<SWidget> SGraphNodeMaterialBase::CreateNodeContentArea()
+{
+	// NODE CONTENT AREA
+	return SNew(SBorder)
+		.BorderImage( FAppStyle::GetBrush("NoBorder") )
+		.HAlign(HAlign_Fill)
+		.VAlign(VAlign_Fill)
+		[
+			SNew(SHorizontalBox)
+			+SHorizontalBox::Slot()
+			.HAlign(HAlign_Left)
+			.FillWidth(1.0f)
+			[
+				// LEFT
+				SAssignNew(LeftNodeBox, SVerticalBox)
+			]
+			+SHorizontalBox::Slot()
+			.AutoWidth()
+			.HAlign(HAlign_Right)
+			[
+				// RIGHT
+				SAssignNew(RightNodeBox, SVerticalBox)
+			]
+		];
+}
+
+void SGraphNodeMaterialBase::OnAdvancedViewChanged(const ECheckBoxState NewCheckedState)
+{
+	SGraphNode::OnAdvancedViewChanged(NewCheckedState);
+
+	// Update the graph node so that the preview is recreated to update its position
+	UpdateGraphNode();
+}
+
+TSharedRef<SWidget> SGraphNodeMaterialBase::CreatePreviewWidget()
+{
+	PreviewViewport.Reset();
+
+	// if this node should currently show a preview
+	if (!MaterialNode->MaterialExpression->bHidePreviewWindow && !MaterialNode->MaterialExpression->bCollapsed)
+	{
+		const float ExpressionPreviewSize = 106.0f;
+		const float CentralPadding = 5.0f;
+
+		TSharedPtr<SViewport> ViewportWidget = 
+			SNew( SViewport )
+			.RenderDirectlyToWindow(true)
+			.EnableGammaCorrection(false);
+
+		PreviewViewport = MakeShareable(new FPreviewViewport(MaterialNode));
+
+		// The viewport widget needs an interface so it knows what should render
+		ViewportWidget->SetViewportInterface( PreviewViewport.ToSharedRef() );
+
+		return SNew(SBox)
+			.WidthOverride(ExpressionPreviewSize)
+			.HeightOverride(ExpressionPreviewSize)
+			.MaxAspectRatio(1.0f)
+			.MaxDesiredHeight(ExpressionPreviewSize)
+			.Visibility(ExpressionPreviewVisibility())
+			[
+				SNew(SBorder)
+				.Padding(CentralPadding)
+				.BorderImage( FAppStyle::GetBrush("NoBorder") )
+				[
+					SNew(SOverlay)
+					+ SOverlay::Slot()
+					[
+						ViewportWidget.ToSharedRef()
+					]
+					+ SOverlay::Slot()
+					.HAlign(HAlign_Right)
+					.VAlign(VAlign_Bottom)
+					.Padding(0.0f,0.0f,8.0f,0.0f)
+					[
+						SNew(STextBlock)
+						.ShadowColorAndOpacity(FLinearColor::Black)
+						.ShadowOffset(FVector2D(1.0f, 1.0f))
+						.Text(this, &SGraphNodeMaterialBase::ExpressionPreviewOverlayText)
+					]
+				]
+			];
+	}
+
+	return SNullWidget::NullWidget;
+}
+
+FText SGraphNodeMaterialBase::ExpressionPreviewOverlayText() const
+{
+	UMaterialExpression* MaterialExpression = MaterialNode->MaterialExpression;
+	return MaterialNode->MaterialExpression->GetPreviewOverlayText();
+}
+
+EVisibility SGraphNodeMaterialBase::ExpressionPreviewVisibility() const
+{
+	UMaterialExpression* MaterialExpression = MaterialNode->MaterialExpression;
+	const bool bShowPreview = !MaterialExpression->bHidePreviewWindow && !MaterialExpression->bCollapsed;
+	return bShowPreview ? EVisibility::Visible : EVisibility::Collapsed;
+}
+
+FReply SGraphNodeMaterialBase::OnTogglePreviewClicked()
+{
+	MaterialNode->ToggleNodePreview();
+	return FReply::Handled();
+}
+
+void SGraphNodeMaterialBase::OnExpressionPreviewChanged( const ECheckBoxState NewCheckedState )
+{
+	UMaterialExpression* MaterialExpression = MaterialNode->MaterialExpression;
+	const bool bCollapsed = (NewCheckedState != ECheckBoxState::Checked);
+	if (MaterialExpression->bCollapsed != bCollapsed)
+	{
+		UMaterialGraph* MaterialGraph = CastChecked<UMaterialGraph>(MaterialNode->GetGraph());
+		MaterialGraph->ToggleCollapsedDelegate.ExecuteIfBound(MaterialExpression);
+
+		// Update the graph node so that preview viewport is created
+		UpdateGraphNode();
+	}
+}
+
+ECheckBoxState SGraphNodeMaterialBase::IsExpressionPreviewChecked() const
+{
+	return MaterialNode->MaterialExpression->bCollapsed ? ECheckBoxState::Unchecked : ECheckBoxState::Checked;
+}
+
+const FSlateBrush* SGraphNodeMaterialBase::GetExpressionPreviewArrow() const
+{
+	return FAppStyle::GetBrush(MaterialNode->MaterialExpression->bCollapsed ? TEXT("Icons.ChevronDown") : TEXT("Icons.ChevronUp"));
+}
+
+void SGraphNodeMaterialBase::PopulateMetaTag(FGraphNodeMetaData* TagMeta) const
+{
+	if (GraphNode != nullptr)
+	{
+		UMaterialGraph* OuterGraph = MaterialNode->GetTypedOuter<UMaterialGraph>();
+		if ((OuterGraph != nullptr) && (MaterialNode->MaterialExpression != nullptr) )
+		{
+			TagMeta->OuterName = OuterGraph->OriginalMaterialFullName;
+			TagMeta->GUID = MaterialNode->MaterialExpression->MaterialExpressionGuid;
+			TagMeta->Tag = FName(*FString::Printf(TEXT("MaterialExprNode_%s_%s"), *TagMeta->OuterName, *TagMeta->GUID.ToString()));
+		}
+		TagMeta->FriendlyName = FString::Printf(TEXT("%s expression node in %s"), *GraphNode->GetNodeTitle(ENodeTitleType::FullTitle).ToString(), *TagMeta->OuterName);		
+	}
+}
diff --git a/Engine/Source/Editor/MaterialEditor/Private/MaterialNodes/SGraphNodeMaterialComment.cpp b/Engine/Source/Editor/MaterialEditor/Private/MaterialNodes/SGraphNodeMaterialComment.cpp
new file mode 100644
--- /dev/null
+++ b/Engine/Source/Editor/MaterialEditor/Private/MaterialNodes/SGraphNodeMaterialComment.cpp
@@ -0,0 +1,30 @@
+// Copyright Epic Games, Inc. All Rights Reserved.
+
+#include "MaterialNodes/SGraphNodeMaterialComment.h"
+
+#include "MaterialGraph/MaterialGraph.h"
+#include "MaterialGraph/MaterialGraphNode_Comment.h"
+#include "Materials/MaterialExpressionComment.h"
+#include "UObject/ObjectPtr.h"
+
+void SGraphNodeMaterialComment::Construct(const FArguments& InArgs, class UMaterialGraphNode_Comment* InNode)
+{
+	SGraphNodeComment::Construct(SGraphNodeComment::FArguments(), InNode);
+
+	this->CommentNode = InNode;
+}
+
+void SGraphNodeMaterialComment::MoveTo(const FVector2f& NewPosition, FNodeSet& NodeFilter, bool bMarkDirty)
+{
+	if( !NodeFilter.Find( SharedThis( this ) ))
+	{
+		SGraphNodeComment::MoveTo(NewPosition, NodeFilter, bMarkDirty);
+		if (CommentNode && CommentNode->MaterialExpressionComment)
+		{
+			CommentNode->MaterialExpressionComment->MaterialExpressionEditorX = CommentNode->NodePosX;
+			CommentNode->MaterialExpressionComment->MaterialExpressionEditorY = CommentNode->NodePosY;
+			CommentNode->MaterialExpressionComment->MarkPackageDirty();
+			CommentNode->MaterialDirtyDelegate.ExecuteIfBound();
+		}
+	}
+}
diff --git a/Engine/Source/Editor/MaterialEditor/Private/MaterialNodes/SGraphNodeMaterialComposite.cpp b/Engine/Source/Editor/MaterialEditor/Private/MaterialNodes/SGraphNodeMaterialComposite.cpp
new file mode 100644
--- /dev/null
+++ b/Engine/Source/Editor/MaterialEditor/Private/MaterialNodes/SGraphNodeMaterialComposite.cpp
@@ -0,0 +1,316 @@
+// Copyright Epic Games, Inc. All Rights Reserved.
+
+#include "MaterialNodes/SGraphNodeMaterialComposite.h"
+
+#include "Containers/Array.h"
+#include "Delegates/Delegate.h"
+#include "EdGraph/EdGraph.h"
+#include "EdGraph/EdGraphNode.h"
+#include "Fonts/SlateFontInfo.h"
+#include "Framework/Application/SlateApplication.h"
+#include "GenericPlatform/GenericApplication.h"
+#include "GenericPlatform/ICursor.h"
+#include "GraphEditorSettings.h"
+#include "IDocumentation.h"
+#include "IDocumentationPage.h"
+#include "Internationalization/Internationalization.h"
+#include "Layout/Margin.h"
+#include "MaterialGraph/MaterialGraph.h"
+#include "MaterialGraph/MaterialGraphNode_Composite.h"
+#include "Materials/MaterialExpression.h"
+#include "Misc/Attribute.h"
+#include "Misc/Optional.h"
+#include "SCommentBubble.h"
+#include "SGraphNode.h"
+#include "SGraphPreviewer.h"
+#include "SlotBase.h"
+#include "Styling/AppStyle.h"
+#include "Styling/CoreStyle.h"
+#include "Styling/SlateColor.h"
+#include "Types/SlateEnums.h"
+#include "UObject/ObjectPtr.h"
+#include "UObject/UObjectGlobals.h"
+#include "Widgets/Images/SImage.h"
+#include "Widgets/Layout/SBorder.h"
+#include "Widgets/Layout/SSpacer.h"
+#include "Widgets/Notifications/SErrorText.h"
+#include "Widgets/SBoxPanel.h"
+#include "Widgets/SOverlay.h"
+#include "Widgets/SToolTip.h"
+#include "Widgets/Text/SInlineEditableTextBlock.h"
+#include "Widgets/Text/STextBlock.h"
+
+class SWidget;
+
+void SGraphNodeMaterialComposite::Construct(const FArguments& InArgs, class UMaterialGraphNode_Composite* InNode)
+{
+	this->GraphNode = InNode;
+	this->CompositeNode = InNode;
+	
+	//@TODO: Add previews to composites, figure out multiple output UX (scroll?).
+	//this->MaterialNode = CompositeNode;
+
+	this->SetCursor(EMouseCursor::CardinalCross);
+
+	this->UpdateGraphNode();
+}
+
+void SGraphNodeMaterialComposite::MoveTo(const FVector2f& NewPosition, FNodeSet& NodeFilter, bool bMarkDirty)
+{
+	SGraphNode::MoveTo(NewPosition, NodeFilter, bMarkDirty);
+
+	CompositeNode->MaterialExpression->MaterialExpressionEditorX = CompositeNode->NodePosX;
+	CompositeNode->MaterialExpression->MaterialExpressionEditorY = CompositeNode->NodePosY;
+	CompositeNode->MaterialExpression->MarkPackageDirty();
+	CompositeNode->MaterialDirtyDelegate.ExecuteIfBound();
+}
+
+void SGraphNodeMaterialComposite::UpdateGraphNode()
+{
+	InputPins.Empty();
+	OutputPins.Empty();
+	
+	// Reset variables that are going to be exposed, in case we are refreshing an already setup node.
+	RightNodeBox.Reset();
+	LeftNodeBox.Reset();
+
+	SetupErrorReporting();
+	TSharedPtr<SNodeTitle> NodeTitle = SNew(SNodeTitle, GraphNode);
+
+	//
+	//             ______________________
+	//            |      TITLE AREA      |
+	//            +-------+------+-------+
+	//            | (>) L |      | R (>) |
+	//            | (>) E |      | I (>) |
+	//            | (>) F |      | G (>) |
+	//            | (>) T |      | H (>) |
+	//            |       |      | T (>) |
+	//            |_______|______|_______|
+	//
+	this->ContentScale.Bind( this, &SGraphNode::GetContentScale );
+	this->GetOrAddSlot( ENodeZone::Center )
+		.HAlign(HAlign_Center)
+		.VAlign(VAlign_Center)
+		[
+			SNew(SBorder)
+			.BorderImage( FAppStyle::GetBrush( "Graph.CollapsedNode.Body" ) )
+			.Padding(0)
+			[
+				SNew(SOverlay)
+				+SOverlay::Slot()
+				[
+					SNew(SImage)
+					.Image( FAppStyle::GetBrush("Graph.CollapsedNode.BodyColorSpill") )
+					.ColorAndOpacity( this, &SGraphNode::GetNodeTitleColor )
+				]
+				+SOverlay::Slot()
+				[
+					SNew(SVerticalBox)
+					+SVerticalBox::Slot()
+					.AutoHeight()
+					.HAlign(HAlign_Fill)
+					.VAlign(VAlign_Top)
+					[
+						SNew(SOverlay)
+						+SOverlay::Slot()
+						.HAlign(HAlign_Left)
+						.VAlign(VAlign_Center)
+						[
+							SNew(SBorder)
+							.BorderImage( FAppStyle::GetBrush("NoBorder") )
+							.Padding( FMargin(10,5,30,3) )
+							[
+								SNew(SVerticalBox)
+								+SVerticalBox::Slot()
+								.AutoHeight()
+								.HAlign(HAlign_Fill)
+								.VAlign(VAlign_Top)
+								[
+									SNew(SVerticalBox)
+									+SVerticalBox::Slot()
+										.AutoHeight()
+									[
+										SAssignNew(InlineEditableText, SInlineEditableTextBlock)
+										.Style( FAppStyle::Get(), "Graph.Node.NodeTitleInlineEditableText" )
+										.Text( NodeTitle.Get(), &SNodeTitle::GetHeadTitle )
+										.OnVerifyTextChanged(this, &SGraphNodeMaterialComposite::OnVerifyNameTextChanged)
+										.OnTextCommitted(this, &SGraphNodeMaterialComposite::OnNameTextCommited)
+										.IsReadOnly( this, &SGraphNodeMaterialComposite::IsNameReadOnly )
+										.IsSelected(this, &SGraphNodeMaterialComposite::IsSelectedExclusively)
+									]
+									+SVerticalBox::Slot()
+										.AutoHeight()
+									[
+										NodeTitle.ToSharedRef()
+									]
+								]
+								+SVerticalBox::Slot()
+								.AutoHeight()
+								.Padding(1.0f)
+								[
+									ErrorReporting->AsWidget()
+								]
+							]
+						]
+					]
+					+SVerticalBox::Slot()
+					.AutoHeight()
+					.HAlign(HAlign_Fill)
+					.VAlign(VAlign_Top)
+					[
+						CreateNodeBody()
+					]
+				]
+			]
+		];
+
+	// Create comment bubble
+	TSharedPtr<SCommentBubble> CommentBubble;
+	const FSlateColor CommentColor = GetDefault<UGraphEditorSettings>()->DefaultCommentNodeTitleColor;
+
+	SAssignNew( CommentBubble, SCommentBubble )
+	.GraphNode( GraphNode )
+	.Text( this, &SGraphNode::GetNodeComment )
+	.OnTextCommitted( this, &SGraphNode::OnCommentTextCommitted )
+	.ColorAndOpacity( CommentColor )
+	.AllowPinning( true )
+	.EnableTitleBarBubble( true )
+	.EnableBubbleCtrls( true )
+	.GraphLOD( this, &SGraphNode::GetCurrentLOD )
+	.IsGraphNodeHovered( this, &SGraphNode::IsHovered );
+
+	GetOrAddSlot( ENodeZone::TopCenter )
+	.SlotOffset2f( TAttribute<FVector2f>( CommentBubble.Get(), &SCommentBubble::GetOffset2f ))
+	.SlotSize2f( TAttribute<FVector2f>( CommentBubble.Get(), &SCommentBubble::GetSize2f ))
+	.AllowScaling( TAttribute<bool>( CommentBubble.Get(), &SCommentBubble::IsScalingAllowed ))
+	.VAlign( VAlign_Top )
+	[
+		CommentBubble.ToSharedRef()
+	];
+
+	CreatePinWidgets();
+}
+
+UEdGraph* SGraphNodeMaterialComposite::GetInnerGraph() const
+{
+	return CompositeNode->BoundGraph;
+}
+
+TSharedPtr<SToolTip> SGraphNodeMaterialComposite::GetComplexTooltip()
+{
+	if (UEdGraph* BoundGraph = GetInnerGraph())
+	{
+		struct LocalUtils
+		{
+			static bool IsInteractive()
+			{
+				const FModifierKeysState ModifierKeys = FSlateApplication::Get().GetModifierKeys();
+				return ( ModifierKeys.IsAltDown() && ModifierKeys.IsControlDown() );
+			}
+		};
+
+		TSharedPtr<SToolTip> FinalToolTip = nullptr;
+		TSharedPtr<SVerticalBox> Container = nullptr;
+		SAssignNew(FinalToolTip, SToolTip)
+		.IsInteractive_Static(&LocalUtils::IsInteractive)
+		[
+			SAssignNew(Container, SVerticalBox)
+			+SVerticalBox::Slot()
+			.AutoHeight()
+			[
+				SNew( STextBlock )
+				.Text(this, &SGraphNodeMaterialComposite::GetTooltipTextForNode)
+				.Font(FCoreStyle::GetDefaultFontStyle("Regular", 8))
+				.WrapTextAt(160.0f)
+			]
+			+SVerticalBox::Slot()
+			.AutoHeight()
+			[
+				// Create preview for the tooltip, make sure to disable state overlays to prevent
+				// PIE and read-only borders obscuring the graph
+				SNew(SGraphPreviewer, BoundGraph)
+				.CornerOverlayText(this, &SGraphNodeMaterialComposite::GetPreviewCornerText)
+				.ShowGraphStateOverlay(false)
+			]
+		];
+
+		// Check to see whether this node has a documentation excerpt. If it does, create a doc box for the tooltip
+		TSharedRef<IDocumentationPage> DocPage = IDocumentation::Get()->GetPage(GraphNode->GetDocumentationLink(), nullptr);
+		if(DocPage->HasExcerpt(GraphNode->GetDocumentationExcerptName()))
+		{
+			Container->AddSlot()
+			.AutoHeight()
+			.Padding(FMargin( 0.0f, 5.0f ))
+			[
+				IDocumentation::Get()->CreateToolTip(FText::FromString("Documentation"), nullptr, GraphNode->GetDocumentationLink(), GraphNode->GetDocumentationExcerptName())
+			];
+		}
+
+		return FinalToolTip;
+	}
+	else
+	{
+		return SNew(SToolTip)
+			[
+				SNew(SVerticalBox)
+				+SVerticalBox::Slot()
+				.AutoHeight()
+				[
+					SNew( STextBlock )
+					.Text(NSLOCTEXT("CompositeNode", "CompositeNodeInvalidGraphMessage", "ERROR: Invalid Graph"))
+					.Font(FCoreStyle::GetDefaultFontStyle("Regular", 8))
+					.WrapTextAt(160.0f)
+				]
+			];
+	}
+
+}
+
+FText SGraphNodeMaterialComposite::GetPreviewCornerText() const
+{
+	UEdGraph* BoundGraph = GetInnerGraph();
+	return FText::FromString(BoundGraph->GetName());
+}
+
+FText SGraphNodeMaterialComposite::GetTooltipTextForNode() const
+{
+	return GraphNode->GetTooltipText();
+}
+
+TSharedRef<SWidget> SGraphNodeMaterialComposite::CreateNodeBody()
+{
+	if( GraphNode && GraphNode->Pins.Num() > 0 )
+	{
+		// Create the input and output pin areas if there are pins
+		return SNew(SBorder)
+			.BorderImage( FAppStyle::GetBrush("NoBorder") )
+			.HAlign(HAlign_Fill)
+			.VAlign(VAlign_Fill)
+			.Padding( FMargin(0,3) )
+			[
+				SNew(SHorizontalBox)
+				+SHorizontalBox::Slot()
+				.HAlign(HAlign_Left)
+				.FillWidth(1.0f)
+				[
+					// LEFT
+					SAssignNew(LeftNodeBox, SVerticalBox)
+				]
+				+SHorizontalBox::Slot()
+				.AutoWidth()
+				.HAlign(HAlign_Right)
+				[
+					// RIGHT
+					SAssignNew(RightNodeBox, SVerticalBox)
+				]
+			];
+	}
+	else
+	{
+		// Create a spacer so the node has some body to it
+		return SNew(SSpacer)
+			.Size(FVector2D(100.f, 50.f));
+	}
+}
+
diff --git a/Engine/Source/Editor/MaterialEditor/Private/MaterialNodes/SGraphNodeMaterialConvert.cpp b/Engine/Source/Editor/MaterialEditor/Private/MaterialNodes/SGraphNodeMaterialConvert.cpp
new file mode 100644
--- /dev/null
+++ b/Engine/Source/Editor/MaterialEditor/Private/MaterialNodes/SGraphNodeMaterialConvert.cpp
@@ -0,0 +1,910 @@
+// Copyright Epic Games, Inc. All Rights Reserved.
+
+#include "MaterialNodes/SGraphNodeMaterialConvert.h"
+
+#include "ConnectionDrawingPolicy.h"
+#include "Framework/MultiBox/MultiBoxBuilder.h"
+#include "MaterialGraph/MaterialGraphNode.h"
+#include "Materials/MaterialExpressionConvert.h"
+#include "Misc/Attribute.h"
+#include "ScopedTransaction.h"
+#include "Styling/AppStyle.h"
+#include "Widgets/Images/SImage.h"
+#include "Widgets/Input/SButton.h"
+#include "Widgets/Input/SNumericEntryBox.h"
+#include "Widgets/Layout/SSpacer.h"
+#include "Widgets/SBoxPanel.h"
+
+#define LOCTEXT_NAMESPACE "GraphNodeMaterialConvert"
+
+////////////////////////////////////////////////////////////////////////////////////////////////////////////////////////////////////////////
+// FConvertDragDropOp: Drag and Drop Operation used to form connections within the convert node
+////////////////////////////////////////////////////////////////////////////////////////////////////////////////////////////////////////////
+
+FConvertDragDropOp::FConvertDragDropOp(TSharedPtr<SConvertInnerPin> InSourcePin)
+: SourcePin(InSourcePin)
+{
+	Construct();
+}
+
+void FConvertDragDropOp::OnDrop(bool bDropWasHandled, const FPointerEvent& MouseEvent)
+{
+	FDragDropOperation::OnDrop(bDropWasHandled, MouseEvent);
+
+	if (SourcePin)
+	{
+		SourcePin->CancelDragDrop();
+	}
+}
+
+void FConvertDragDropOp::OnDragged(const class FDragDropEvent& DragDropEvent)
+{
+	FDragDropOperation::OnDragged(DragDropEvent);
+
+	ScreenPosition = DragDropEvent.GetScreenSpacePosition();
+}
+
+TSharedPtr<SConvertInnerPin> FConvertDragDropOp::GetSourcePin() const
+{
+	return SourcePin;
+}
+
+UE::Slate::FDeprecateVector2DResult FConvertDragDropOp::GetScreenPosition() const
+{ 
+	return ScreenPosition;
+}
+
+////////////////////////////////////////////////////////////////////////////////////////////////////////////////////////////////////////////
+// SMaterialExpressionConvertGraphPin: The Outer pins that form connections to other material graph nodes
+////////////////////////////////////////////////////////////////////////////////////////////////////////////////////////////////////////////
+
+void SMaterialExpressionConvertGraphPin::Construct(const FArguments& InArgs, UEdGraphPin* InPin)
+{
+	SGraphPin::Construct(SGraphPin::FArguments(), InPin);
+
+	// Setup custom hover logic to only treat this pin as hovered when you hover over the main pin image
+	this->SetHover(TAttribute<bool>::CreateSP(this, &SMaterialExpressionConvertGraphPin::IsHoveredOverPrimaryPin));
+}
+
+void SMaterialExpressionConvertGraphPin::CreateInnerPins(TSharedRef<SGraphNodeMaterialConvert> InOwningGraphNodeWidget)
+{
+	UMaterialGraphNode* MaterialNode = InOwningGraphNodeWidget->GetMaterialGraphNode();
+	if (!MaterialNode)
+	{
+		return;
+	}
+
+	UMaterialExpressionConvert* ConvertExpression = Cast<UMaterialExpressionConvert>(MaterialNode->MaterialExpression);
+	if (!ConvertExpression)
+	{
+		return;
+	}
+
+	UEdGraphPin* EdGraphPin = GetPinObj();
+	if (!EdGraphPin)
+	{
+		return;
+	}
+		
+	if (TSharedPtr<SHorizontalBox> PinnedRowWidget = GetFullPinHorizontalRowWidget().Pin())
+	{
+		const bool bIsInputPin = EdGraphPin->Direction == EGPD_Input;
+		const int32 PinIndex = EdGraphPin->SourceIndex;
+				
+		const EMaterialExpressionConvertType ConvertType 	= bIsInputPin
+			? ConvertExpression->ConvertInputs[PinIndex].Type
+			: ConvertExpression->ConvertOutputs[PinIndex].Type;
+		const int32 PinComponentCount = MaterialExpressionConvertType::GetComponentCount(ConvertType);
+				
+		InnerPins.Reset(PinComponentCount);
+		TSharedPtr<SVerticalBox> InnerPinsVerticalBox = SNew(SVerticalBox);
+		for (int32 ComponentIndex = 0; ComponentIndex < PinComponentCount; ++ComponentIndex)
+		{
+			TSharedRef<SConvertInnerPin> NewInnerPin = SNew(SConvertInnerPin, 
+				InOwningGraphNodeWidget,
+				SharedThis(this),
+				bIsInputPin, 
+				PinIndex, 
+				ComponentIndex
+			);
+			InnerPinsVerticalBox->AddSlot()
+			.HAlign(bIsInputPin ? HAlign_Right : HAlign_Left)
+			[
+				NewInnerPin
+			];
+			InnerPins.Add(NewInnerPin);
+		}
+
+		// Insert as last element (INDEX_NONE) for input pins, first element for output pins
+		PinnedRowWidget->InsertSlot(bIsInputPin ? INDEX_NONE : 0)
+		.HAlign(bIsInputPin ? HAlign_Right : HAlign_Left)
+		.VAlign(VAlign_Center)
+		.AutoWidth()
+		[
+			InnerPinsVerticalBox.ToSharedRef()
+		];
+	}
+}
+
+bool SMaterialExpressionConvertGraphPin::IsHoveredOverPrimaryPin() const
+{
+	TSharedPtr<SWidget> PinImageWidget = this->GetPinImageWidget();
+	return PinImageWidget ? PinImageWidget->IsHovered() : false;
+}
+
+////////////////////////////////////////////////////////////////////////////////////////////////////////////////////////////////////////////
+// SConvertInnerPin: The inner pins used to route values with the convert node
+////////////////////////////////////////////////////////////////////////////////////////////////////////////////////////////////////////////
+
+void SConvertInnerPin::Construct(
+	const FArguments& InArgs,
+	TSharedPtr<SGraphNodeMaterialConvert> InOwningNode,
+	TSharedPtr<SMaterialExpressionConvertGraphPin> InOwningPin,
+	bool bInIsInputPin, 
+	int32 InPinIndex, 
+	int32 InComponentIndex
+)
+{
+	WeakOwningNode = InOwningNode;
+	WeakOwningPin = InOwningPin;
+	bIsInputPin = bInIsInputPin;
+	PinIndex = InPinIndex;
+	ComponentIndex = InComponentIndex;
+
+	// Construct Horizontal Box and add the PinImage to it
+	TSharedPtr<SHorizontalBox> HorizontalBox = 
+	SNew(SHorizontalBox)
+	+SHorizontalBox::Slot()
+	.AutoWidth()
+	.HAlign(HAlign_Left)
+	.VAlign(VAlign_Center)
+	.Padding(2.5f, 0.0f)
+	[
+		SAssignNew(PinImage, SImage)
+		.Image(this, &SConvertInnerPin::GetPinBrush)
+	];
+
+	// Add default value entry to either the front or back of the hbox depending on if we're an input pin
+	HorizontalBox->InsertSlot(bIsInputPin ? 0 : INDEX_NONE)
+	.AutoWidth()
+	.HAlign(HAlign_Left)
+	.VAlign(VAlign_Center)
+	[
+		SNew(SHorizontalBox)
+		+SHorizontalBox::Slot()
+		.AutoWidth()
+		.HAlign(HAlign_Left)
+		.VAlign(VAlign_Center)
+		.Padding(2.5f, 0.0f)
+		[
+			SNew(STextBlock)
+			.TextStyle(FAppStyle::Get(), "MonospacedText")
+			.Visibility(this, &SConvertInnerPin::GetPinNameVisibility)
+			.Text(this, &SConvertInnerPin::GetPinName)
+		]
+		+SHorizontalBox::Slot()
+		.AutoWidth()
+		.HAlign(HAlign_Left)
+		.VAlign(VAlign_Center)
+		.Padding(2.5f, 0.0f)
+		[
+			SNew(SNumericEntryBox<float>)
+			.EditableTextBoxStyle(FAppStyle::Get(), "Graph.EditableTextBox")
+			.BorderForegroundColor(FSlateColor::UseForeground())
+			.Visibility(this, &SConvertInnerPin::GetDefaultValueVisibility)
+			.Value(this, &SConvertInnerPin::GetDefaultValue)
+			.OnValueCommitted(this, &SConvertInnerPin::SetDefaultValue)
+		]
+	];
+
+	// Add inner node connection space to either the front or back of the hbox depending on if we're an input pin
+	HorizontalBox->InsertSlot(bIsInputPin ? INDEX_NONE : 0)
+	.AutoWidth()
+	.HAlign(HAlign_Left)
+	.VAlign(VAlign_Center)
+	[
+		SNew(SSpacer)
+		.Size(FVector2D(30, 0))
+	];
+
+	ChildSlot
+	[
+		HorizontalBox.ToSharedRef()
+	];
+}
+
+int32 SConvertInnerPin::OnPaint(
+	const FPaintArgs& Args,
+	const FGeometry& AllottedGeometry,
+	const FSlateRect& MyCullingRect,
+	FSlateWindowElementList& OutDrawElements,
+	int32 LayerId,
+	const FWidgetStyle& InWidgetStyle,
+	bool bParentEnabled
+) const
+{
+	LayerId = SCompoundWidget::OnPaint(
+		Args, 
+		AllottedGeometry, 
+		MyCullingRect, 
+		OutDrawElements, 
+		LayerId, 
+		InWidgetStyle, 
+		bParentEnabled
+	);
+
+	// Update CenterAbsolute after we've painted our PinImage
+	if (PinImage.IsValid())
+	{
+		CenterAbsolute = FGeometryHelper::CenterOf(PinImage->GetPaintSpaceGeometry());
+	}
+
+	return LayerId;
+}
+
+FReply SConvertInnerPin::OnMouseButtonDown(const FGeometry& MyGeometry, const FPointerEvent& MouseEvent)
+{
+	if (MouseEvent.GetEffectingButton() == EKeys::LeftMouseButton)
+	{
+		if (MouseEvent.IsAltDown())
+		{
+			BreakConnections();
+			return FReply::Handled();
+		}
+
+		return FReply::Handled()
+		.DetectDrag(SharedThis(this), MouseEvent.GetEffectingButton())
+		.CaptureMouse(SharedThis(this));
+	}
+	
+	// We still want to detect right-clicks on mouse button up
+	return FReply::Handled();
+}
+
+FReply SConvertInnerPin::OnMouseButtonUp(const FGeometry& MyGeometry, const FPointerEvent& MouseEvent)
+{
+	if (MouseEvent.GetEffectingButton() == EKeys::RightMouseButton)
+	{
+		// Only show to context menu if we right click the actual InnerPin
+		if (PinImage && PinImage->GetCachedGeometry().IsUnderLocation(MouseEvent.GetScreenSpacePosition()))
+		{
+			FMenuBuilder ContextMenuBuilder(true, NULL);
+
+			const FText BreakConnectionText =
+				WeakConnectedPins.Num() == 1
+				? LOCTEXT("ContextMenu.BreakInnerConnection", "Break Connection")
+				: LOCTEXT("ContextMenu.BreakInnerConnections", "Break Connections");
+
+			ContextMenuBuilder.AddMenuEntry(
+				BreakConnectionText,
+				FText(),
+				FSlateIcon(),
+				FUIAction(FExecuteAction::CreateSP(this, &SConvertInnerPin::BreakConnections))
+			);
+
+			FSlateApplication::Get().PushMenu(
+				AsShared(),
+				FWidgetPath(),
+				ContextMenuBuilder.MakeWidget(),
+				FSlateApplication::Get().GetCursorPos(),
+				FPopupTransitionEffect(FPopupTransitionEffect::ContextMenu)
+			);
+
+			return FReply::Handled();
+		}
+	}
+		
+	return FReply::Unhandled();
+}
+
+FReply SConvertInnerPin::OnDragDetected(const FGeometry& MyGeometry, const FPointerEvent& MouseEvent)
+{
+	if (TSharedPtr<SGraphNodeMaterialConvert> OwningNode = WeakOwningNode.Pin())
+	{
+		TSharedRef<FConvertDragDropOp> DragDropOp = MakeShareable(new FConvertDragDropOp(SharedThis(this)));
+		OwningNode->SetCurrentDragDropOp(DragDropOp);
+		return FReply::Handled().BeginDragDrop(DragDropOp);
+	}
+
+	return FReply::Unhandled();
+}
+
+FReply SConvertInnerPin::OnDrop(const FGeometry& MyGeometry, const FDragDropEvent& DragDropEvent)
+{
+	if (TSharedPtr<FConvertDragDropOp> DragDropOp = DragDropEvent.GetOperationAs<FConvertDragDropOp>())
+	{
+		if (TSharedPtr<SConvertInnerPin> SourcePin = DragDropOp->GetSourcePin())
+		{
+			// Is this a valid connection
+			if (SourcePin.Get() != this && SourcePin->WeakOwningNode == this->WeakOwningNode && SourcePin->bIsInputPin != this->bIsInputPin)
+			{
+				if (TSharedPtr<SGraphNodeMaterialConvert> OwningNode = WeakOwningNode.Pin())
+				{
+					// Form connection between two pins
+					OwningNode->FormConnection(SourcePin, SharedThis(this));
+				
+					// Clear The Current Drag Drop Op
+					OwningNode->SetCurrentDragDropOp(nullptr);
+					return FReply::Handled();
+				}
+			}
+		}
+	}
+
+	return FReply::Unhandled();
+}
+
+const FSlateBrush* SConvertInnerPin::GetPinBrush() const
+{
+	const bool bHasConnections = WeakConnectedPins.Num() > 0;
+	const FSlateBrush* ConnectedBrush = FAppStyle::GetBrush("Graph.Pin.Connected");
+	const FSlateBrush* DisconnectedBrush = FAppStyle::GetBrush("Graph.Pin.Disconnected");
+	return bHasConnections ? ConnectedBrush : DisconnectedBrush;
+}
+
+void SConvertInnerPin::CancelDragDrop()
+{
+	if (TSharedPtr<SGraphNodeMaterialConvert> OwningNode = WeakOwningNode.Pin())
+	{
+		OwningNode->SetCurrentDragDropOp(nullptr);
+	}
+}
+
+UE::Slate::FDeprecateVector2DResult SConvertInnerPin::GetPinCenterAbsolute() const
+{ 
+	return CenterAbsolute; 
+}
+
+void SConvertInnerPin::AddConnection(TSharedPtr<SConvertInnerPin> InOtherPin)
+{ 
+	WeakConnectedPins.Add(InOtherPin); 
+}
+
+void SConvertInnerPin::RemoveConnection(TSharedPtr<SConvertInnerPin> InOtherPin)
+{ 
+	WeakConnectedPins.Remove(InOtherPin); 
+}
+
+void SConvertInnerPin::RemoveAllConnections() 
+{
+	for (TWeakPtr<SConvertInnerPin> WeakConnectedPin : WeakConnectedPins)
+	{
+		if (TSharedPtr<SConvertInnerPin> ConnectedPin = WeakConnectedPin.Pin())
+		{
+			ConnectedPin->RemoveConnection(SharedThis(this));
+		}
+	}
+
+	WeakConnectedPins.Reset();
+}
+
+void SConvertInnerPin::BreakConnections()
+{
+	if (TSharedPtr<SGraphNodeMaterialConvert> OwningNode = WeakOwningNode.Pin())
+	{
+		OwningNode->BreakConnections(SharedThis(this));
+	}
+}
+
+EVisibility SConvertInnerPin::GetDefaultValueVisibility() const
+{ 
+	bool bIsVisible = false;
+
+	if (bIsInputPin)
+	{
+		if (TSharedPtr<SMaterialExpressionConvertGraphPin> OwningPin = WeakOwningPin.Pin())
+		{
+			// Only show Input Default values if the owning SMaterialExpressionConvertGraphPin isn't connected 
+			bIsVisible = !OwningPin->IsConnected();
+		}
+	}
+	else
+	{
+		// Show Output pin default values if they don't yet have a connection to an input inner pin
+		bIsVisible = WeakConnectedPins.IsEmpty();
+	}
+
+	return bIsVisible ? EVisibility::Visible : EVisibility::Collapsed;
+}
+
+TOptional<float> SConvertInnerPin::GetDefaultValue() const
+{
+	TSharedPtr<SGraphNodeMaterialConvert> OwningNode = WeakOwningNode.Pin();
+	if (!OwningNode)
+	{
+		return TOptional<float>();
+	}
+	return OwningNode->GetDefaultValue(SharedThis(this));
+}
+
+void SConvertInnerPin::SetDefaultValue(float InDefaultValue, ETextCommit::Type CommitType)
+{
+	if (TSharedPtr<SGraphNodeMaterialConvert> OwningNode = WeakOwningNode.Pin())
+	{
+		OwningNode->SetDefaultValue(
+			SharedThis(this),
+			InDefaultValue
+		);
+	}
+}
+
+EVisibility SConvertInnerPin::GetPinNameVisibility() const
+{
+	TSharedPtr<SMaterialExpressionConvertGraphPin> OwningPin = WeakOwningPin.Pin();
+	if (OwningPin && OwningPin->GetInnerPins().Num() == 1)
+	{
+		return EVisibility::Collapsed;
+	}
+
+	return EVisibility::SelfHitTestInvisible;
+}
+
+FText SConvertInnerPin::GetPinName() const
+{
+	TSharedPtr<SMaterialExpressionConvertGraphPin> OwningPin = WeakOwningPin.Pin();
+	if (OwningPin && OwningPin->GetInnerPins().Num() == 1)
+	{
+		return INVTEXT("");
+	}
+
+	switch (ComponentIndex)
+	{
+		case 0: 	return INVTEXT("R");
+		case 1: 	return INVTEXT("G");
+		case 2: 	return INVTEXT("B");
+		case 3: 	return INVTEXT("A");
+		default: 	return INVTEXT("_");
+	}
+}
+
+////////////////////////////////////////////////////////////////////////////////////////////////////////////////////////////////////////////
+// SGraphNodeMaterialConvert: Custom Slate widget for UMaterialExpressionConvert
+////////////////////////////////////////////////////////////////////////////////////////////////////////////////////////////////////////////
+
+void SGraphNodeMaterialConvert::Construct(const FArguments& InArgs, class UMaterialGraphNode* InNode)
+{
+	this->GraphNode = InNode;
+	this->MaterialNode = InNode;
+	this->UpdateGraphNode();
+}
+
+void SGraphNodeMaterialConvert::CreatePinWidgets()
+{
+	SGraphNodeMaterialBase::CreatePinWidgets();
+
+	if (!MaterialNode)
+	{
+		return;
+	}
+
+	UMaterialExpressionConvert* ConvertExpression = Cast<UMaterialExpressionConvert>(MaterialNode->MaterialExpression);
+	if (!ConvertExpression)
+	{
+		return;
+	}
+
+	// After we've set up all of our pins, set up initial connections based on our ConvertMappings
+	for (const FMaterialExpressionConvertMapping& ConvertMapping : ConvertExpression->ConvertMappings)
+	{
+		TSharedPtr<SConvertInnerPin> InputInnerPin = GetInnerPin(true, ConvertMapping.InputIndex, ConvertMapping.InputComponentIndex);
+		TSharedPtr<SConvertInnerPin> OutputInnerPin = GetInnerPin(false, ConvertMapping.OutputIndex, ConvertMapping.OutputComponentIndex);
+		if (InputInnerPin && OutputInnerPin)
+		{
+			InputInnerPin->AddConnection(OutputInnerPin);
+			OutputInnerPin->AddConnection(InputInnerPin);
+		}
+	}
+}
+
+TSharedPtr<SGraphPin> SGraphNodeMaterialConvert::CreatePinWidget(UEdGraphPin* Pin) const
+{
+	return SNew(SMaterialExpressionConvertGraphPin, Pin);
+}
+
+void SGraphNodeMaterialConvert::AddPin(const TSharedRef<SGraphPin>& PinToAdd)
+{
+	SGraphNodeMaterialBase::AddPin(PinToAdd);
+
+	// Create our Inner Pins, which are used to route input values to output values
+	TSharedRef<SMaterialExpressionConvertGraphPin> ConvertOuterPin = StaticCastSharedRef<SMaterialExpressionConvertGraphPin>(PinToAdd);
+	ConvertOuterPin->CreateInnerPins(SharedThis(this));
+}
+
+void SGraphNodeMaterialConvert::CreateInputSideAddButton(TSharedPtr<SVerticalBox> InputBox)
+{
+	TSharedRef<SButton> AddPinButton = StaticCastSharedRef<SButton>(AddPinButtonContent(
+		LOCTEXT("GraphNodeMaterialConvert_AddInputPin", "Add Input"),
+		LOCTEXT("GraphNodeMaterialConvert_AddInputPin_Tooltip", "Add an input to this convert node."),
+		false /* bRightSide */
+	));
+
+	AddPinButton->SetOnClicked(FOnClicked::CreateSP(this, &SGraphNodeMaterialConvert::OnAddInputPinClicked));
+
+	InputBox->AddSlot()
+	.FillHeight(1.f)
+	.Padding(5.f)
+	.VAlign(VAlign_Bottom)
+	.HAlign(HAlign_Left)
+	[
+		AddPinButton
+	];
+}
+
+void SGraphNodeMaterialConvert::CreateOutputSideAddButton(TSharedPtr<SVerticalBox> OutputBox)
+{
+	TSharedRef<SButton> AddPinButton = StaticCastSharedRef<SButton>(AddPinButtonContent(
+		LOCTEXT("GraphNodeMaterialConvert_AddOutputPin", "Add Output"),
+		LOCTEXT("GraphNodeMaterialConvert_AddOutputPin_Tooltip", "Add an output to this convert node."),
+		true /* bRightSide */
+	));
+
+	AddPinButton->SetOnClicked(FOnClicked::CreateSP(this, &SGraphNodeMaterialConvert::OnAddOutputPinClicked));
+
+	OutputBox->AddSlot()
+	.FillHeight(1.f)
+	.Padding(5.f)
+	.VAlign(VAlign_Bottom)
+	.HAlign(HAlign_Right)
+	[
+		AddPinButton
+	];
+}
+
+EVisibility SGraphNodeMaterialConvert::IsAddPinButtonVisible() const
+{
+	return EVisibility::Visible;
+}
+
+FReply SGraphNodeMaterialConvert::OnAddInputPinClicked()
+{
+	FSlateApplication::Get().PushMenu(
+		AsShared(), 
+		FWidgetPath(), 
+		CreateAddPinContextMenu(/* bInputPin = */ true), 
+		FSlateApplication::Get().GetCursorPos(), 
+		FPopupTransitionEffect(FPopupTransitionEffect::ContextMenu)
+	);
+
+	return FReply::Handled();
+}
+
+FReply SGraphNodeMaterialConvert::OnAddOutputPinClicked()
+{
+	FSlateApplication::Get().PushMenu(
+		AsShared(), 
+		FWidgetPath(), 
+		CreateAddPinContextMenu(/* bInputPin = */ false),
+		FSlateApplication::Get().GetCursorPos(), 
+		FPopupTransitionEffect(FPopupTransitionEffect::ContextMenu)
+	);
+
+	return FReply::Handled();
+}
+
+TSharedRef<SWidget> SGraphNodeMaterialConvert::CreateAddPinContextMenu(bool bInputPin)
+{
+	FMenuBuilder ContextMenuBuilder(true, NULL);
+
+	static const EMaterialExpressionConvertType ConvertTypes[] =
+	{ 
+		EMaterialExpressionConvertType::Scalar, 
+		EMaterialExpressionConvertType::Vector2, 
+		EMaterialExpressionConvertType::Vector3, 
+		EMaterialExpressionConvertType::Vector4 
+	};
+
+	for (const EMaterialExpressionConvertType ConvertType : ConvertTypes)
+	{
+		ContextMenuBuilder.AddMenuEntry(
+			MaterialExpressionConvertType::ToText(ConvertType),
+			FText(),
+			FSlateIcon(),
+			FUIAction(FExecuteAction::CreateSP(
+				this, 
+				&SGraphNodeMaterialConvert::AddNewPin, bInputPin, ConvertType
+			))
+		);
+	}
+	
+	return ContextMenuBuilder.MakeWidget();
+}
+
+void SGraphNodeMaterialConvert::AddNewPin(bool bInputPin, EMaterialExpressionConvertType ConvertType)
+{
+	UMaterialExpressionConvert* ConvertExpression = Cast<UMaterialExpressionConvert>(MaterialNode->MaterialExpression);
+	if (!ConvertExpression)
+	{
+		return;
+	}
+
+	const FScopedTransaction Transaction(LOCTEXT("AddPin", "Add Pin"));
+	ConvertExpression->Modify();
+
+	if (bInputPin)
+	{
+		FMaterialExpressionConvertInput NewInput;
+		NewInput.Type = ConvertType;
+		ConvertExpression->ConvertInputs.Add(NewInput);
+	}
+	else
+	{
+		FMaterialExpressionConvertOutput NewOutput;
+		NewOutput.Type = ConvertType;
+		ConvertExpression->ConvertOutputs.Add(NewOutput);
+	}
+
+	MaterialNode->ReconstructNode();
+}
+
+int32 SGraphNodeMaterialConvert::OnPaint(
+	const FPaintArgs& Args,
+	const FGeometry& AllottedGeometry,
+	const FSlateRect& MyCullingRect,
+	FSlateWindowElementList& OutDrawElements,
+	int32 LayerId,
+	const FWidgetStyle& InWidgetStyle,
+	bool bParentEnabled
+) const
+{
+	LayerId = SGraphNodeMaterialBase::OnPaint(
+		Args, 
+		AllottedGeometry, 
+		MyCullingRect, 
+		OutDrawElements, 
+		LayerId, 
+		InWidgetStyle, 
+		bParentEnabled
+	);
+
+	// Only Draw Connections when we're at a high-enough detail. Early-out otherwise
+	if (GetCurrentLOD() <= EGraphRenderingLOD::LowDetail)
+	{
+		return LayerId;
+	}
+
+	if (!MaterialNode)
+	{
+		return LayerId;
+	}
+
+	UMaterialExpressionConvert* ConvertExpression = Cast<UMaterialExpressionConvert>(MaterialNode->MaterialExpression);
+	if (!ConvertExpression)
+	{
+		return LayerId;
+	}
+
+	for (const FMaterialExpressionConvertMapping& ConvertMapping : ConvertExpression->ConvertMappings)
+	{
+		TSharedPtr<SConvertInnerPin> InputInnerPin = GetInnerPin(true, ConvertMapping.InputIndex, ConvertMapping.InputComponentIndex);
+		TSharedPtr<SConvertInnerPin> OutputInnerPin = GetInnerPin(false, ConvertMapping.OutputIndex, ConvertMapping.OutputComponentIndex);
+		if (InputInnerPin && OutputInnerPin)
+		{
+			const FVector2f CurveStart = InputInnerPin->GetPinCenterAbsolute();
+			const FVector2f CurveEnd = OutputInnerPin->GetPinCenterAbsolute();
+			MakeConnectionCurve(AllottedGeometry, OutDrawElements, LayerId, CurveStart, CurveEnd);
+		}
+	}
+
+	if (CurrentDragDropOp.IsValid())
+	{
+		if (TSharedPtr<SConvertInnerPin> SourcePin = CurrentDragDropOp->GetSourcePin())
+		{
+			const FVector2f CurveStart = SourcePin->GetPinCenterAbsolute();
+			const FVector2f CurveEnd = CurrentDragDropOp->GetScreenPosition() + Inverse(Args.GetWindowToDesktopTransform());
+			MakeConnectionCurve(AllottedGeometry, OutDrawElements, LayerId, CurveStart, CurveEnd);
+		}
+	}
+
+
+	return LayerId;
+}
+
+void SGraphNodeMaterialConvert::FormConnection(TSharedPtr<SConvertInnerPin> InnerPinA, TSharedPtr<SConvertInnerPin> InnerPinB)
+{
+	check(InnerPinA->IsInputPin() != InnerPinB->IsInputPin());
+
+	if (!MaterialNode)
+	{
+		return;
+	}
+
+	UMaterialExpressionConvert* ConvertExpression = Cast<UMaterialExpressionConvert>(MaterialNode->MaterialExpression);
+	if (!ConvertExpression)
+	{
+		return;
+	}
+
+	const FScopedTransaction Transaction(LOCTEXT("AddConnection", "Convert Node: Add Connection"));
+	ConvertExpression->Modify();
+
+	// Determine Which Pin is the input and which is the output
+	TSharedPtr<SConvertInnerPin> InputPin 	= InnerPinA->IsInputPin() ? InnerPinA : InnerPinB;
+	TSharedPtr<SConvertInnerPin> OutputPin 	= InnerPinA->IsInputPin() ? InnerPinB : InnerPinA;
+
+	FMaterialExpressionConvertMapping NewConvertMapping;
+	NewConvertMapping.InputIndex = InputPin->GetPinIndex();
+	NewConvertMapping.InputComponentIndex = InputPin->GetComponentIndex();
+	NewConvertMapping.OutputIndex = OutputPin->GetPinIndex();
+	NewConvertMapping.OutputComponentIndex = OutputPin->GetComponentIndex();
+
+	// Remove any mappings with the same OutputIndex and OutputComponentIndex
+	ConvertExpression->ConvertMappings.RemoveAll(
+		[&NewConvertMapping](const FMaterialExpressionConvertMapping& ExistingConvertMapping)
+		{
+			return ExistingConvertMapping.OutputIndex == NewConvertMapping.OutputIndex
+				&& ExistingConvertMapping.OutputComponentIndex == NewConvertMapping.OutputComponentIndex;
+		}
+	);
+
+	// Add new convert mapping and refresh
+	ConvertExpression->ConvertMappings.Add(NewConvertMapping);
+	ConvertExpression->RefreshNode();
+
+	// Input Pins can have multiple connections
+	InputPin->AddConnection(OutputPin);
+
+	// Output Pins can only have one connection, so clean up any existing connections
+	OutputPin->RemoveAllConnections();
+	OutputPin->AddConnection(InputPin);
+}
+
+void SGraphNodeMaterialConvert::BreakConnections(TSharedPtr<SConvertInnerPin> InnerPin)
+{
+	check(InnerPin);
+
+	if (!MaterialNode)
+	{
+		return;
+	}
+
+	UMaterialExpressionConvert* ConvertExpression = Cast<UMaterialExpressionConvert>(MaterialNode->MaterialExpression);
+	if (!ConvertExpression)
+	{
+		return;
+	}
+
+	const FScopedTransaction Transaction(LOCTEXT("BreakConnections", "Convert Node: Break Connections"));
+	ConvertExpression->Modify();
+
+	const bool bIsInputPin = InnerPin->IsInputPin();
+	const int32 PinIndex = InnerPin->GetPinIndex();
+	const int32 ComponentIndex = InnerPin->GetComponentIndex();
+
+	ConvertExpression->ConvertMappings.RemoveAll(
+		[bIsInputPin, PinIndex, ComponentIndex](const FMaterialExpressionConvertMapping& ExistingConvertMapping)
+		{
+			if (bIsInputPin)
+			{
+				return ExistingConvertMapping.InputIndex == PinIndex && ExistingConvertMapping.InputComponentIndex == ComponentIndex;
+			}
+			else
+			{
+				return ExistingConvertMapping.OutputIndex == PinIndex && ExistingConvertMapping.OutputComponentIndex == ComponentIndex;
+			}
+		}
+	);
+
+	ConvertExpression->RefreshNode();
+
+	InnerPin->RemoveAllConnections();
+}
+
+TOptional<float> SGraphNodeMaterialConvert::GetDefaultValue(TSharedPtr<const SConvertInnerPin> InnerPin) const
+{
+	UMaterialExpressionConvert* ConvertExpression = Cast<UMaterialExpressionConvert>(MaterialNode->MaterialExpression);
+	if (!ConvertExpression)
+	{
+		return TOptional<float>();
+	}
+
+	const int32 PinIndex = InnerPin->GetPinIndex();
+	const int32 ComponentIndex = InnerPin->GetComponentIndex();
+
+	if (InnerPin->IsInputPin())
+	{
+		return ConvertExpression->ConvertInputs[PinIndex].DefaultValue.Component(ComponentIndex);
+	}
+	else
+	{
+		return ConvertExpression->ConvertOutputs[PinIndex].DefaultValue.Component(ComponentIndex);
+	}
+}
+
+void SGraphNodeMaterialConvert::SetDefaultValue(TSharedPtr<SConvertInnerPin> InnerPin, const float InDefaultValue)
+{
+	UMaterialExpressionConvert* ConvertExpression = Cast<UMaterialExpressionConvert>(MaterialNode->MaterialExpression);
+	if (!ConvertExpression)
+	{
+		return;
+	}
+
+	const FScopedTransaction Transaction( LOCTEXT( "SetPinDefaultValue", "Set Pin Default Value" ) );
+	ConvertExpression->Modify();
+
+	const int32 PinIndex = InnerPin->GetPinIndex();
+	const int32 ComponentIndex = InnerPin->GetComponentIndex();
+
+	if (InnerPin->IsInputPin())
+	{
+		ConvertExpression->ConvertInputs[PinIndex].DefaultValue.Component(ComponentIndex) = InDefaultValue;
+	}
+	else
+	{
+		ConvertExpression->ConvertOutputs[PinIndex].DefaultValue.Component(ComponentIndex) = InDefaultValue;
+	}
+
+	ConvertExpression->RefreshNode();
+}
+
+TSharedPtr<SConvertInnerPin> SGraphNodeMaterialConvert::GetInnerPin(bool bInputPin, const int32 InPinIndex, const int32 InComponentIndex) const
+{
+	// First Find relevant input/output graph pin
+	TSharedPtr<SMaterialExpressionConvertGraphPin> GraphPin = nullptr;
+	if (bInputPin)
+	{
+		if (InputPins.IsValidIndex(InPinIndex))
+		{
+			GraphPin = StaticCastSharedRef<SMaterialExpressionConvertGraphPin>(InputPins[InPinIndex]);
+		}
+	}
+	else
+	{
+		if (OutputPins.IsValidIndex(InPinIndex))
+		{
+			GraphPin = StaticCastSharedRef<SMaterialExpressionConvertGraphPin>(OutputPins[InPinIndex]);
+		}
+	}
+
+	// Then get that the relevant inner pin from that graph pin
+	if (GraphPin.IsValid() && GraphPin->GetInnerPins().IsValidIndex(InComponentIndex))
+	{
+		return GraphPin->GetInnerPins()[InComponentIndex];
+	}
+
+	// No valid inner pin found, return nullptr
+	return nullptr;
+}
+
+void SGraphNodeMaterialConvert::MakeConnectionCurve(
+	const FGeometry& InAllottedGeometry,
+	FSlateWindowElementList& OutDrawElements,
+	int32& InOutLayerId,
+	const FVector2f& InCurveStart,
+	const FVector2f& InCurveEnd
+) const
+{
+	FVector2f LocalStart = InAllottedGeometry.AbsoluteToLocal(InCurveStart);
+	FVector2f LocalEnd = InAllottedGeometry.AbsoluteToLocal(InCurveEnd);
+
+	//If LocalStart.X is > LocalEnd.X, swap values
+	if (LocalStart.X > LocalEnd.X)
+	{
+		std::swap(LocalStart, LocalEnd);
+	}
+
+	const bool bAllValuesValid = 
+		(LocalStart.X != -FLT_MAX) 
+		&& (LocalStart.Y != -FLT_MAX) 
+		&& (LocalEnd.X != -FLT_MAX) 
+		&& (LocalEnd.Y != -FLT_MAX);
+
+	if (bAllValuesValid)
+	{
+		/* Because we make sure the left-most position is always the start position,
+		 * curve direction is always pointing to the right */
+		static const FVector2f Direction = FVector2f(100.0f, 0.0f);
+
+		FSlateDrawElement::MakeSpline(
+			OutDrawElements,
+			++InOutLayerId,
+			InAllottedGeometry.ToPaintGeometry(),
+			LocalStart,
+			Direction,
+			LocalEnd,
+			Direction,
+			2
+		);
+	}
+}
+
+#undef LOCTEXT_NAMESPACE
diff --git a/Engine/Source/Editor/MaterialEditor/Private/MaterialNodes/SGraphNodeMaterialCustom.cpp b/Engine/Source/Editor/MaterialEditor/Private/MaterialNodes/SGraphNodeMaterialCustom.cpp
new file mode 100644
--- /dev/null
+++ b/Engine/Source/Editor/MaterialEditor/Private/MaterialNodes/SGraphNodeMaterialCustom.cpp
@@ -0,0 +1,161 @@
+// Copyright Epic Games, Inc. All Rights Reserved.
+
+#include "MaterialNodes/SGraphNodeMaterialCustom.h"
+
+#include "Templates/Casts.h"
+#include "Widgets/Input/SCheckBox.h"
+#include "Widgets/Images/SImage.h"
+#include "MaterialGraph/MaterialGraphNode_Custom.h"
+#include "GraphEditorSettings.h"
+
+void SGraphNodeMaterialCustom::Construct(const FArguments& InArgs, class UMaterialGraphNode* InNode)
+{
+	this->GraphNode = InNode;
+	this->MaterialNode = InNode;
+
+	// Create the syntax highlighter
+	FHLSLSyntaxHighlighterMarshaller::FSyntaxTextStyle CodeStyle(
+		FAppStyle::Get().GetWidgetStyle<FTextBlockStyle>("SyntaxHighlight.SourceCode.Normal"),
+		FAppStyle::Get().GetWidgetStyle<FTextBlockStyle>("SyntaxHighlight.SourceCode.Operator"),
+		FAppStyle::Get().GetWidgetStyle<FTextBlockStyle>("SyntaxHighlight.SourceCode.Keyword"),
+		FAppStyle::Get().GetWidgetStyle<FTextBlockStyle>("SyntaxHighlight.SourceCode.String"),
+		FAppStyle::Get().GetWidgetStyle<FTextBlockStyle>("SyntaxHighlight.SourceCode.Number"),
+		FAppStyle::Get().GetWidgetStyle<FTextBlockStyle>("SyntaxHighlight.SourceCode.Comment"),
+		FAppStyle::Get().GetWidgetStyle<FTextBlockStyle>("SyntaxHighlight.SourceCode.PreProcessorKeyword"),
+		FAppStyle::Get().GetWidgetStyle<FTextBlockStyle>("SyntaxHighlight.SourceCode.Error"));
+
+	SyntaxHighlighter = FHLSLSyntaxHighlighterMarshaller::Create(CodeStyle);
+
+	this->SetCursor(EMouseCursor::CardinalCross);
+
+	this->UpdateGraphNode();
+}
+
+void SGraphNodeMaterialCustom::CreateBelowPinControls(TSharedPtr<SVerticalBox> MainBox)
+{
+	UMaterialGraphNode_Custom* CustomNode = CastChecked<UMaterialGraphNode_Custom>(GraphNode);
+
+	TAttribute<FText> GetText;
+	GetText.BindUObject(CustomNode, &UMaterialGraphNode_Custom::GetHlslText);
+
+	FOnTextCommitted TextCommit;
+	TextCommit.BindUObject(CustomNode, &UMaterialGraphNode_Custom::OnCustomHlslTextCommitted);
+
+	SAssignNew(ShaderTextBox, SMultiLineEditableTextBox)
+		.AutoWrapText(false)
+		.Margin(FMargin(5, 5, 5, 5))
+		.Text(GetText)
+		.Visibility(this, &SGraphNodeMaterialCustom::CodeVisibility)
+		.OnTextCommitted(TextCommit)
+		.Marshaller(SyntaxHighlighter);
+
+	TSharedPtr<SVerticalBox> PreviewBox;
+	SAssignNew(PreviewBox, SVerticalBox);
+
+	SGraphNodeMaterialBase::CreateBelowPinControls(PreviewBox);
+
+	MainBox->AddSlot()
+	.Padding(Settings->GetNonPinNodeBodyPadding())
+	.AutoHeight()
+	[
+		SNew(SHorizontalBox)
+		+ SHorizontalBox::Slot()
+		.AutoWidth()
+		.Padding(FMargin(10.f, 5.f, 10.f, 10.f))
+		[
+			ShaderTextBox.ToSharedRef()
+		]
+		+ SHorizontalBox::Slot()
+		.AutoWidth()
+		[
+			PreviewBox.ToSharedRef()
+		]
+	];
+}
+
+EVisibility SGraphNodeMaterialCustom::CodeVisibility() const
+{
+	UMaterialGraphNode_Custom* CastedNode = CastChecked<UMaterialGraphNode_Custom>(GraphNode);
+	if (!CastedNode)
+	{
+		return EVisibility::Collapsed;
+	}
+
+	ECheckBoxState State = CastedNode->IsCodeViewChecked();
+
+	return (State == ECheckBoxState::Checked) ? EVisibility::Visible : EVisibility::Collapsed;
+}
+
+void SGraphNodeMaterialCustom::CreateAdvancedViewArrow(TSharedPtr<SVerticalBox> MainBox)
+{
+	UMaterialGraphNode_Custom* CustomNode = CastChecked<UMaterialGraphNode_Custom>(GraphNode);
+
+	if (!CustomNode)
+	{
+		return;
+	}
+
+	MainBox->AddSlot()
+	.AutoHeight()
+	.HAlign(HAlign_Fill)
+	.VAlign(VAlign_Top)
+	.Padding(3, 0, 3, 3)
+	[
+		SNew(SCheckBox)
+		.Visibility(this, &SGraphNodeMaterialCustom::AdvancedViewArrowVisibility)
+		.OnCheckStateChanged( this, &SGraphNodeMaterialCustom::OnAdvancedViewChanged )
+		.IsChecked( this, &SGraphNodeMaterialCustom::IsAdvancedViewChecked )
+		.Cursor(EMouseCursor::Default)
+		.Style(FAppStyle::Get(), "Graph.Node.AdvancedView")
+		[
+			SNew(SHorizontalBox)
+			+SHorizontalBox::Slot()
+			.VAlign(VAlign_Center)
+			.HAlign(HAlign_Center)
+			[
+				SNew(SImage)
+				. Image(this, &SGraphNodeMaterialCustom::GetAdvancedViewArrow)
+			]
+		]
+	];
+}
+
+EVisibility SGraphNodeMaterialCustom::AdvancedViewArrowVisibility() const
+{
+	// Always show the arrow no matter the configuration of the custom hlsl node
+	return GraphNode ? EVisibility::Visible : EVisibility::Collapsed;
+}
+
+void SGraphNodeMaterialCustom::OnAdvancedViewChanged(const ECheckBoxState NewCheckedState)
+{
+	if (UMaterialGraphNode_Custom* CastedNode = Cast<UMaterialGraphNode_Custom>(GraphNode))
+	{
+		CastedNode->OnCodeViewChanged(NewCheckedState);
+	}
+}
+
+ECheckBoxState SGraphNodeMaterialCustom::IsAdvancedViewChecked() const
+{
+	if (UMaterialGraphNode_Custom* CastedNode = Cast<UMaterialGraphNode_Custom>(GraphNode))
+	{
+		return CastedNode->IsCodeViewChecked();
+	}
+
+	return ECheckBoxState::Unchecked;
+}
+
+#define CHEVRON_UP TEXT("Icons.ChevronUp")
+#define CHEVRON_DOWN TEXT("Icons.ChevronDown")
+
+const FSlateBrush* SGraphNodeMaterialCustom::GetAdvancedViewArrow() const
+{
+	UMaterialGraphNode_Custom* CastedNode = Cast<UMaterialGraphNode_Custom>(GraphNode);
+	if (!CastedNode)
+	{
+		return FAppStyle::GetBrush(CHEVRON_UP);
+	}
+
+	ECheckBoxState State = CastedNode->IsCodeViewChecked();
+
+	return FAppStyle::GetBrush((State == ECheckBoxState::Checked) ? CHEVRON_UP : CHEVRON_DOWN);
+}
diff --git a/Engine/Source/Editor/MaterialEditor/Private/MaterialNodes/SGraphNodeMaterialResult.cpp b/Engine/Source/Editor/MaterialEditor/Private/MaterialNodes/SGraphNodeMaterialResult.cpp
new file mode 100644
--- /dev/null
+++ b/Engine/Source/Editor/MaterialEditor/Private/MaterialNodes/SGraphNodeMaterialResult.cpp
@@ -0,0 +1,113 @@
+// Copyright Epic Games, Inc. All Rights Reserved.
+
+#include "MaterialNodes/SGraphNodeMaterialResult.h"
+
+#include "Containers/Array.h"
+#include "Containers/UnrealString.h"
+#include "EdGraph/EdGraphNode.h"
+#include "EdGraph/EdGraphPin.h"
+#include "GenericPlatform/ICursor.h"
+#include "GraphEditor.h"
+#include "HAL/Platform.h"
+#include "IDocumentation.h"
+#include "Internationalization/Text.h"
+#include "MaterialGraph/MaterialGraph.h"
+#include "MaterialGraph/MaterialGraphNode_Base.h"
+#include "MaterialGraph/MaterialGraphNode_Root.h"
+#include "MaterialGraph/MaterialGraphSchema.h"
+#include "Materials/Material.h"
+#include "Misc/AssertionMacros.h"
+#include "Misc/Attribute.h"
+#include "Misc/Guid.h"
+#include "Misc/Optional.h"
+#include "SGraphPanel.h"
+#include "SGraphPin.h"
+#include "Templates/Casts.h"
+#include "Templates/SharedPointer.h"
+#include "TutorialMetaData.h"
+#include "UObject/NameTypes.h"
+#include "UObject/ObjectPtr.h"
+#include "Widgets/SToolTip.h"
+
+class SToolTip;
+
+/////////////////////////////////////////////////////
+// SGraphNodeMaterialResult
+
+void SGraphNodeMaterialResult::Construct(const FArguments& InArgs, UMaterialGraphNode_Root* InNode)
+{
+	this->GraphNode = InNode;
+	this->RootNode = InNode;
+
+	this->SetCursor(EMouseCursor::CardinalCross);
+
+	this->UpdateGraphNode();
+}
+
+void SGraphNodeMaterialResult::CreatePinWidgets()
+{
+	// Create Pin widgets for each of the pins.
+	UMaterialGraphNode_Base* MaterialGraphNode = Cast<UMaterialGraphNode_Base>(GraphNode);
+	UMaterialGraph* MaterialGraph = CastChecked<UMaterialGraph>(GraphNode->GetGraph());
+
+	bool bHideNoConnectionPins = false;
+	if (OwnerGraphPanelPtr.IsValid())
+	{
+		bHideNoConnectionPins = OwnerGraphPanelPtr.Pin()->GetPinVisibility() == SGraphEditor::Pin_HideNoConnection;
+	}
+
+	for (UEdGraphPin* CurPin : MaterialGraphNode->Pins)
+	{
+		const bool bPinHasConections = CurPin->LinkedTo.Num() > 0;
+
+		bool bPinDesiresToBeHidden = bHideNoConnectionPins && !bPinHasConections;
+		if (CurPin->PinType.PinCategory != UMaterialGraphSchema::PC_Exec)
+		{
+			if (!MaterialGraph->MaterialInputs[CurPin->SourceIndex].IsVisiblePin(MaterialGraph->Material))
+			{
+				bPinDesiresToBeHidden = true;
+			}
+		}
+
+		if (!bPinDesiresToBeHidden)
+		{
+			TSharedPtr<SGraphPin> NewPin = CreatePinWidget(CurPin);
+			check(NewPin.IsValid());
+
+			if (CurPin->PinType.PinCategory != UMaterialGraphSchema::PC_Exec)
+			{
+				TSharedPtr<SToolTip> ToolTipWidget = IDocumentation::Get()->CreateToolTip(MaterialGraph->MaterialInputs[CurPin->SourceIndex].GetToolTip(), nullptr, FString(TEXT("")), FString(TEXT("")));
+				NewPin->SetToolTip(ToolTipWidget.ToSharedRef());
+			}
+
+			this->AddPin(NewPin.ToSharedRef());
+		}
+	}
+}
+
+void SGraphNodeMaterialResult::MoveTo(const FVector2f& NewPosition, FNodeSet& NodeFilter, bool bMarkDirty)
+{
+	SGraphNode::MoveTo(NewPosition, NodeFilter, bMarkDirty);
+
+	RootNode->Material->EditorX = RootNode->NodePosX;
+	RootNode->Material->EditorY = RootNode->NodePosY;
+	RootNode->Material->MarkPackageDirty();
+	RootNode->Material->MaterialGraph->MaterialDirtyDelegate.ExecuteIfBound();
+}
+
+
+void SGraphNodeMaterialResult::PopulateMetaTag(FGraphNodeMetaData* TagMeta) const
+{
+	if( (GraphNode != nullptr) && (RootNode != nullptr) )
+	{		
+		UMaterialGraph* OuterGraph = RootNode->GetTypedOuter<UMaterialGraph>();
+		if (OuterGraph != nullptr)
+		{
+			TagMeta->OuterName = OuterGraph->OriginalMaterialFullName;
+			// There is only one root node - so we dont need a guid. 
+			TagMeta->Tag = FName(*FString::Printf(TEXT("MaterialResNode_%s"), *TagMeta->OuterName));
+			TagMeta->GUID.Invalidate();
+			TagMeta->FriendlyName = FString::Printf(TEXT("Material Result node in %s"), *TagMeta->OuterName);
+ 		}		
+	}
+}
diff --git a/Engine/Source/Editor/MaterialEditor/Private/MaterialNodes/SGraphSubstrateMaterial.cpp b/Engine/Source/Editor/MaterialEditor/Private/MaterialNodes/SGraphSubstrateMaterial.cpp
new file mode 100644
--- /dev/null
+++ b/Engine/Source/Editor/MaterialEditor/Private/MaterialNodes/SGraphSubstrateMaterial.cpp
@@ -0,0 +1,484 @@
+// Copyright Epic Games, Inc. All Rights Reserved.
+
+#include "SGraphSubstrateMaterial.h"
+#include "Internationalization/Text.h"
+#include "Rendering/RenderingCommon.h"
+#include "Rendering/SubstrateMaterialShared.h"
+#include "MaterialGraph/MaterialGraphNode.h"
+#include "MaterialGraph/MaterialGraphNode_Root.h"
+#include "MaterialGraph/MaterialGraphSchema.h"
+#include "Materials/MaterialExpression.h"
+#include "Materials/MaterialExpressionSubstrate.h"
+#include "Materials/MaterialExpressionFunctionInput.h"
+#include "Materials/MaterialExpressionFunctionOutput.h"
+#include "Materials/MaterialExpressionMaterialFunctionCall.h"
+
+#include "SGraphPin.h"
+#include "Math/Color.h"
+#include "Styling/StyleColors.h"
+#include "SubstrateDefinitions.h"
+#include "Widgets/SBoxPanel.h"
+#include "MaterialValueType.h"
+
+enum class ESubstrateWidgetOutputType : uint8
+{
+	Node,
+	DetailPanel,
+};
+
+static EStyleColor GetSubstrateWidgetColor0() { return EStyleColor::AccentBlue;  }
+static EStyleColor GetSubstrateWidgetColor1() { return EStyleColor::AccentGreen;  }
+FLinearColor FSubstrateWidget::GetConnectionColor() { return FLinearColor(0.16f, 0.015f, 0.24f) * 4.f; }
+
+bool FSubstrateWidget::HasInputSubstrateType(const UEdGraphPin* InPin)
+{
+	if (InPin == nullptr) return false;
+
+	if (UMaterialGraphNode_Root* RootPinNode = Cast<UMaterialGraphNode_Root>(InPin->GetOwningNode()))
+	{
+		if(RootPinNode->GetPinMaterialValueType(InPin) == MCT_Substrate)
+		{
+			return true;
+		}
+	}
+	if (UMaterialGraphNode* PinNode = Cast<UMaterialGraphNode>(InPin->GetOwningNode()))
+	{
+		FName TargetPinName = PinNode->GetShortenPinName(InPin->PinName);
+
+		for (FExpressionInputIterator It{ PinNode->MaterialExpression }; It; ++It)
+		{
+			FName InputName = PinNode->MaterialExpression->GetInputName(It.Index);
+			InputName = PinNode->GetShortenPinName(InputName);
+
+			if (InputName == TargetPinName)
+			{
+				switch (PinNode->MaterialExpression->GetInputValueType(It.Index))
+				{
+					case MCT_Substrate:
+						return true;
+				}
+				break;
+			}
+		}
+	}
+	return false;
+}
+
+bool FSubstrateWidget::HasOutputSubstrateType(const UEdGraphPin* InPin)
+{
+	if (InPin == nullptr || InPin->Direction != EGPD_Output) return false;
+
+	if (UMaterialGraphNode* PinNode = Cast<UMaterialGraphNode>(InPin->GetOwningNode()))
+	{
+		const TArray<FExpressionOutput>& ExpressionOutputs = PinNode->MaterialExpression->GetOutputs();
+		if (InPin->SourceIndex < ExpressionOutputs.Num() && PinNode->MaterialExpression->GetOutputValueType(InPin->SourceIndex) == MCT_Substrate)
+		{
+			return true;
+		}
+	}
+	return false;
+}
+
+static const TSharedRef<SWidget> InternalProcessOperator(
+	const FSubstrateMaterialCompilationOutput& CompilationOutput, 
+	const FSubstrateOperator& Op, 
+	ESubstrateWidgetOutputType OutputType,
+	const TArray<FGuid>& InGuid,
+	EStyleColor OverrideColor, 
+	TSharedPtr<SWidget>& OutFeatureWidget)
+{
+	const bool bIsCurrent = OutputType == ESubstrateWidgetOutputType::Node ? InGuid.Find(Op.MaterialExpressionGuid) != INDEX_NONE : false;
+	const EStyleColor Color0 = bIsCurrent ? GetSubstrateWidgetColor0() : OverrideColor;
+	const EStyleColor Color1 = bIsCurrent ? GetSubstrateWidgetColor1() : OverrideColor;
+	switch (Op.OperatorType)
+	{
+		case SUBSTRATE_OPERATOR_WEIGHT:
+		{
+			const EStyleColor Color = bIsCurrent ? EStyleColor::AccentGreen : OverrideColor;
+			return InternalProcessOperator(CompilationOutput, CompilationOutput.Operators[Op.LeftIndex], OutputType, InGuid, Color, OutFeatureWidget);
+		}
+		case SUBSTRATE_OPERATOR_VERTICAL:
+		{
+			auto VerticalOperator = SNew(SVerticalBox)
+				+SVerticalBox::Slot()
+				.AutoHeight()
+				.VAlign(VAlign_Fill)
+				.HAlign(HAlign_Fill)
+				.Padding(0.0f, 0.0f, 1.0f, 1.0f)
+				[
+					InternalProcessOperator(CompilationOutput, CompilationOutput.Operators[Op.LeftIndex], OutputType, InGuid, Color0, OutFeatureWidget)
+				]
+				+ SVerticalBox::Slot()
+				.AutoHeight()
+				.VAlign(VAlign_Fill)
+				.HAlign(HAlign_Fill)
+				.Padding(0.0f, 0.0f, 1.0f, 1.0f)
+				[
+					InternalProcessOperator(CompilationOutput, CompilationOutput.Operators[Op.RightIndex], OutputType, InGuid, Color1, OutFeatureWidget)
+				];
+			return VerticalOperator->AsShared();
+		}
+		case SUBSTRATE_OPERATOR_HORIZONTAL:
+		case SUBSTRATE_OPERATOR_SELECT:
+		{
+			auto HorizontalOperator = SNew(SHorizontalBox)
+				+SHorizontalBox::Slot()
+				.AutoWidth()
+				.VAlign(VAlign_Fill)
+				.HAlign(HAlign_Fill)
+				.Padding(0.0f, 0.0f, 1.0f, 1.0f)
+				[
+					InternalProcessOperator(CompilationOutput, CompilationOutput.Operators[Op.LeftIndex], OutputType, InGuid, Color0, OutFeatureWidget)
+				]
+				+SHorizontalBox::Slot()
+				.AutoWidth()
+				.VAlign(VAlign_Fill)
+				.HAlign(HAlign_Fill)
+				.Padding(0.0f, 0.0f, 1.0f, 1.0f)
+				[
+					InternalProcessOperator(CompilationOutput, CompilationOutput.Operators[Op.RightIndex], OutputType, InGuid, Color1, OutFeatureWidget)
+				];
+			return HorizontalOperator->AsShared();
+		}
+		case SUBSTRATE_OPERATOR_ADD:
+		{
+			auto HorizontalOperator = SNew(SHorizontalBox)
+				+SHorizontalBox::Slot()
+				.AutoWidth()
+				.VAlign(VAlign_Fill)
+				.HAlign(HAlign_Fill)
+				.Padding(0.0f, 0.0f, 1.0f, 1.0f)
+				[
+					InternalProcessOperator(CompilationOutput, CompilationOutput.Operators[Op.LeftIndex], OutputType, InGuid, Color0, OutFeatureWidget)
+				]
+				+SHorizontalBox::Slot()
+				.AutoWidth()
+				.VAlign(VAlign_Fill)
+				.HAlign(HAlign_Fill)
+				.Padding(0.0f, 0.0f, 1.0f, 1.0f)
+				[
+					InternalProcessOperator(CompilationOutput, CompilationOutput.Operators[Op.RightIndex], OutputType, InGuid, Color1, OutFeatureWidget)
+				];
+			return HorizontalOperator->AsShared();
+		}
+		case SUBSTRATE_OPERATOR_BSDF_LEGACY:	// legacy BSDF should have been converted to BSDF already.
+		case SUBSTRATE_OPERATOR_BSDF:
+		{
+			auto GetEyeDesc = [&Op]()
+			{
+				FString EyeDesc = TEXT("Eye");
+				const bool bIrisN = Op.Has(ESubstrateBsdfFeature_EyeIrisNormalPluggedIn);
+				const bool bIrisT = Op.Has(ESubstrateBsdfFeature_EyeIrisTangentPluggedIn);
+				if (bIrisN || bIrisT)
+				{
+					EyeDesc += TEXT("Iris");
+				}
+				if (bIrisN)
+				{
+					EyeDesc += TEXT("N");
+				}
+				if (bIrisT)
+				{
+					EyeDesc += TEXT("T");
+				}
+				EyeDesc += TEXT(" ");
+				return EyeDesc;
+			};
+
+			FString BSDFDesc = OutputType == ESubstrateWidgetOutputType::Node ? 
+											 FString(TEXT("BSDF")) : 
+											 FString::Printf(TEXT("BSDF (%s%s%s%s%s%s%s%s%s%s)")
+											, Op.Has(ESubstrateBsdfFeature_EdgeColor) ? TEXT("F90 ") : TEXT("")
+											, Op.Has(ESubstrateBsdfFeature_SSS) ? TEXT("SSS ") : TEXT("")
+											, Op.Has(ESubstrateBsdfFeature_MFPPluggedIn) ? TEXT("MFP ") : TEXT("")
+											, Op.Has(ESubstrateBsdfFeature_Anisotropy) ? TEXT("Ani ") : TEXT("")
+											, Op.Has(ESubstrateBsdfFeature_SecondRoughnessOrSimpleClearCoat) ? TEXT("2Ro ") : TEXT("")
+											, Op.Has(ESubstrateBsdfFeature_Fuzz) ? TEXT("Fuz ") : TEXT("")
+											, Op.Has(ESubstrateBsdfFeature_Glint) ? TEXT("Gli ") : TEXT("")
+											, Op.Has(ESubstrateBsdfFeature_SpecularProfile) ? TEXT("Spc ") : TEXT("")
+											, Op.Has(ESubstrateBsdfFeature_Eye) ? *GetEyeDesc() : TEXT("")
+											, Op.Has(ESubstrateBsdfFeature_Hair) ? TEXT("Hai ") : TEXT("")
+											);
+
+			static FString ToolTip;
+			if (OutputType == ESubstrateWidgetOutputType::DetailPanel && ToolTip.IsEmpty())
+			{
+				ToolTip += TEXT("SSS means the BSDF features subsurface profile or subsurface setup using MFP.\n");
+				ToolTip += TEXT("MFP means the BSDF MFP is specified by the user.\n");
+				ToolTip += TEXT("F90 means the BSDF edge specular color representing reflectivity at grazing angle is used.\n");
+				ToolTip += TEXT("Fuz means the BSDF fuzz layer is enabled.\n");
+				ToolTip += TEXT("2Ro means the BSDF either uses a second specular lob with a second roughness, or the legacy simple clear coat.\n");
+				ToolTip += TEXT("Ani means the BSDF anisotropic specular lighting is used.\n");
+				ToolTip += TEXT("Gli means the BSDF features glints.\n");
+				ToolTip += TEXT("Spc means the BSDF features specular profile.\n");
+				ToolTip += TEXT("Eye means the BSDF features eye is used.\n");
+				ToolTip += TEXT("Hai means the BSDF features hair is used.\n");
+			}
+
+			// Features and SubSurface type
+			if (bIsCurrent && OutputType == ESubstrateWidgetOutputType::Node)
+			{
+				EMaterialSubSurfaceType SubSurfaceType = EMaterialSubSurfaceType(Op.SubSurfaceType);
+				if (!Op.bIsBottom)
+				{
+					SubSurfaceType = EMaterialSubSurfaceType::MSS_SimpleVolume;
+				}
+				else if (!Op.Has(ESubstrateBsdfFeature_SSS))
+				{
+					SubSurfaceType = EMaterialSubSurfaceType::MSS_None;
+				}
+
+				const uint32 DisplayedFeatureMask = 
+					ESubstrateBsdfFeature::ESubstrateBsdfFeature_EdgeColor |
+					ESubstrateBsdfFeature::ESubstrateBsdfFeature_Fuzz |
+					ESubstrateBsdfFeature::ESubstrateBsdfFeature_SecondRoughnessOrSimpleClearCoat |
+					ESubstrateBsdfFeature::ESubstrateBsdfFeature_Anisotropy |
+					ESubstrateBsdfFeature::ESubstrateBsdfFeature_Glint |
+					ESubstrateBsdfFeature::ESubstrateBsdfFeature_SpecularProfile;
+
+				const bool bHasFeatures = (Op.BSDFFeatures & DisplayedFeatureMask) != 0;
+
+				if (bHasFeatures || SubSurfaceType != EMaterialSubSurfaceType::MSS_None)
+				{
+					auto Features = SNew(SVerticalBox);
+
+					if (SubSurfaceType != EMaterialSubSurfaceType::MSS_None)
+					{
+						const FString BehaviorHint = SubSurfaceType != EMaterialSubSurfaceType::MSS_SimpleVolume ? TEXT("Opaque") : TEXT("Translucent");
+						const FText SSSTypeName =StaticEnum<EMaterialSubSurfaceType>()->GetDisplayNameTextByValue(int64(SubSurfaceType));
+						const FText Message = FText::Format(FTextFormat::FromString(TEXT("SSS {0} ({1})")), SSSTypeName, FText::FromString(BehaviorHint));
+
+						Features->AddSlot()
+						.AutoHeight()
+						.VAlign(VAlign_Fill)
+						.HAlign(HAlign_Fill)
+						.Padding(0.0f, 0.0f, 1.0f, 1.0f)
+						[
+							SNew(SErrorText)
+							.ErrorText(Message)
+							.BackgroundColor(FSlateColor(FLinearColor(0.65f,0.25f,0.05f)))
+						];
+					}
+
+					if (bHasFeatures)
+					{
+						auto HorizontalOperator = SNew(SHorizontalBox);
+						auto AddFeatureSlot = [&HorizontalOperator](FText Message)
+						{
+							HorizontalOperator->AddSlot()
+							.AutoWidth()
+							.VAlign(VAlign_Fill)
+							.HAlign(HAlign_Fill)
+							.Padding(0.0f, 0.0f, 1.0f, 1.0f)
+							[
+								SNew(SErrorText)
+								.ErrorText(Message)
+								.BackgroundColor(FSlateColor(FLinearColor(0.16f, 0.015f, 0.24f)))
+							];
+						};
+
+						if (Op.Has(ESubstrateBsdfFeature::ESubstrateBsdfFeature_EdgeColor))
+						{
+							AddFeatureSlot(FText::FromString(TEXT("F90")));
+						}
+						if (Op.Has(ESubstrateBsdfFeature::ESubstrateBsdfFeature_Fuzz))
+						{
+							AddFeatureSlot(FText::FromString(TEXT("Fuzz")));
+						}
+						if (Op.Has(ESubstrateBsdfFeature::ESubstrateBsdfFeature_SecondRoughnessOrSimpleClearCoat))
+						{
+							AddFeatureSlot(FText::FromString(TEXT("Dual Spec.")));
+						}
+						if (Op.Has(ESubstrateBsdfFeature::ESubstrateBsdfFeature_Anisotropy))
+						{
+							AddFeatureSlot(FText::FromString(TEXT("Aniso")));
+						}
+						if (Op.Has(ESubstrateBsdfFeature::ESubstrateBsdfFeature_Glint))
+						{
+							AddFeatureSlot(FText::FromString(TEXT("Glints")));
+						}
+						if (Op.Has(ESubstrateBsdfFeature::ESubstrateBsdfFeature_SpecularProfile))
+						{
+							AddFeatureSlot(FText::FromString(TEXT("Spec. Profile")));
+						}
+
+						Features->AddSlot()
+						.AutoHeight()
+						.VAlign(VAlign_Fill)
+						.HAlign(HAlign_Fill)
+						.Padding(0.0f, 0.0f, 1.0f, 1.0f)
+						[
+							HorizontalOperator
+						];
+					}
+
+					OutFeatureWidget = Features.ToSharedPtr();
+				}
+			}
+
+			const EStyleColor Color = OverrideColor != EStyleColor::MAX ? OverrideColor : (bIsCurrent ? EStyleColor::AccentGreen : EStyleColor::AccentGray);
+			const FSlateColor SlateColor = OutputType == ESubstrateWidgetOutputType::Node ? FSlateColor(Color) : FSlateColor(FLinearColor(0.16f, 0.015f, 0.24f));
+			auto BSDF = SNew(SErrorText)
+				.ErrorText(FText::FromString(BSDFDesc))
+				.BackgroundColor(SlateColor)
+				.ToolTipText(FText::FromString(ToolTip));
+			return BSDF->AsShared();
+		}
+	}
+
+	static FString NoVisualization = OutputType == ESubstrateWidgetOutputType::DetailPanel ? FString(TEXT("Tree Operator Error")) : FString();
+	auto TreeOperatorError = SNew(SErrorText)
+		.ErrorText(FText::FromString(NoVisualization))
+		.BackgroundColor(FSlateColor(EStyleColor::AccentRed));
+	return TreeOperatorError->AsShared();
+}
+
+static const TSharedRef<SWidget> InternalProcessOperator(
+	const FSubstrateMaterialCompilationOutput& CompilationOutput,
+	const FSubstrateOperator& Op,
+	ESubstrateWidgetOutputType OutputType,
+	const TArray<FGuid>& InGuid,
+	EStyleColor OverrideColor)
+{
+	TSharedPtr<SWidget> FeatureWidget;
+	TSharedRef<SWidget> Topology = InternalProcessOperator(CompilationOutput, Op, OutputType, InGuid, OverrideColor, FeatureWidget);
+	if (FeatureWidget)
+	{
+		return SNew(SVerticalBox)
+		+SVerticalBox::Slot()
+		.AutoHeight()
+		.VAlign(VAlign_Fill)
+		.HAlign(HAlign_Fill)
+		.Padding(0.0f, 0.0f, 1.0f, 1.0f)
+		[
+			FeatureWidget.ToSharedRef()
+		]
+		+SVerticalBox::Slot()
+		.MinHeight(15)
+		.FillHeight(1.f)
+		.Padding(0.0f, 0.0f, 1.0f, 1.0f)
+		[
+			SNullWidget::NullWidget
+		]
+		+SVerticalBox::Slot()
+		.AutoHeight()
+		.VAlign(VAlign_Center)
+		.HAlign(HAlign_Center)
+		.Padding(0.0f, 0.0f, 1.0f, 1.0f)
+		[
+			Topology
+		];
+	}
+	else
+	{
+		return Topology;
+	}
+}
+const TSharedRef<SWidget> FSubstrateWidget::ProcessOperator(const FSubstrateMaterialCompilationOutput& CompilationOutput)
+{
+	return InternalProcessOperator(CompilationOutput, CompilationOutput.Operators[CompilationOutput.RootOperatorIndex], ESubstrateWidgetOutputType::DetailPanel, TArray<FGuid>(), EStyleColor::MAX);
+}
+
+const TSharedRef<SWidget> FSubstrateWidget::ProcessOperator(const FSubstrateMaterialCompilationOutput& CompilationOutput, const TArray<FGuid>& InGuid)
+{
+	return InternalProcessOperator(CompilationOutput, CompilationOutput.Operators[CompilationOutput.RootOperatorIndex], ESubstrateWidgetOutputType::Node, InGuid, EStyleColor::MAX);
+}
+
+void FSubstrateWidget::GetPinColor(TSharedPtr<SGraphPin>& Out, const UMaterialGraphNode* InNode)
+{	
+	if (!InNode || !InNode->MaterialExpression)
+	{
+		return;
+	}
+
+	const FLinearColor Color0 = USlateThemeManager::Get().GetColor(GetSubstrateWidgetColor0());
+	const FLinearColor Color1 = USlateThemeManager::Get().GetColor(GetSubstrateWidgetColor1());
+
+	FLinearColor ColorModifier = FLinearColor::Black;
+	bool bHasColorModifier = false;
+	// Substrate operator override pin color to ease material topology visualization
+	const UEdGraphPin* Pin = Out->SGraphPin::GetPinObj();
+	const FName PinName = Pin->PinName;
+	if (InNode->MaterialExpression->IsA(UMaterialExpressionSubstrateVerticalLayering::StaticClass()))
+	{			
+		if (PinName == InNode->MaterialExpression->GetInputName(0)) // Top
+		{
+			bHasColorModifier = true;
+			ColorModifier = Color0;
+		}
+		else if (PinName == InNode->MaterialExpression->GetInputName(1)) // Base
+		{
+			bHasColorModifier = true;
+			ColorModifier = Color1;
+		}
+	}
+	else if (InNode->MaterialExpression->IsA(UMaterialExpressionSubstrateHorizontalMixing::StaticClass()))
+	{
+		if (PinName == InNode->MaterialExpression->GetInputName(1)) // Foreground
+		{
+			bHasColorModifier = true;
+			ColorModifier = Color1;
+		}
+		else if (PinName == InNode->MaterialExpression->GetInputName(0)) // Background
+		{
+			bHasColorModifier = true;
+			ColorModifier = Color0;
+		}
+	}
+	else if (InNode->MaterialExpression->IsA(UMaterialExpressionSubstrateSelect::StaticClass()))
+	{
+		if (PinName == InNode->MaterialExpression->GetInputName(1)) // B
+		{
+			bHasColorModifier = true;
+			ColorModifier = Color1;
+		}
+		else if (PinName == InNode->MaterialExpression->GetInputName(0)) // A
+		{
+			bHasColorModifier = true;
+			ColorModifier = Color0;
+		}
+	}
+	else if (InNode->MaterialExpression->IsA(UMaterialExpressionSubstrateAdd::StaticClass()))
+	{
+		if (PinName == InNode->MaterialExpression->GetInputName(0)) // A
+		{
+			bHasColorModifier = true;
+			ColorModifier = Color0;
+		}
+		else if (PinName == InNode->MaterialExpression->GetInputName(1)) // B
+		{
+			bHasColorModifier = true;
+			ColorModifier = Color1;
+		}
+	}
+	else if (InNode->MaterialExpression->IsA(UMaterialExpressionMaterialFunctionCall::StaticClass()) ||
+			 InNode->MaterialExpression->IsA(UMaterialExpressionFunctionInput::StaticClass()) ||
+			 InNode->MaterialExpression->IsA(UMaterialExpressionFunctionOutput::StaticClass()))
+	{
+		const uint32 PinType = UMaterialGraphSchema::GetMaterialValueType(Pin);
+		if ((Out->GetDirection() == EGPD_Input || Out->GetDirection() == EGPD_Output) && (PinType & MCT_Substrate) != 0)
+		{
+			bHasColorModifier = true;
+			ColorModifier = FSubstrateWidget::GetConnectionColor();
+		}
+	}
+
+	if (InNode->MaterialExpression->IsA(UMaterialExpressionSubstrateBSDF::StaticClass()) && Out->GetDirection() == EGPD_Output && UMaterialGraphSchema::GetMaterialValueType(Pin) == MCT_Substrate)
+	{
+		bHasColorModifier = true;
+		ColorModifier = FSubstrateWidget::GetConnectionColor();
+	}
+	else if (Pin && UMaterialGraphSchema::GetMaterialValueType(Pin) == MCT_Substrate)
+	{
+		if (!Out->IsConnected())
+		{
+			bHasColorModifier = true;
+			ColorModifier = FSubstrateWidget::GetConnectionColor();
+		}
+	}	
+	
+	if (bHasColorModifier)
+	{
+		Out->SetPinColorModifier(ColorModifier);
+	}
+}
\ No newline at end of file
diff --git a/Engine/Source/Editor/MaterialEditor/Private/MaterialPins/SGraphPinMaterialInput.cpp b/Engine/Source/Editor/MaterialEditor/Private/MaterialPins/SGraphPinMaterialInput.cpp
new file mode 100644
--- /dev/null
+++ b/Engine/Source/Editor/MaterialEditor/Private/MaterialPins/SGraphPinMaterialInput.cpp
@@ -0,0 +1,49 @@
+// Copyright Epic Games, Inc. All Rights Reserved.
+
+
+#include "MaterialPins/SGraphPinMaterialInput.h"
+
+#include "EdGraph/EdGraphNode.h"
+#include "EdGraph/EdGraphPin.h"
+#include "MaterialGraph/MaterialGraph.h"
+#include "MaterialGraph/MaterialGraphSchema.h"
+#include "Misc/AssertionMacros.h"
+#include "Templates/Casts.h"
+#include "SGraphSubstrateMaterial.h"
+#include "MaterialValueType.h"
+
+void SGraphPinMaterialInput::Construct(const FArguments& InArgs, UEdGraphPin* InGraphPinObj)
+{
+	bool bUsePinColor = true;
+	if (Substrate::IsSubstrateEnabled())
+	{
+		if (UMaterialGraphSchema::GetMaterialValueType(InGraphPinObj) == MCT_Substrate)
+		{
+			bUsePinColor = false;
+		}			
+	}
+	SGraphPin::Construct(SGraphPin::FArguments().UsePinColorForText(bUsePinColor), InGraphPinObj);
+}
+
+FSlateColor SGraphPinMaterialInput::GetPinColor() const
+{
+	check(GraphPinObj);
+	UMaterialGraph* MaterialGraph = CastChecked<UMaterialGraph>(GraphPinObj->GetOwningNode()->GetGraph());
+	const UMaterialGraphSchema* Schema = CastChecked<UMaterialGraphSchema>(MaterialGraph->GetSchema());
+
+	if (MaterialGraph->IsInputActive(GraphPinObj))
+	{
+		if (Substrate::IsSubstrateEnabled())
+		{
+			if (UMaterialGraphSchema::GetMaterialValueType(GraphPinObj) == MCT_Substrate)
+			{
+				return FSlateColor(FSubstrateWidget::GetConnectionColor());
+			}			
+		}
+		return Schema->ActivePinColor;
+	}
+	else
+	{
+		return Schema->InactivePinColor;
+	}
+}
diff --git a/Engine/Source/Editor/MaterialEditor/Public/MaterialNodes/SGraphNodeMaterialBase.h b/Engine/Source/Editor/MaterialEditor/Public/MaterialNodes/SGraphNodeMaterialBase.h
new file mode 100644
--- /dev/null
+++ b/Engine/Source/Editor/MaterialEditor/Public/MaterialNodes/SGraphNodeMaterialBase.h
@@ -0,0 +1,148 @@
+// Copyright Epic Games, Inc. All Rights Reserved.
+
+#pragma once
+
+#include "Internationalization/Text.h"
+#include "Layout/Visibility.h"
+#include "Math/UnrealMathSSE.h"
+#include "Math/Vector2D.h"
+#include "Rendering/RenderingCommon.h"
+#include "SGraphNode.h"
+#include "SNodePanel.h"
+#include "Styling/SlateTypes.h"
+#include "Templates/SharedPointer.h"
+#include "Textures/SlateShaderResource.h"
+#include "Widgets/DeclarativeSyntaxSupport.h"
+
+class FMaterialRenderProxy;
+class FRHICommandListImmediate;
+class FSlateRect;
+class FWidgetStyle;
+class SGraphPin;
+class SOverlay;
+class SVerticalBox;
+class SWidget;
+class UMaterialGraphNode;
+struct FGeometry;
+struct FSlateBrush;
+
+typedef TSharedPtr<class FPreviewElement, ESPMode::ThreadSafe> FThreadSafePreviewPtr;
+
+class FPreviewViewport : public ISlateViewport
+{
+public:
+	FPreviewViewport(class UMaterialGraphNode* InNode);
+	~FPreviewViewport();
+
+	// ISlateViewport interface
+	virtual void OnDrawViewport( const FGeometry& AllottedGeometry, const FSlateRect& MyCullingRect, class FSlateWindowElementList& OutDrawElements, int32 LayerId, const FWidgetStyle& InWidgetStyle, bool bParentEnabled ) override;
+	virtual FIntPoint GetSize() const override;
+	virtual class FSlateShaderResource* GetViewportRenderTargetTexture() const override {return NULL;}
+	virtual bool RequiresVsync() const override {return false;}
+
+	/** Material node to get expression preview from */
+	UMaterialGraphNode* MaterialNode;
+	/** Custom Slate Element to display the preview */
+	FThreadSafePreviewPtr PreviewElement;
+private:
+	/**
+	 * Updates the expression preview render proxy from a graph node
+	 */
+	void UpdatePreviewNodeRenderProxy();
+};
+
+class FPreviewElement : public ICustomSlateElement
+{
+public:
+	FPreviewElement();
+	~FPreviewElement();
+
+	/**
+	 * Sets up the canvas for rendering
+	 *
+	 * @param	InCanvasRect			Size of the canvas tile
+	 * @param	InClippingRect			How to clip the canvas tile
+	 * @param	InGraphNode				The graph node for the material preview
+	 * @param	bInIsRealtime			Whether preview is using realtime values
+	 *
+	 * @return	Whether there is anything to render
+	 */
+	bool BeginRenderingCanvas( const FIntRect& InCanvasRect, const FIntRect& InClippingRect, UMaterialGraphNode* InGraphNode, bool bInIsRealtime );
+
+	/**
+	 * Updates the expression preview render proxy from a graph node on the render thread
+	 */
+	void UpdateExpressionPreview(UMaterialGraphNode* PreviewNode);
+private:
+	/**
+	 * ICustomSlateElement interface 
+	 */
+	virtual void Draw_RenderThread(FRDGBuilder& GraphBuilder, const FDrawPassInputs& Inputs) override;
+
+private:
+	/** Render target that the canvas renders to */
+	class FSlateMaterialPreviewRenderTarget* RenderTarget;
+	/** Render proxy for the expression preview */
+	FMaterialRenderProxy* ExpressionPreview;
+	/** Whether preview is using realtime values */
+	bool bIsRealtime;
+};
+
+class MATERIALEDITOR_API SGraphNodeMaterialBase : public SGraphNode
+{
+public:
+	SLATE_BEGIN_ARGS(SGraphNodeMaterialBase){}
+	SLATE_END_ARGS()
+
+	void Construct(const FArguments& InArgs, class UMaterialGraphNode* InNode);
+
+	// SGraphNode interface
+	virtual void CreatePinWidgets() override;
+	// End of SGraphNode interface
+
+	// SNodePanel::SNode interface
+	virtual void MoveTo(const FVector2f& NewPosition, FNodeSet& NodeFilter, bool bMarkDirty = true) override;
+	// End of SNodePanel::SNode interface
+
+	UMaterialGraphNode* GetMaterialGraphNode() const {return MaterialNode;}
+
+	/* Populate a meta data tag with information about this graph node */
+	virtual void PopulateMetaTag(class FGraphNodeMetaData* TagMeta) const override;
+
+protected:
+	// SGraphNode interface
+	virtual void AddPin( const TSharedRef<SGraphPin>& PinToAdd ) override;
+	virtual void CreateBelowPinControls(TSharedPtr<SVerticalBox> MainBox) override;
+	virtual void SetDefaultTitleAreaWidget(TSharedRef<SOverlay> DefaultTitleAreaWidget) override;
+	virtual TSharedRef<SWidget> CreateNodeContentArea() override;
+	virtual void OnAdvancedViewChanged(const ECheckBoxState NewCheckedState) override;
+	// End of SGraphNode interface
+
+	/** Creates a preview viewport if necessary */
+	TSharedRef<SWidget> CreatePreviewWidget();
+
+	/** Returns visibility of Expression Preview viewport */
+	EVisibility ExpressionPreviewVisibility() const;
+
+	/** Returns text to over lay over the expression preview viewport */
+	FText ExpressionPreviewOverlayText() const;
+
+	/** Called when we click to preview node button */
+	FReply OnTogglePreviewClicked();
+
+	/** Show/hide Expression Preview */
+	void OnExpressionPreviewChanged(const ECheckBoxState NewCheckedState);
+
+	/** hidden == unchecked, shown == checked */
+	ECheckBoxState IsExpressionPreviewChecked() const;
+
+	/** Up when shown, down when hidden */
+	const FSlateBrush* GetExpressionPreviewArrow() const;
+
+protected:
+	/** Slate viewport for rendering preview via custom slate element */
+	TSharedPtr<FPreviewViewport> PreviewViewport;
+
+	/** Cached material graph node pointer to avoid casting */
+	UMaterialGraphNode* MaterialNode;
+};
diff --git a/Engine/Source/Editor/MaterialEditor/Public/MaterialNodes/SGraphNodeMaterialComment.h b/Engine/Source/Editor/MaterialEditor/Public/MaterialNodes/SGraphNodeMaterialComment.h
new file mode 100644
--- /dev/null
+++ b/Engine/Source/Editor/MaterialEditor/Public/MaterialNodes/SGraphNodeMaterialComment.h
@@ -0,0 +1,27 @@
+// Copyright Epic Games, Inc. All Rights Reserved.
+
+#pragma once
+
+#include "Math/Vector2D.h"
+#include "SGraphNodeComment.h"
+#include "SNodePanel.h"
+#include "Widgets/DeclarativeSyntaxSupport.h"
+
+class UMaterialGraphNode_Comment;
+
+class MATERIALEDITOR_API SGraphNodeMaterialComment : public SGraphNodeComment
+{
+public:
+	SLATE_BEGIN_ARGS(SGraphNodeMaterialComment){}
+	SLATE_END_ARGS()
+
+	void Construct(const FArguments& InArgs, class UMaterialGraphNode_Comment* InNode);
+
+	// SNodePanel::SNode interface
+	virtual void MoveTo(const FVector2f& NewPosition, FNodeSet& NodeFilter, bool bMarkDirty = true) override;
+	// End of SNodePanel::SNode interface
+
+private:
+	/** Cached material graph node pointer to avoid casting */
+	UMaterialGraphNode_Comment* CommentNode;
+};
diff --git a/Engine/Source/Editor/MaterialEditor/Public/MaterialNodes/SGraphNodeMaterialComposite.h b/Engine/Source/Editor/MaterialEditor/Public/MaterialNodes/SGraphNodeMaterialComposite.h
new file mode 100644
--- /dev/null
+++ b/Engine/Source/Editor/MaterialEditor/Public/MaterialNodes/SGraphNodeMaterialComposite.h
@@ -0,0 +1,45 @@
+// Copyright Epic Games, Inc. All Rights Reserved.
+
+#pragma once
+
+#include "Internationalization/Text.h"
+#include "MaterialNodes/SGraphNodeMaterialBase.h"
+#include "Math/Vector2D.h"
+#include "SNodePanel.h"
+#include "Templates/SharedPointer.h"
+#include "Widgets/DeclarativeSyntaxSupport.h"
+
+class SToolTip;
+class SWidget;
+class UEdGraph;
+class UMaterialGraphNode_Composite;
+
+//@TODO: This class is mostly C&P from UK2Node_Composite, consolidate composites to not be BP specific.
+class MATERIALEDITOR_API SGraphNodeMaterialComposite : public SGraphNodeMaterialBase
+{
+public:
+	SLATE_BEGIN_ARGS(SGraphNodeMaterialComposite){}
+	SLATE_END_ARGS()
+
+	void Construct(const FArguments& InArgs, UMaterialGraphNode_Composite* InNode);
+
+	// SGraphNode interface
+	virtual void UpdateGraphNode() override;
+	virtual TSharedPtr<SToolTip> GetComplexTooltip() override;
+	// End of SGraphNode interface
+
+	// SNodePanel::SNode interface
+	virtual void MoveTo(const FVector2f& NewPosition, FNodeSet& NodeFilter, bool bMarkDirty = true) override;
+	// End of SNodePanel::SNode interface
+protected:
+	virtual UEdGraph* GetInnerGraph() const;
+
+private:
+	FText GetPreviewCornerText() const;
+	FText GetTooltipTextForNode() const;
+
+	TSharedRef<SWidget> CreateNodeBody();
+
+	/** Cached material graph node pointer to avoid casting */
+	UMaterialGraphNode_Composite* CompositeNode;
+};
diff --git a/Engine/Source/Editor/MaterialEditor/Public/MaterialNodes/SGraphNodeMaterialConvert.h b/Engine/Source/Editor/MaterialEditor/Public/MaterialNodes/SGraphNodeMaterialConvert.h
new file mode 100644
--- /dev/null
+++ b/Engine/Source/Editor/MaterialEditor/Public/MaterialNodes/SGraphNodeMaterialConvert.h
@@ -0,0 +1,201 @@
+// Copyright Epic Games, Inc. All Rights Reserved.
+
+#pragma once
+
+#include "Input/DragAndDrop.h"
+#include "MaterialNodes/SGraphNodeMaterialBase.h"
+#include "SGraphPin.h"
+#include "Types/SlateVector2.h"
+#include "Widgets/SCompoundWidget.h"
+
+class FConvertDragDropOp;
+class SConvertInnerPin;
+class SGraphNodeMaterialConvert;
+class SImage;
+enum class EMaterialExpressionConvertType : uint8;
+
+////////////////////////////////////////////////////////////////////////////////////////////////////////////////////////////////////////////
+// FConvertDragDropOp: Drag and Drop Operation used to form connections within the convert node
+////////////////////////////////////////////////////////////////////////////////////////////////////////////////////////////////////////////
+
+class FConvertDragDropOp : public FDragDropOperation
+{
+public:
+	DRAG_DROP_OPERATOR_TYPE(FConvertDragDropOp, FDragDropOperation)
+
+		FConvertDragDropOp(TSharedPtr<SConvertInnerPin> InSourcePin);
+
+	//~ Begin FDragDropOperation interface
+	virtual void OnDrop(bool bDropWasHandled, const FPointerEvent& MouseEvent) override;
+	virtual void OnDragged(const class FDragDropEvent& DragDropEvent) override;
+	//~ End FDragDropOperation interface
+
+	TSharedPtr<SConvertInnerPin> GetSourcePin() const;
+	UE::Slate::FDeprecateVector2DResult GetScreenPosition() const;
+
+protected:
+	TSharedPtr<SConvertInnerPin> SourcePin;
+	FVector2f ScreenPosition;
+};
+
+////////////////////////////////////////////////////////////////////////////////////////////////////////////////////////////////////////////
+// SMaterialExpressionConvertGraphPin: The Outer pins that form connections to other material graph nodes
+////////////////////////////////////////////////////////////////////////////////////////////////////////////////////////////////////////////
+
+class SMaterialExpressionConvertGraphPin : public SGraphPin
+{
+public:
+	SLATE_BEGIN_ARGS(SMaterialExpressionConvertGraphPin){}
+	SLATE_END_ARGS()
+
+	void Construct(const FArguments& InArgs, UEdGraphPin* InPin);
+
+	void CreateInnerPins(TSharedRef<SGraphNodeMaterialConvert> InOwningGraphNodeWidget);
+
+	const TArray<TSharedRef<SConvertInnerPin>>& GetInnerPins() const { return InnerPins; }
+
+	bool IsHoveredOverPrimaryPin() const;
+private:
+
+	// Inner Pins, one for each component
+	TArray<TSharedRef<SConvertInnerPin>> InnerPins;
+};
+
+////////////////////////////////////////////////////////////////////////////////////////////////////////////////////////////////////////////
+// SConvertInnerPin: The inner pins used to route values with the convert node
+////////////////////////////////////////////////////////////////////////////////////////////////////////////////////////////////////////////
+
+/** Class for the inner pins we'll use to route values in SGraphNodeMaterialConvert */
+class SConvertInnerPin : public SCompoundWidget
+{
+public:
+	SLATE_BEGIN_ARGS(SConvertInnerPin) {}
+	SLATE_END_ARGS()
+
+	void Construct(
+		const FArguments& InArgs,
+		TSharedPtr<SGraphNodeMaterialConvert> InOwningNode,
+		TSharedPtr<SMaterialExpressionConvertGraphPin> InOwningPin,
+		bool bInIsInputPin, 
+		int32 InPinIndex, 
+		int32 InComponentIndex
+	);
+
+	//~ Begin SWidget interface
+	virtual int32 OnPaint(
+		const FPaintArgs& Args, 
+		const FGeometry& AllottedGeometry, 
+		const FSlateRect& MyCullingRect, 
+		FSlateWindowElementList& OutDrawElements, 
+		int32 LayerId, 
+		const FWidgetStyle& InWidgetStyle, 
+		bool bParentEnabled
+	) const override;
+	virtual FReply OnMouseButtonDown(const FGeometry& MyGeometry, const FPointerEvent& MouseEvent) override;
+	virtual FReply OnMouseButtonUp(const FGeometry& MyGeometry, const FPointerEvent& MouseEvent) override;
+	virtual FReply OnDragDetected(const FGeometry& MyGeometry, const FPointerEvent& MouseEvent) override;
+	virtual FReply OnDrop(const FGeometry& MyGeometry, const FDragDropEvent& DragDropEvent) override;
+	//~ End SWidget interface
+
+	const FSlateBrush* GetPinBrush() const;
+	void CancelDragDrop();
+	UE::Slate::FDeprecateVector2DResult GetPinCenterAbsolute() const;
+
+	/** Adds a connection to WeakConnectedPins */
+	void AddConnection(TSharedPtr<SConvertInnerPin> InOtherPin);
+
+	/** Removes a connection from WeakConnectedPins */
+	void RemoveConnection(TSharedPtr<SConvertInnerPin> InOtherPin);
+
+	/** Clears WeakConnectedPins */
+	void RemoveAllConnections();
+
+	/** Calls into the owning node to break connections and refresh the material node */
+	void BreakConnections();
+
+	EVisibility GetDefaultValueVisibility() const;
+	TOptional<float> GetDefaultValue() const;
+	void SetDefaultValue(float InDefaultValue, ETextCommit::Type CommitType);
+
+	EVisibility GetPinNameVisibility() const;
+	FText GetPinName() const;
+
+	bool IsInputPin() const { return bIsInputPin; }
+	int32 GetPinIndex() const { return PinIndex; }
+	int32 GetComponentIndex() const { return ComponentIndex; }
+
+protected:
+	// Input Data
+	TWeakPtr<SGraphNodeMaterialConvert> WeakOwningNode;
+	TWeakPtr<SMaterialExpressionConvertGraphPin> WeakOwningPin;
+	bool bIsInputPin = false;
+	int32 PinIndex = INDEX_NONE;
+	int32 ComponentIndex = INDEX_NONE;
+
+	TSharedPtr<SImage> PinImage;
+	TSet<TWeakPtr<SConvertInnerPin>> WeakConnectedPins;
+	mutable FVector2f CenterAbsolute;
+};
+
+////////////////////////////////////////////////////////////////////////////////////////////////////////////////////////////////////////////
+// SGraphNodeMaterialConvert: Custom Slate widget for UMaterialExpressionConvert
+////////////////////////////////////////////////////////////////////////////////////////////////////////////////////////////////////////////
+
+class MATERIALEDITOR_API SGraphNodeMaterialConvert : public SGraphNodeMaterialBase
+{
+public:
+	SLATE_BEGIN_ARGS(SGraphNodeMaterialConvert){}
+	SLATE_END_ARGS()
+
+	void Construct(const FArguments& InArgs, class UMaterialGraphNode* InNode);
+
+	// Begin SGraphNode interface
+	virtual void CreatePinWidgets() override;
+	virtual TSharedPtr<SGraphPin> CreatePinWidget(UEdGraphPin* Pin) const override;
+	virtual void AddPin(const TSharedRef<SGraphPin>& PinToAdd) override;
+
+	virtual void CreateInputSideAddButton(TSharedPtr<SVerticalBox> InputBox) override;
+	virtual void CreateOutputSideAddButton(TSharedPtr<SVerticalBox> OutputBox) override;
+	virtual EVisibility IsAddPinButtonVisible() const override;
+	// End SGraphNode interface
+
+	FReply OnAddInputPinClicked();
+	FReply OnAddOutputPinClicked();
+
+	TSharedRef<SWidget> CreateAddPinContextMenu(bool bInputPin);
+	void AddNewPin(bool bInputPin, EMaterialExpressionConvertType ConvertType);
+
+	// Begin SWidget Interface
+	virtual int32 OnPaint(
+		const FPaintArgs& Args, 
+		const FGeometry& AllottedGeometry, 
+		const FSlateRect& MyCullingRect, 
+		FSlateWindowElementList& OutDrawElements, 
+		int32 LayerId, 
+		const FWidgetStyle& InWidgetStyle, 
+		bool bParentEnabled
+	) const override;
+	// End SWidget Interface
+
+	void FormConnection(TSharedPtr<SConvertInnerPin> InnerPinA, TSharedPtr<SConvertInnerPin> InnerPinB);
+	void BreakConnections(TSharedPtr<SConvertInnerPin> InnerPin);
+
+	TOptional<float> GetDefaultValue(TSharedPtr<const SConvertInnerPin> InnerPin) const;
+	void SetDefaultValue(TSharedPtr<SConvertInnerPin> InnerPin, const float InDefaultValue);
+
+	void SetCurrentDragDropOp(TSharedPtr<FConvertDragDropOp> InDragDropOp) { CurrentDragDropOp = InDragDropOp; }
+	TSharedPtr<FConvertDragDropOp> GetCurrentDragDropOp() const { return CurrentDragDropOp; }
+
+protected:
+	TSharedPtr<SConvertInnerPin> GetInnerPin(bool bInputPin, const int32 InPinIndex, const int32 InComponentIndex) const;
+	void MakeConnectionCurve(
+		const FGeometry& InAllottedGeometry,
+		FSlateWindowElementList& OutDrawElements,
+		int32& InOutLayerId, 
+		const FVector2f& InCurveStart, 
+		const FVector2f& InCurveEnd
+	) const;
+
+private:
+	TSharedPtr<FConvertDragDropOp> CurrentDragDropOp;
+};
\ No newline at end of file
diff --git a/Engine/Source/Editor/MaterialEditor/Public/MaterialNodes/SGraphNodeMaterialCustom.h b/Engine/Source/Editor/MaterialEditor/Public/MaterialNodes/SGraphNodeMaterialCustom.h
new file mode 100644
--- /dev/null
+++ b/Engine/Source/Editor/MaterialEditor/Public/MaterialNodes/SGraphNodeMaterialCustom.h
@@ -0,0 +1,39 @@
+// Copyright Epic Games, Inc. All Rights Reserved.
+
+#pragma once
+
+#include "MaterialNodes/SGraphNodeMaterialBase.h"
+#include "Text/HLSLSyntaxHighlighterMarshaller.h"
+#include "Widgets/Input/SMultiLineEditableTextBox.h"
+
+class MATERIALEDITOR_API SGraphNodeMaterialCustom : public SGraphNodeMaterialBase
+{
+public:
+	SLATE_BEGIN_ARGS(SGraphNodeMaterialCustom){}
+	SLATE_END_ARGS()
+
+	void Construct(const FArguments& InArgs, class UMaterialGraphNode* InNode);
+
+	UMaterialGraphNode* GetMaterialGraphNode() const {return MaterialNode;}
+
+	FText GetHlslText() const;
+	void OnCustomHlslTextCommited(const FText& InText, ETextCommit::Type InType);
+
+protected:
+	// SGraphNode interface
+	virtual void CreateBelowPinControls(TSharedPtr<SVerticalBox> MainBox) override;
+	virtual void CreateAdvancedViewArrow(TSharedPtr<SVerticalBox> MainBox) override;
+	virtual EVisibility AdvancedViewArrowVisibility() const override;
+	virtual void OnAdvancedViewChanged( const ECheckBoxState NewCheckedState ) override;
+	virtual ECheckBoxState IsAdvancedViewChecked() const override;
+	virtual const FSlateBrush* GetAdvancedViewArrow() const override;
+	// End of SGraphNode interface
+
+private:
+	EVisibility CodeVisibility() const;
+
+protected:
+	/** Shader text box */
+	TSharedPtr<FHLSLSyntaxHighlighterMarshaller> SyntaxHighlighter;
+	TSharedPtr<SMultiLineEditableTextBox> ShaderTextBox;
+};
diff --git a/Engine/Source/Editor/MaterialEditor/Public/MaterialNodes/SGraphNodeMaterialResult.h b/Engine/Source/Editor/MaterialEditor/Public/MaterialNodes/SGraphNodeMaterialResult.h
new file mode 100644
--- /dev/null
+++ b/Engine/Source/Editor/MaterialEditor/Public/MaterialNodes/SGraphNodeMaterialResult.h
@@ -0,0 +1,32 @@
+// Copyright Epic Games, Inc. All Rights Reserved.
+
+#pragma once
+
+#include "Math/Vector2D.h"
+#include "SGraphNode.h"
+#include "SNodePanel.h"
+#include "Widgets/DeclarativeSyntaxSupport.h"
+
+class FGraphNodeMetaData;
+class UMaterialGraphNode_Root;
+
+class MATERIALEDITOR_API SGraphNodeMaterialResult : public SGraphNode
+{
+public:
+	SLATE_BEGIN_ARGS(SGraphNodeMaterialResult){}
+	SLATE_END_ARGS()
+
+	void Construct(const FArguments& InArgs, class UMaterialGraphNode_Root* InNode);
+
+	// SGraphNode interface
+	virtual void CreatePinWidgets() override;
+	virtual void PopulateMetaTag(FGraphNodeMetaData* TagMeta) const override;
+	// End of SGraphNode interface
+
+	// SNodePanel::SNode interface
+	virtual void MoveTo(const FVector2f& NewPosition, FNodeSet& NodeFilter, bool bMarkDirty = true) override;
+	// End of SNodePanel::SNode interface
+
+private:
+	UMaterialGraphNode_Root* RootNode;
+};
diff --git a/Engine/Source/Editor/MaterialEditor/Public/MaterialPins/SGraphPinMaterialInput.h b/Engine/Source/Editor/MaterialEditor/Public/MaterialPins/SGraphPinMaterialInput.h
new file mode 100644
--- /dev/null
+++ b/Engine/Source/Editor/MaterialEditor/Public/MaterialPins/SGraphPinMaterialInput.h
@@ -0,0 +1,25 @@
+// Copyright Epic Games, Inc. All Rights Reserved.
+
+#pragma once
+
+#include "CoreMinimal.h"
+#include "SGraphPin.h"
+#include "Styling/SlateColor.h"
+#include "Widgets/DeclarativeSyntaxSupport.h"
+
+class UEdGraphPin;
+
+class MATERIALEDITOR_API SGraphPinMaterialInput : public SGraphPin
+{
+public:
+	SLATE_BEGIN_ARGS(SGraphPinMaterialInput) {}
+	SLATE_END_ARGS()
+
+	void Construct(const FArguments& InArgs, UEdGraphPin* InGraphPinObj);
+
+protected:
+	//~ Begin SGraphPin Interface
+	virtual FSlateColor GetPinColor() const override;
+	//~ End SGraphPin Interface
+
+};
diff --git a/Engine/Source/Editor/MaterialEditor/Public/SGraphSubstrateMaterial.h b/Engine/Source/Editor/MaterialEditor/Public/SGraphSubstrateMaterial.h
new file mode 100644
--- /dev/null
+++ b/Engine/Source/Editor/MaterialEditor/Public/SGraphSubstrateMaterial.h
@@ -0,0 +1,20 @@
+// Copyright Epic Games, Inc. All Rights Reserved.
+
+#pragma once
+
+#include "Internationalization/Text.h"
+#include "Rendering/RenderingCommon.h"
+#include "SGraphPin.h"
+
+class UMaterialGraphNode;
+struct FSubstrateMaterialCompilationOutput;
+
+struct MATERIALEDITOR_API FSubstrateWidget
+{
+	static const TSharedRef<SWidget> ProcessOperator(const FSubstrateMaterialCompilationOutput& CompilationOutput);
+	static const TSharedRef<SWidget> ProcessOperator(const FSubstrateMaterialCompilationOutput& CompilationOutput, const TArray<FGuid>& InGuid);
+	static void GetPinColor(TSharedPtr<SGraphPin>& Out, const UMaterialGraphNode* InNode);
+	static FLinearColor GetConnectionColor();
+	static bool HasInputSubstrateType(const UEdGraphPin* InPin);
+	static bool HasOutputSubstrateType(const UEdGraphPin* InPin);
+};
\ No newline at end of file
diff --git a/Engine/Source/Editor/UnrealEd/Classes/MaterialGraph/MaterialGraphNode_Comment.h b/Engine/Source/Editor/UnrealEd/Classes/MaterialGraph/MaterialGraphNode_Comment.h
--- a/Engine/Source/Editor/UnrealEd/Classes/MaterialGraph/MaterialGraphNode_Comment.h
+++ b/Engine/Source/Editor/UnrealEd/Classes/MaterialGraph/MaterialGraphNode_Comment.h
@@ -38,6 +38,7 @@ class UMaterialGraphNode_Comment : public UEdGraphNode_Comment
 	virtual void OnRenameNode(const FString& NewName) override;
 	virtual void ResizeNode(const FVector2f& NewSize) override;
 	virtual int32 GetFontSize() const override;
+	virtual TSharedPtr<SGraphNode> CreateVisualWidget() override;
 	//~ End UEdGraphNode Interface.
 
 private:
diff --git a/Engine/Source/Editor/UnrealEd/Private/MaterialGraphNode_Comment.cpp b/Engine/Source/Editor/UnrealEd/Private/MaterialGraphNode_Comment.cpp
--- a/Engine/Source/Editor/UnrealEd/Private/MaterialGraphNode_Comment.cpp
+++ b/Engine/Source/Editor/UnrealEd/Private/MaterialGraphNode_Comment.cpp
@@ -8,6 +8,7 @@
 #include "ToolMenus.h"
 #include "MaterialGraph/MaterialGraphSchema.h"
 #include "Materials/MaterialExpressionComment.h"
+#include "MaterialNodes/SGraphNodeMaterialComment.h"
 #include "Materials/MaterialFunction.h"
 #include "Framework/Commands/GenericCommands.h"
 #include "GraphEditorActions.h"
@@ -62,6 +63,11 @@ void UMaterialGraphNode_Comment::PostDuplicate(bool bDuplicateForPIE)
 	}
 }
 
+TSharedPtr<SGraphNode> UMaterialGraphNode_Comment::CreateVisualWidget()
+{
+	return SNew(SGraphNodeMaterialComment, this);
+}
+
 void UMaterialGraphNode_Comment::PrepareForCopying()
 {
 	if (MaterialExpressionComment)
diff --git a/Engine/Source/Runtime/Engine/Engine.Build.cs b/Engine/Source/Runtime/Engine/Engine.Build.cs
--- a/Engine/Source/Runtime/Engine/Engine.Build.cs
+++ b/Engine/Source/Runtime/Engine/Engine.Build.cs
@@ -329,6 +329,8 @@ public class Engine : ModuleRules
 			PrivateIncludePathModuleNames.Add("GraphEditor");
 			PrivateDependencyModuleNames.Add("GraphEditor");
 
+			PrivateDependencyModuleNames.Add("MaterialEditor");
+
 			if (Target.Platform == UnrealTargetPlatform.Win64)
 			{
 				DynamicallyLoadedModuleNames.AddRange(
