Add header shims with deprecation warnings for material headers that were moved from GraphEditor module to MaterialEditor module.

- Remove unused header in TextureGraph that was being included.
- Move all public headers within the MaterialEditor module.  They moved from /MaterialEditor/Public/MaterialNodes to /MaterialEditor/Public/MaterialEditor/MaterialNodes to avoid ambiguity.
- Fixed up all includes to account for the new path.

[RN] major | Rendering - Materials and Shaders

diff --git a/Engine/Plugins/TextureGraph/Source/TextureGraphEditor/Private/EdGraph/TG_EditorGraphPanelPinFactory.cpp b/Engine/Plugins/TextureGraph/Source/TextureGraphEditor/Private/EdGraph/TG_EditorGraphPanelPinFactory.cpp
--- a/Engine/Plugins/TextureGraph/Source/TextureGraphEditor/Private/EdGraph/TG_EditorGraphPanelPinFactory.cpp
+++ b/Engine/Plugins/TextureGraph/Source/TextureGraphEditor/Private/EdGraph/TG_EditorGraphPanelPinFactory.cpp
@@ -20,7 +20,6 @@
 #include "KismetPins/SGraphPinVector4Slider.h"
 #include "KismetPins/SGraphPinObject.h"
 #include "MaterialGraph/MaterialGraphSchema.h"
-#include "MaterialPins/SGraphPinMaterialInput.h"
 #include "Misc/Optional.h"
 #include "SGraphPin.h"
 #include "TG_EdGraphSchema.h"
diff --git a/Engine/Source/Editor/GraphEditor/Public/MaterialNodes/SGraphNodeMaterialBase.h b/Engine/Source/Editor/GraphEditor/Public/MaterialNodes/SGraphNodeMaterialBase.h
new file mode 100644
--- /dev/null
+++ b/Engine/Source/Editor/GraphEditor/Public/MaterialNodes/SGraphNodeMaterialBase.h
@@ -0,0 +1,10 @@
+// Copyright Epic Games, Inc. All Rights Reserved.
+
+#pragma once
+
+#include "Misc/CoreMiscDefines.h"
+
+UE_DEPRECATED_HEADER(5.6, "This header has moved to the MaterialEditor module. Use MaterialEditor/MaterialNodes/SGraphNodeMaterialBase.h.")
+
+// Include the new file so that the project still compiles, but has warnings
+#include "MaterialEditor/MaterialNodes/SGraphNodeMaterialBase.h"
\ No newline at end of file
diff --git a/Engine/Source/Editor/GraphEditor/Public/MaterialNodes/SGraphNodeMaterialComposite.h b/Engine/Source/Editor/GraphEditor/Public/MaterialNodes/SGraphNodeMaterialComposite.h
new file mode 100644
--- /dev/null
+++ b/Engine/Source/Editor/GraphEditor/Public/MaterialNodes/SGraphNodeMaterialComposite.h
@@ -0,0 +1,10 @@
+// Copyright Epic Games, Inc. All Rights Reserved.
+
+#pragma once
+
+#include "Misc/CoreMiscDefines.h"
+
+UE_DEPRECATED_HEADER(5.6, "This header has moved to the MaterialEditor module. Use MaterialEditor/MaterialNodes/SGraphNodeMaterialComposite.h.")
+
+// Include the new file so that the project still compiles, but has warnings
+#include "MaterialEditor/MaterialNodes/SGraphNodeMaterialComposite.h"
\ No newline at end of file
diff --git a/Engine/Source/Editor/GraphEditor/Public/MaterialNodes/SGraphNodeMaterialConvert.h b/Engine/Source/Editor/GraphEditor/Public/MaterialNodes/SGraphNodeMaterialConvert.h
new file mode 100644
--- /dev/null
+++ b/Engine/Source/Editor/GraphEditor/Public/MaterialNodes/SGraphNodeMaterialConvert.h
@@ -0,0 +1,10 @@
+// Copyright Epic Games, Inc. All Rights Reserved.
+
+#pragma once
+
+#include "Misc/CoreMiscDefines.h"
+
+UE_DEPRECATED_HEADER(5.6, "This header has moved to the MaterialEditor module. Use MaterialEditor/MaterialNodes/SGraphNodeMaterialConvert.h.")
+
+// Include the new file so that the project still compiles, but has warnings
+#include "MaterialEditor/MaterialNodes/SGraphNodeMaterialConvert.h"
\ No newline at end of file
diff --git a/Engine/Source/Editor/GraphEditor/Public/MaterialNodes/SGraphNodeMaterialCustom.h b/Engine/Source/Editor/GraphEditor/Public/MaterialNodes/SGraphNodeMaterialCustom.h
new file mode 100644
--- /dev/null
+++ b/Engine/Source/Editor/GraphEditor/Public/MaterialNodes/SGraphNodeMaterialCustom.h
@@ -0,0 +1,10 @@
+// Copyright Epic Games, Inc. All Rights Reserved.
+
+#pragma once
+
+#include "Misc/CoreMiscDefines.h"
+
+UE_DEPRECATED_HEADER(5.6, "This header has moved to the MaterialEditor module. Use MaterialEditor/MaterialNodes/SGraphNodeMaterialCustom.h.")
+
+// Include the new file so that the project still compiles, but has warnings
+#include "MaterialEditor/MaterialNodes/SGraphNodeMaterialCustom.h"
\ No newline at end of file
diff --git a/Engine/Source/Editor/GraphEditor/Public/MaterialNodes/SGraphNodeMaterialResult.h b/Engine/Source/Editor/GraphEditor/Public/MaterialNodes/SGraphNodeMaterialResult.h
new file mode 100644
--- /dev/null
+++ b/Engine/Source/Editor/GraphEditor/Public/MaterialNodes/SGraphNodeMaterialResult.h
@@ -0,0 +1,10 @@
+// Copyright Epic Games, Inc. All Rights Reserved.
+
+#pragma once
+
+#include "Misc/CoreMiscDefines.h"
+
+UE_DEPRECATED_HEADER(5.6, "This header has moved to the MaterialEditor module. Use MaterialEditor/MaterialNodes/SGraphNodeMaterialResult.h.")
+
+// Include the new file so that the project still compiles, but has warnings
+#include "MaterialEditor/MaterialNodes/SGraphNodeMaterialResult.h"
diff --git a/Engine/Source/Editor/GraphEditor/Public/MaterialPins/SGraphPinMaterialInput.h b/Engine/Source/Editor/GraphEditor/Public/MaterialPins/SGraphPinMaterialInput.h
new file mode 100644
--- /dev/null
+++ b/Engine/Source/Editor/GraphEditor/Public/MaterialPins/SGraphPinMaterialInput.h
@@ -0,0 +1,10 @@
+// Copyright Epic Games, Inc. All Rights Reserved.
+
+#pragma once
+
+#include "Misc/CoreMiscDefines.h"
+
+UE_DEPRECATED_HEADER(5.6, "This header has moved to the MaterialEditor module. Use MaterialEditor/MaterialPins/SGraphPinMaterialInput.h.")
+
+// Include the new file so that the project still compiles, but has warnings
+#include "MaterialEditor/MaterialPins/SGraphPinMaterialInput.h"
diff --git a/Engine/Source/Editor/GraphEditor/Public/SGraphSubstrateMaterial.h b/Engine/Source/Editor/GraphEditor/Public/SGraphSubstrateMaterial.h
new file mode 100644
--- /dev/null
+++ b/Engine/Source/Editor/GraphEditor/Public/SGraphSubstrateMaterial.h
@@ -0,0 +1,10 @@
+// Copyright Epic Games, Inc. All Rights Reserved.
+
+#pragma once
+
+#include "Misc/CoreMiscDefines.h"
+
+UE_DEPRECATED_HEADER(5.6, "This header has moved to the MaterialEditor module. Use MaterialEditor/SGraphSubstrateMaterial.h.")
+
+// Include the new file so that the project still compiles, but has warnings
+#include "MaterialEditor/SGraphSubstrateMaterial.h"
diff --git a/Engine/Source/Editor/MaterialEditor/Private/MaterialEditor.cpp b/Engine/Source/Editor/MaterialEditor/Private/MaterialEditor.cpp
--- a/Engine/Source/Editor/MaterialEditor/Private/MaterialEditor.cpp
+++ b/Engine/Source/Editor/MaterialEditor/Private/MaterialEditor.cpp
@@ -166,7 +166,7 @@
 #include "MaterialEditorModes.h"
 #include "Materials/MaterialExpression.h"
 #include "SMaterialEditorSubstrateWidget.h"
-#include "SGraphSubstrateMaterial.h"
+#include "MaterialEditor/SGraphSubstrateMaterial.h"
 
 #include "SMaterialParametersOverviewWidget.h"
 #include "SMaterialEditorCustomPrimitiveDataWidget.h"
diff --git a/Engine/Source/Editor/MaterialEditor/Private/MaterialEditorGraphPanelPinFactory.cpp b/Engine/Source/Editor/MaterialEditor/Private/MaterialEditorGraphPanelPinFactory.cpp
--- a/Engine/Source/Editor/MaterialEditor/Private/MaterialEditorGraphPanelPinFactory.cpp
+++ b/Engine/Source/Editor/MaterialEditor/Private/MaterialEditorGraphPanelPinFactory.cpp
@@ -10,7 +10,7 @@
 #include "KismetPins/SGraphPinColor.h"
 #include "KismetPins/SGraphPinExec.h"
 #include "MaterialGraph/MaterialGraphSchema.h"
-#include "MaterialPins/SGraphPinMaterialInput.h"
+#include "MaterialEditor/MaterialPins/SGraphPinMaterialInput.h"
 #include "Misc/Optional.h"
 #include "SGraphPin.h"
 #include "Templates/Casts.h"
diff --git a/Engine/Source/Editor/MaterialEditor/Private/MaterialGraphConnectionDrawingPolicy.cpp b/Engine/Source/Editor/MaterialEditor/Private/MaterialGraphConnectionDrawingPolicy.cpp
--- a/Engine/Source/Editor/MaterialEditor/Private/MaterialGraphConnectionDrawingPolicy.cpp
+++ b/Engine/Source/Editor/MaterialEditor/Private/MaterialGraphConnectionDrawingPolicy.cpp
@@ -17,7 +17,7 @@
 #include "Templates/Casts.h"
 #include "Templates/SharedPointer.h"
 #include "UObject/NameTypes.h"
-#include "SGraphSubstrateMaterial.h"
+#include "MaterialEditor/SGraphSubstrateMaterial.h"
 #include "Widgets/SToolTip.h"
 #include "Widgets/SBoxPanel.h"
 #include "Widgets/Text/STextBlock.h"
diff --git a/Engine/Source/Editor/MaterialEditor/Private/MaterialNodes/SGraphNodeMaterialBase.cpp b/Engine/Source/Editor/MaterialEditor/Private/MaterialNodes/SGraphNodeMaterialBase.cpp
--- a/Engine/Source/Editor/MaterialEditor/Private/MaterialNodes/SGraphNodeMaterialBase.cpp
+++ b/Engine/Source/Editor/MaterialEditor/Private/MaterialNodes/SGraphNodeMaterialBase.cpp
@@ -1,6 +1,6 @@
 // Copyright Epic Games, Inc. All Rights Reserved.
 
-#include "MaterialNodes/SGraphNodeMaterialBase.h"
+#include "MaterialEditor/MaterialNodes/SGraphNodeMaterialBase.h"
 
 #include "CanvasItem.h"
 #include "CanvasTypes.h"
@@ -64,7 +64,7 @@
 #include "Widgets/SViewport.h"
 #include "Widgets/Text/STextBlock.h"
 #include "Rendering/SubstrateMaterialShared.h"
-#include "SGraphSubstrateMaterial.h"
+#include "MaterialEditor/SGraphSubstrateMaterial.h"
 #include "MaterialShared.h"
 #include "RenderGraphBuilder.h"
 
diff --git a/Engine/Source/Editor/MaterialEditor/Private/MaterialNodes/SGraphNodeMaterialComposite.cpp b/Engine/Source/Editor/MaterialEditor/Private/MaterialNodes/SGraphNodeMaterialComposite.cpp
--- a/Engine/Source/Editor/MaterialEditor/Private/MaterialNodes/SGraphNodeMaterialComposite.cpp
+++ b/Engine/Source/Editor/MaterialEditor/Private/MaterialNodes/SGraphNodeMaterialComposite.cpp
@@ -1,6 +1,6 @@
 // Copyright Epic Games, Inc. All Rights Reserved.
 
-#include "MaterialNodes/SGraphNodeMaterialComposite.h"
+#include "MaterialEditor/MaterialNodes/SGraphNodeMaterialComposite.h"
 
 #include "Containers/Array.h"
 #include "Delegates/Delegate.h"
diff --git a/Engine/Source/Editor/MaterialEditor/Private/MaterialNodes/SGraphNodeMaterialConvert.cpp b/Engine/Source/Editor/MaterialEditor/Private/MaterialNodes/SGraphNodeMaterialConvert.cpp
--- a/Engine/Source/Editor/MaterialEditor/Private/MaterialNodes/SGraphNodeMaterialConvert.cpp
+++ b/Engine/Source/Editor/MaterialEditor/Private/MaterialNodes/SGraphNodeMaterialConvert.cpp
@@ -1,6 +1,6 @@
 // Copyright Epic Games, Inc. All Rights Reserved.
 
-#include "MaterialNodes/SGraphNodeMaterialConvert.h"
+#include "MaterialEditor/MaterialNodes/SGraphNodeMaterialConvert.h"
 
 #include "ConnectionDrawingPolicy.h"
 #include "Framework/MultiBox/MultiBoxBuilder.h"
diff --git a/Engine/Source/Editor/MaterialEditor/Private/MaterialNodes/SGraphNodeMaterialCustom.cpp b/Engine/Source/Editor/MaterialEditor/Private/MaterialNodes/SGraphNodeMaterialCustom.cpp
--- a/Engine/Source/Editor/MaterialEditor/Private/MaterialNodes/SGraphNodeMaterialCustom.cpp
+++ b/Engine/Source/Editor/MaterialEditor/Private/MaterialNodes/SGraphNodeMaterialCustom.cpp
@@ -1,6 +1,6 @@
 // Copyright Epic Games, Inc. All Rights Reserved.
 
-#include "MaterialNodes/SGraphNodeMaterialCustom.h"
+#include "MaterialEditor/MaterialNodes/SGraphNodeMaterialCustom.h"
 
 #include "Templates/Casts.h"
 #include "Widgets/Input/SCheckBox.h"
diff --git a/Engine/Source/Editor/MaterialEditor/Private/MaterialNodes/SGraphNodeMaterialResult.cpp b/Engine/Source/Editor/MaterialEditor/Private/MaterialNodes/SGraphNodeMaterialResult.cpp
--- a/Engine/Source/Editor/MaterialEditor/Private/MaterialNodes/SGraphNodeMaterialResult.cpp
+++ b/Engine/Source/Editor/MaterialEditor/Private/MaterialNodes/SGraphNodeMaterialResult.cpp
@@ -1,6 +1,6 @@
 // Copyright Epic Games, Inc. All Rights Reserved.
 
-#include "MaterialNodes/SGraphNodeMaterialResult.h"
+#include "MaterialEditor/MaterialNodes/SGraphNodeMaterialResult.h"
 
 #include "Containers/Array.h"
 #include "Containers/UnrealString.h"
diff --git a/Engine/Source/Editor/MaterialEditor/Private/MaterialNodes/SGraphSubstrateMaterial.cpp b/Engine/Source/Editor/MaterialEditor/Private/MaterialNodes/SGraphSubstrateMaterial.cpp
--- a/Engine/Source/Editor/MaterialEditor/Private/MaterialNodes/SGraphSubstrateMaterial.cpp
+++ b/Engine/Source/Editor/MaterialEditor/Private/MaterialNodes/SGraphSubstrateMaterial.cpp
@@ -1,6 +1,6 @@
 // Copyright Epic Games, Inc. All Rights Reserved.
 
-#include "SGraphSubstrateMaterial.h"
+#include "MaterialEditor/SGraphSubstrateMaterial.h"
 #include "Internationalization/Text.h"
 #include "Rendering/RenderingCommon.h"
 #include "Rendering/SubstrateMaterialShared.h"
diff --git a/Engine/Source/Editor/MaterialEditor/Private/MaterialPins/SGraphPinMaterialInput.cpp b/Engine/Source/Editor/MaterialEditor/Private/MaterialPins/SGraphPinMaterialInput.cpp
--- a/Engine/Source/Editor/MaterialEditor/Private/MaterialPins/SGraphPinMaterialInput.cpp
+++ b/Engine/Source/Editor/MaterialEditor/Private/MaterialPins/SGraphPinMaterialInput.cpp
@@ -1,7 +1,7 @@
 // Copyright Epic Games, Inc. All Rights Reserved.
 
 
-#include "MaterialPins/SGraphPinMaterialInput.h"
+#include "MaterialEditor/MaterialPins/SGraphPinMaterialInput.h"
 
 #include "EdGraph/EdGraphNode.h"
 #include "EdGraph/EdGraphPin.h"
@@ -9,7 +9,7 @@
 #include "MaterialGraph/MaterialGraphSchema.h"
 #include "Misc/AssertionMacros.h"
 #include "Templates/Casts.h"
-#include "SGraphSubstrateMaterial.h"
+#include "MaterialEditor/SGraphSubstrateMaterial.h"
 #include "MaterialValueType.h"
 
 void SGraphPinMaterialInput::Construct(const FArguments& InArgs, UEdGraphPin* InGraphPinObj)
diff --git a/Engine/Source/Editor/MaterialEditor/Private/SMaterialEditorSubstrateWidget.cpp b/Engine/Source/Editor/MaterialEditor/Private/SMaterialEditorSubstrateWidget.cpp
--- a/Engine/Source/Editor/MaterialEditor/Private/SMaterialEditorSubstrateWidget.cpp
+++ b/Engine/Source/Editor/MaterialEditor/Private/SMaterialEditorSubstrateWidget.cpp
@@ -15,7 +15,7 @@
 #include "Widgets/Layout/SBorder.h"
 #include "Widgets/SNullWidget.h"
 #include "Widgets/Notifications/SErrorText.h"
-#include "SGraphSubstrateMaterial.h"
+#include "MaterialEditor/SGraphSubstrateMaterial.h"
 #include "RHIShaderPlatform.h"
 
 #define LOCTEXT_NAMESPACE "SMaterialEditorSubstrateWidget"
diff --git a/Engine/Source/Editor/MaterialEditor/Private/SMaterialEditorTopologyWidget.cpp b/Engine/Source/Editor/MaterialEditor/Private/SMaterialEditorTopologyWidget.cpp
--- a/Engine/Source/Editor/MaterialEditor/Private/SMaterialEditorTopologyWidget.cpp
+++ b/Engine/Source/Editor/MaterialEditor/Private/SMaterialEditorTopologyWidget.cpp
@@ -19,7 +19,7 @@
 #include "Widgets/SNullWidget.h"
 #include "Widgets/Notifications/SErrorText.h"
 #include "AssetThumbnail.h"
-#include "SGraphSubstrateMaterial.h"
+#include "MaterialEditor/SGraphSubstrateMaterial.h"
 #include "RHIShaderPlatform.h"
 
 #define LOCTEXT_NAMESPACE "SMaterialEditorTopologyWidget"
diff --git a/Engine/Source/Editor/MaterialEditor/Public/MaterialEditor/MaterialNodes/SGraphNodeMaterialBase.h b/Engine/Source/Editor/MaterialEditor/Public/MaterialEditor/MaterialNodes/SGraphNodeMaterialBase.h
new file mode 100644
--- /dev/null
+++ b/Engine/Source/Editor/MaterialEditor/Public/MaterialEditor/MaterialNodes/SGraphNodeMaterialBase.h
@@ -0,0 +1,148 @@
+// Copyright Epic Games, Inc. All Rights Reserved.
+
+#pragma once
+
+#include "Internationalization/Text.h"
+#include "Layout/Visibility.h"
+#include "Math/UnrealMathSSE.h"
+#include "Math/Vector2D.h"
+#include "Rendering/RenderingCommon.h"
+#include "SGraphNode.h"
+#include "SNodePanel.h"
+#include "Styling/SlateTypes.h"
+#include "Templates/SharedPointer.h"
+#include "Textures/SlateShaderResource.h"
+#include "Widgets/DeclarativeSyntaxSupport.h"
+
+class FMaterialRenderProxy;
+class FRHICommandListImmediate;
+class FSlateRect;
+class FWidgetStyle;
+class SGraphPin;
+class SOverlay;
+class SVerticalBox;
+class SWidget;
+class UMaterialGraphNode;
+struct FGeometry;
+struct FSlateBrush;
+
+typedef TSharedPtr<class FPreviewElement, ESPMode::ThreadSafe> FThreadSafePreviewPtr;
+
+class FPreviewViewport : public ISlateViewport
+{
+public:
+	FPreviewViewport(class UMaterialGraphNode* InNode);
+	~FPreviewViewport();
+
+	// ISlateViewport interface
+	virtual void OnDrawViewport( const FGeometry& AllottedGeometry, const FSlateRect& MyCullingRect, class FSlateWindowElementList& OutDrawElements, int32 LayerId, const FWidgetStyle& InWidgetStyle, bool bParentEnabled ) override;
+	virtual FIntPoint GetSize() const override;
+	virtual class FSlateShaderResource* GetViewportRenderTargetTexture() const override {return NULL;}
+	virtual bool RequiresVsync() const override {return false;}
+
+	/** Material node to get expression preview from */
+	UMaterialGraphNode* MaterialNode;
+	/** Custom Slate Element to display the preview */
+	FThreadSafePreviewPtr PreviewElement;
+private:
+	/**
+	 * Updates the expression preview render proxy from a graph node
+	 */
+	void UpdatePreviewNodeRenderProxy();
+};
+
+class FPreviewElement : public ICustomSlateElement
+{
+public:
+	FPreviewElement();
+	~FPreviewElement();
+
+	/**
+	 * Sets up the canvas for rendering
+	 *
+	 * @param	InCanvasRect			Size of the canvas tile
+	 * @param	InClippingRect			How to clip the canvas tile
+	 * @param	InGraphNode				The graph node for the material preview
+	 * @param	bInIsRealtime			Whether preview is using realtime values
+	 *
+	 * @return	Whether there is anything to render
+	 */
+	bool BeginRenderingCanvas( const FIntRect& InCanvasRect, const FIntRect& InClippingRect, UMaterialGraphNode* InGraphNode, bool bInIsRealtime );
+
+	/**
+	 * Updates the expression preview render proxy from a graph node on the render thread
+	 */
+	void UpdateExpressionPreview(UMaterialGraphNode* PreviewNode);
+private:
+	/**
+	 * ICustomSlateElement interface 
+	 */
+	virtual void Draw_RenderThread(FRDGBuilder& GraphBuilder, const FDrawPassInputs& Inputs) override;
+
+private:
+	/** Render target that the canvas renders to */
+	class FSlateMaterialPreviewRenderTarget* RenderTarget;
+	/** Render proxy for the expression preview */
+	FMaterialRenderProxy* ExpressionPreview;
+	/** Whether preview is using realtime values */
+	bool bIsRealtime;
+};
+
+class MATERIALEDITOR_API SGraphNodeMaterialBase : public SGraphNode
+{
+public:
+	SLATE_BEGIN_ARGS(SGraphNodeMaterialBase){}
+	SLATE_END_ARGS()
+
+	void Construct(const FArguments& InArgs, class UMaterialGraphNode* InNode);
+
+	// SGraphNode interface
+	virtual void CreatePinWidgets() override;
+	// End of SGraphNode interface
+
+	// SNodePanel::SNode interface
+	virtual void MoveTo(const FVector2f& NewPosition, FNodeSet& NodeFilter, bool bMarkDirty = true) override;
+	// End of SNodePanel::SNode interface
+
+	UMaterialGraphNode* GetMaterialGraphNode() const {return MaterialNode;}
+
+	/* Populate a meta data tag with information about this graph node */
+	virtual void PopulateMetaTag(class FGraphNodeMetaData* TagMeta) const override;
+
+protected:
+	// SGraphNode interface
+	virtual void AddPin( const TSharedRef<SGraphPin>& PinToAdd ) override;
+	virtual void CreateBelowPinControls(TSharedPtr<SVerticalBox> MainBox) override;
+	virtual void SetDefaultTitleAreaWidget(TSharedRef<SOverlay> DefaultTitleAreaWidget) override;
+	virtual TSharedRef<SWidget> CreateNodeContentArea() override;
+	virtual void OnAdvancedViewChanged(const ECheckBoxState NewCheckedState) override;
+	// End of SGraphNode interface
+
+	/** Creates a preview viewport if necessary */
+	TSharedRef<SWidget> CreatePreviewWidget();
+
+	/** Returns visibility of Expression Preview viewport */
+	EVisibility ExpressionPreviewVisibility() const;
+
+	/** Returns text to over lay over the expression preview viewport */
+	FText ExpressionPreviewOverlayText() const;
+
+	/** Called when we click to preview node button */
+	FReply OnTogglePreviewClicked();
+
+	/** Show/hide Expression Preview */
+	void OnExpressionPreviewChanged(const ECheckBoxState NewCheckedState);
+
+	/** hidden == unchecked, shown == checked */
+	ECheckBoxState IsExpressionPreviewChecked() const;
+
+	/** Up when shown, down when hidden */
+	const FSlateBrush* GetExpressionPreviewArrow() const;
+
+protected:
+	/** Slate viewport for rendering preview via custom slate element */
+	TSharedPtr<FPreviewViewport> PreviewViewport;
+
+	/** Cached material graph node pointer to avoid casting */
+	UMaterialGraphNode* MaterialNode;
+};
diff --git a/Engine/Source/Editor/MaterialEditor/Public/MaterialEditor/MaterialNodes/SGraphNodeMaterialComposite.h b/Engine/Source/Editor/MaterialEditor/Public/MaterialEditor/MaterialNodes/SGraphNodeMaterialComposite.h
new file mode 100644
--- /dev/null
+++ b/Engine/Source/Editor/MaterialEditor/Public/MaterialEditor/MaterialNodes/SGraphNodeMaterialComposite.h
@@ -0,0 +1,45 @@
+// Copyright Epic Games, Inc. All Rights Reserved.
+
+#pragma once
+
+#include "Internationalization/Text.h"
+#include "MaterialEditor/MaterialNodes/SGraphNodeMaterialBase.h"
+#include "Math/Vector2D.h"
+#include "SNodePanel.h"
+#include "Templates/SharedPointer.h"
+#include "Widgets/DeclarativeSyntaxSupport.h"
+
+class SToolTip;
+class SWidget;
+class UEdGraph;
+class UMaterialGraphNode_Composite;
+
+//@TODO: This class is mostly C&P from UK2Node_Composite, consolidate composites to not be BP specific.
+class MATERIALEDITOR_API SGraphNodeMaterialComposite : public SGraphNodeMaterialBase
+{
+public:
+	SLATE_BEGIN_ARGS(SGraphNodeMaterialComposite){}
+	SLATE_END_ARGS()
+
+	void Construct(const FArguments& InArgs, UMaterialGraphNode_Composite* InNode);
+
+	// SGraphNode interface
+	virtual void UpdateGraphNode() override;
+	virtual TSharedPtr<SToolTip> GetComplexTooltip() override;
+	// End of SGraphNode interface
+
+	// SNodePanel::SNode interface
+	virtual void MoveTo(const FVector2f& NewPosition, FNodeSet& NodeFilter, bool bMarkDirty = true) override;
+	// End of SNodePanel::SNode interface
+protected:
+	virtual UEdGraph* GetInnerGraph() const;
+
+private:
+	FText GetPreviewCornerText() const;
+	FText GetTooltipTextForNode() const;
+
+	TSharedRef<SWidget> CreateNodeBody();
+
+	/** Cached material graph node pointer to avoid casting */
+	UMaterialGraphNode_Composite* CompositeNode;
+};
diff --git a/Engine/Source/Editor/MaterialEditor/Public/MaterialEditor/MaterialNodes/SGraphNodeMaterialConvert.h b/Engine/Source/Editor/MaterialEditor/Public/MaterialEditor/MaterialNodes/SGraphNodeMaterialConvert.h
new file mode 100644
--- /dev/null
+++ b/Engine/Source/Editor/MaterialEditor/Public/MaterialEditor/MaterialNodes/SGraphNodeMaterialConvert.h
@@ -0,0 +1,201 @@
+// Copyright Epic Games, Inc. All Rights Reserved.
+
+#pragma once
+
+#include "Input/DragAndDrop.h"
+#include "MaterialEditor/MaterialNodes/SGraphNodeMaterialBase.h"
+#include "SGraphPin.h"
+#include "Types/SlateVector2.h"
+#include "Widgets/SCompoundWidget.h"
+
+class FConvertDragDropOp;
+class SConvertInnerPin;
+class SGraphNodeMaterialConvert;
+class SImage;
+enum class EMaterialExpressionConvertType : uint8;
+
+////////////////////////////////////////////////////////////////////////////////////////////////////////////////////////////////////////////
+// FConvertDragDropOp: Drag and Drop Operation used to form connections within the convert node
+////////////////////////////////////////////////////////////////////////////////////////////////////////////////////////////////////////////
+
+class FConvertDragDropOp : public FDragDropOperation
+{
+public:
+	DRAG_DROP_OPERATOR_TYPE(FConvertDragDropOp, FDragDropOperation)
+
+		FConvertDragDropOp(TSharedPtr<SConvertInnerPin> InSourcePin);
+
+	//~ Begin FDragDropOperation interface
+	virtual void OnDrop(bool bDropWasHandled, const FPointerEvent& MouseEvent) override;
+	virtual void OnDragged(const class FDragDropEvent& DragDropEvent) override;
+	//~ End FDragDropOperation interface
+
+	TSharedPtr<SConvertInnerPin> GetSourcePin() const;
+	UE::Slate::FDeprecateVector2DResult GetScreenPosition() const;
+
+protected:
+	TSharedPtr<SConvertInnerPin> SourcePin;
+	FVector2f ScreenPosition;
+};
+
+////////////////////////////////////////////////////////////////////////////////////////////////////////////////////////////////////////////
+// SMaterialExpressionConvertGraphPin: The Outer pins that form connections to other material graph nodes
+////////////////////////////////////////////////////////////////////////////////////////////////////////////////////////////////////////////
+
+class SMaterialExpressionConvertGraphPin : public SGraphPin
+{
+public:
+	SLATE_BEGIN_ARGS(SMaterialExpressionConvertGraphPin){}
+	SLATE_END_ARGS()
+
+	void Construct(const FArguments& InArgs, UEdGraphPin* InPin);
+
+	void CreateInnerPins(TSharedRef<SGraphNodeMaterialConvert> InOwningGraphNodeWidget);
+
+	const TArray<TSharedRef<SConvertInnerPin>>& GetInnerPins() const { return InnerPins; }
+
+	bool IsHoveredOverPrimaryPin() const;
+private:
+
+	// Inner Pins, one for each component
+	TArray<TSharedRef<SConvertInnerPin>> InnerPins;
+};
+
+////////////////////////////////////////////////////////////////////////////////////////////////////////////////////////////////////////////
+// SConvertInnerPin: The inner pins used to route values with the convert node
+////////////////////////////////////////////////////////////////////////////////////////////////////////////////////////////////////////////
+
+/** Class for the inner pins we'll use to route values in SGraphNodeMaterialConvert */
+class SConvertInnerPin : public SCompoundWidget
+{
+public:
+	SLATE_BEGIN_ARGS(SConvertInnerPin) {}
+	SLATE_END_ARGS()
+
+	void Construct(
+		const FArguments& InArgs,
+		TSharedPtr<SGraphNodeMaterialConvert> InOwningNode,
+		TSharedPtr<SMaterialExpressionConvertGraphPin> InOwningPin,
+		bool bInIsInputPin, 
+		int32 InPinIndex, 
+		int32 InComponentIndex
+	);
+
+	//~ Begin SWidget interface
+	virtual int32 OnPaint(
+		const FPaintArgs& Args, 
+		const FGeometry& AllottedGeometry, 
+		const FSlateRect& MyCullingRect, 
+		FSlateWindowElementList& OutDrawElements, 
+		int32 LayerId, 
+		const FWidgetStyle& InWidgetStyle, 
+		bool bParentEnabled
+	) const override;
+	virtual FReply OnMouseButtonDown(const FGeometry& MyGeometry, const FPointerEvent& MouseEvent) override;
+	virtual FReply OnMouseButtonUp(const FGeometry& MyGeometry, const FPointerEvent& MouseEvent) override;
+	virtual FReply OnDragDetected(const FGeometry& MyGeometry, const FPointerEvent& MouseEvent) override;
+	virtual FReply OnDrop(const FGeometry& MyGeometry, const FDragDropEvent& DragDropEvent) override;
+	//~ End SWidget interface
+
+	const FSlateBrush* GetPinBrush() const;
+	void CancelDragDrop();
+	UE::Slate::FDeprecateVector2DResult GetPinCenterAbsolute() const;
+
+	/** Adds a connection to WeakConnectedPins */
+	void AddConnection(TSharedPtr<SConvertInnerPin> InOtherPin);
+
+	/** Removes a connection from WeakConnectedPins */
+	void RemoveConnection(TSharedPtr<SConvertInnerPin> InOtherPin);
+
+	/** Clears WeakConnectedPins */
+	void RemoveAllConnections();
+
+	/** Calls into the owning node to break connections and refresh the material node */
+	void BreakConnections();
+
+	EVisibility GetDefaultValueVisibility() const;
+	TOptional<float> GetDefaultValue() const;
+	void SetDefaultValue(float InDefaultValue, ETextCommit::Type CommitType);
+
+	EVisibility GetPinNameVisibility() const;
+	FText GetPinName() const;
+
+	bool IsInputPin() const { return bIsInputPin; }
+	int32 GetPinIndex() const { return PinIndex; }
+	int32 GetComponentIndex() const { return ComponentIndex; }
+
+protected:
+	// Input Data
+	TWeakPtr<SGraphNodeMaterialConvert> WeakOwningNode;
+	TWeakPtr<SMaterialExpressionConvertGraphPin> WeakOwningPin;
+	bool bIsInputPin = false;
+	int32 PinIndex = INDEX_NONE;
+	int32 ComponentIndex = INDEX_NONE;
+
+	TSharedPtr<SImage> PinImage;
+	TSet<TWeakPtr<SConvertInnerPin>> WeakConnectedPins;
+	mutable FVector2f CenterAbsolute;
+};
+
+////////////////////////////////////////////////////////////////////////////////////////////////////////////////////////////////////////////
+// SGraphNodeMaterialConvert: Custom Slate widget for UMaterialExpressionConvert
+////////////////////////////////////////////////////////////////////////////////////////////////////////////////////////////////////////////
+
+class MATERIALEDITOR_API SGraphNodeMaterialConvert : public SGraphNodeMaterialBase
+{
+public:
+	SLATE_BEGIN_ARGS(SGraphNodeMaterialConvert){}
+	SLATE_END_ARGS()
+
+	void Construct(const FArguments& InArgs, class UMaterialGraphNode* InNode);
+
+	// Begin SGraphNode interface
+	virtual void CreatePinWidgets() override;
+	virtual TSharedPtr<SGraphPin> CreatePinWidget(UEdGraphPin* Pin) const override;
+	virtual void AddPin(const TSharedRef<SGraphPin>& PinToAdd) override;
+
+	virtual void CreateInputSideAddButton(TSharedPtr<SVerticalBox> InputBox) override;
+	virtual void CreateOutputSideAddButton(TSharedPtr<SVerticalBox> OutputBox) override;
+	virtual EVisibility IsAddPinButtonVisible() const override;
+	// End SGraphNode interface
+
+	FReply OnAddInputPinClicked();
+	FReply OnAddOutputPinClicked();
+
+	TSharedRef<SWidget> CreateAddPinContextMenu(bool bInputPin);
+	void AddNewPin(bool bInputPin, EMaterialExpressionConvertType ConvertType);
+
+	// Begin SWidget Interface
+	virtual int32 OnPaint(
+		const FPaintArgs& Args, 
+		const FGeometry& AllottedGeometry, 
+		const FSlateRect& MyCullingRect, 
+		FSlateWindowElementList& OutDrawElements, 
+		int32 LayerId, 
+		const FWidgetStyle& InWidgetStyle, 
+		bool bParentEnabled
+	) const override;
+	// End SWidget Interface
+
+	void FormConnection(TSharedPtr<SConvertInnerPin> InnerPinA, TSharedPtr<SConvertInnerPin> InnerPinB);
+	void BreakConnections(TSharedPtr<SConvertInnerPin> InnerPin);
+
+	TOptional<float> GetDefaultValue(TSharedPtr<const SConvertInnerPin> InnerPin) const;
+	void SetDefaultValue(TSharedPtr<SConvertInnerPin> InnerPin, const float InDefaultValue);
+
+	void SetCurrentDragDropOp(TSharedPtr<FConvertDragDropOp> InDragDropOp) { CurrentDragDropOp = InDragDropOp; }
+	TSharedPtr<FConvertDragDropOp> GetCurrentDragDropOp() const { return CurrentDragDropOp; }
+
+protected:
+	TSharedPtr<SConvertInnerPin> GetInnerPin(bool bInputPin, const int32 InPinIndex, const int32 InComponentIndex) const;
+	void MakeConnectionCurve(
+		const FGeometry& InAllottedGeometry,
+		FSlateWindowElementList& OutDrawElements,
+		int32& InOutLayerId, 
+		const FVector2f& InCurveStart, 
+		const FVector2f& InCurveEnd
+	) const;
+
+private:
+	TSharedPtr<FConvertDragDropOp> CurrentDragDropOp;
+};
\ No newline at end of file
diff --git a/Engine/Source/Editor/MaterialEditor/Public/MaterialEditor/MaterialNodes/SGraphNodeMaterialCustom.h b/Engine/Source/Editor/MaterialEditor/Public/MaterialEditor/MaterialNodes/SGraphNodeMaterialCustom.h
new file mode 100644
--- /dev/null
+++ b/Engine/Source/Editor/MaterialEditor/Public/MaterialEditor/MaterialNodes/SGraphNodeMaterialCustom.h
@@ -0,0 +1,39 @@
+// Copyright Epic Games, Inc. All Rights Reserved.
+
+#pragma once
+
+#include "MaterialEditor/MaterialNodes/SGraphNodeMaterialBase.h"
+#include "Text/HLSLSyntaxHighlighterMarshaller.h"
+#include "Widgets/Input/SMultiLineEditableTextBox.h"
+
+class MATERIALEDITOR_API SGraphNodeMaterialCustom : public SGraphNodeMaterialBase
+{
+public:
+	SLATE_BEGIN_ARGS(SGraphNodeMaterialCustom){}
+	SLATE_END_ARGS()
+
+	void Construct(const FArguments& InArgs, class UMaterialGraphNode* InNode);
+
+	UMaterialGraphNode* GetMaterialGraphNode() const {return MaterialNode;}
+
+	FText GetHlslText() const;
+	void OnCustomHlslTextCommited(const FText& InText, ETextCommit::Type InType);
+
+protected:
+	// SGraphNode interface
+	virtual void CreateBelowPinControls(TSharedPtr<SVerticalBox> MainBox) override;
+	virtual void CreateAdvancedViewArrow(TSharedPtr<SVerticalBox> MainBox) override;
+	virtual EVisibility AdvancedViewArrowVisibility() const override;
+	virtual void OnAdvancedViewChanged( const ECheckBoxState NewCheckedState ) override;
+	virtual ECheckBoxState IsAdvancedViewChecked() const override;
+	virtual const FSlateBrush* GetAdvancedViewArrow() const override;
+	// End of SGraphNode interface
+
+private:
+	EVisibility CodeVisibility() const;
+
+protected:
+	/** Shader text box */
+	TSharedPtr<FHLSLSyntaxHighlighterMarshaller> SyntaxHighlighter;
+	TSharedPtr<SMultiLineEditableTextBox> ShaderTextBox;
+};
diff --git a/Engine/Source/Editor/MaterialEditor/Public/MaterialEditor/MaterialNodes/SGraphNodeMaterialResult.h b/Engine/Source/Editor/MaterialEditor/Public/MaterialEditor/MaterialNodes/SGraphNodeMaterialResult.h
new file mode 100644
--- /dev/null
+++ b/Engine/Source/Editor/MaterialEditor/Public/MaterialEditor/MaterialNodes/SGraphNodeMaterialResult.h
@@ -0,0 +1,32 @@
+// Copyright Epic Games, Inc. All Rights Reserved.
+
+#pragma once
+
+#include "Math/Vector2D.h"
+#include "SGraphNode.h"
+#include "SNodePanel.h"
+#include "Widgets/DeclarativeSyntaxSupport.h"
+
+class FGraphNodeMetaData;
+class UMaterialGraphNode_Root;
+
+class MATERIALEDITOR_API SGraphNodeMaterialResult : public SGraphNode
+{
+public:
+	SLATE_BEGIN_ARGS(SGraphNodeMaterialResult){}
+	SLATE_END_ARGS()
+
+	void Construct(const FArguments& InArgs, class UMaterialGraphNode_Root* InNode);
+
+	// SGraphNode interface
+	virtual void CreatePinWidgets() override;
+	virtual void PopulateMetaTag(FGraphNodeMetaData* TagMeta) const override;
+	// End of SGraphNode interface
+
+	// SNodePanel::SNode interface
+	virtual void MoveTo(const FVector2f& NewPosition, FNodeSet& NodeFilter, bool bMarkDirty = true) override;
+	// End of SNodePanel::SNode interface
+
+private:
+	UMaterialGraphNode_Root* RootNode;
+};
diff --git a/Engine/Source/Editor/MaterialEditor/Public/MaterialEditor/MaterialPins/SGraphPinMaterialInput.h b/Engine/Source/Editor/MaterialEditor/Public/MaterialEditor/MaterialPins/SGraphPinMaterialInput.h
new file mode 100644
--- /dev/null
+++ b/Engine/Source/Editor/MaterialEditor/Public/MaterialEditor/MaterialPins/SGraphPinMaterialInput.h
@@ -0,0 +1,25 @@
+// Copyright Epic Games, Inc. All Rights Reserved.
+
+#pragma once
+
+#include "CoreMinimal.h"
+#include "SGraphPin.h"
+#include "Styling/SlateColor.h"
+#include "Widgets/DeclarativeSyntaxSupport.h"
+
+class UEdGraphPin;
+
+class MATERIALEDITOR_API SGraphPinMaterialInput : public SGraphPin
+{
+public:
+	SLATE_BEGIN_ARGS(SGraphPinMaterialInput) {}
+	SLATE_END_ARGS()
+
+	void Construct(const FArguments& InArgs, UEdGraphPin* InGraphPinObj);
+
+protected:
+	//~ Begin SGraphPin Interface
+	virtual FSlateColor GetPinColor() const override;
+	//~ End SGraphPin Interface
+
+};
diff --git a/Engine/Source/Editor/MaterialEditor/Public/MaterialEditor/SGraphSubstrateMaterial.h b/Engine/Source/Editor/MaterialEditor/Public/MaterialEditor/SGraphSubstrateMaterial.h
new file mode 100644
--- /dev/null
+++ b/Engine/Source/Editor/MaterialEditor/Public/MaterialEditor/SGraphSubstrateMaterial.h
@@ -0,0 +1,20 @@
+// Copyright Epic Games, Inc. All Rights Reserved.
+
+#pragma once
+
+#include "Internationalization/Text.h"
+#include "Rendering/RenderingCommon.h"
+#include "SGraphPin.h"
+
+class UMaterialGraphNode;
+struct FSubstrateMaterialCompilationOutput;
+
+struct MATERIALEDITOR_API FSubstrateWidget
+{
+	static const TSharedRef<SWidget> ProcessOperator(const FSubstrateMaterialCompilationOutput& CompilationOutput);
+	static const TSharedRef<SWidget> ProcessOperator(const FSubstrateMaterialCompilationOutput& CompilationOutput, const TArray<FGuid>& InGuid);
+	static void GetPinColor(TSharedPtr<SGraphPin>& Out, const UMaterialGraphNode* InNode);
+	static FLinearColor GetConnectionColor();
+	static bool HasInputSubstrateType(const UEdGraphPin* InPin);
+	static bool HasOutputSubstrateType(const UEdGraphPin* InPin);
+};
\ No newline at end of file
diff --git a/Engine/Source/Editor/MaterialEditor/Public/MaterialNodes/SGraphNodeMaterialBase.h b/Engine/Source/Editor/MaterialEditor/Public/MaterialNodes/SGraphNodeMaterialBase.h
deleted file mode 100644
diff --git a/Engine/Source/Editor/MaterialEditor/Public/MaterialNodes/SGraphNodeMaterialComposite.h b/Engine/Source/Editor/MaterialEditor/Public/MaterialNodes/SGraphNodeMaterialComposite.h
deleted file mode 100644
diff --git a/Engine/Source/Editor/MaterialEditor/Public/MaterialNodes/SGraphNodeMaterialConvert.h b/Engine/Source/Editor/MaterialEditor/Public/MaterialNodes/SGraphNodeMaterialConvert.h
deleted file mode 100644
diff --git a/Engine/Source/Editor/MaterialEditor/Public/MaterialNodes/SGraphNodeMaterialCustom.h b/Engine/Source/Editor/MaterialEditor/Public/MaterialNodes/SGraphNodeMaterialCustom.h
deleted file mode 100644
diff --git a/Engine/Source/Editor/MaterialEditor/Public/MaterialNodes/SGraphNodeMaterialResult.h b/Engine/Source/Editor/MaterialEditor/Public/MaterialNodes/SGraphNodeMaterialResult.h
deleted file mode 100644
diff --git a/Engine/Source/Editor/MaterialEditor/Public/MaterialPins/SGraphPinMaterialInput.h b/Engine/Source/Editor/MaterialEditor/Public/MaterialPins/SGraphPinMaterialInput.h
deleted file mode 100644
diff --git a/Engine/Source/Editor/MaterialEditor/Public/SGraphSubstrateMaterial.h b/Engine/Source/Editor/MaterialEditor/Public/SGraphSubstrateMaterial.h
deleted file mode 100644
diff --git a/Engine/Source/Editor/UnrealEd/Private/MaterialGraphNode.cpp b/Engine/Source/Editor/UnrealEd/Private/MaterialGraphNode.cpp
--- a/Engine/Source/Editor/UnrealEd/Private/MaterialGraphNode.cpp
+++ b/Engine/Source/Editor/UnrealEd/Private/MaterialGraphNode.cpp
@@ -54,7 +54,7 @@
 #include "Framework/Commands/GenericCommands.h"
 #include "ScopedTransaction.h"
 #include "Subsystems/AssetEditorSubsystem.h"
-#include "MaterialNodes/SGraphNodeMaterialBase.h"
+#include "MaterialEditor/MaterialNodes/SGraphNodeMaterialBase.h"
 #include "SourceCodeNavigation.h"
 
 #define LOCTEXT_NAMESPACE "MaterialGraphNode"
diff --git a/Engine/Source/Editor/UnrealEd/Private/MaterialGraphNode_Composite.cpp b/Engine/Source/Editor/UnrealEd/Private/MaterialGraphNode_Composite.cpp
--- a/Engine/Source/Editor/UnrealEd/Private/MaterialGraphNode_Composite.cpp
+++ b/Engine/Source/Editor/UnrealEd/Private/MaterialGraphNode_Composite.cpp
@@ -16,7 +16,7 @@
 #include "EdGraphUtilities.h"
 #include "Kismet2/BlueprintEditorUtils.h"
 #include "Framework/Commands/GenericCommands.h"
-#include "MaterialNodes/SGraphNodeMaterialComposite.h"
+#include "MaterialEditor/MaterialNodes/SGraphNodeMaterialComposite.h"
 
 #define LOCTEXT_NAMESPACE "MaterialGraphNode_Composite"
 
diff --git a/Engine/Source/Editor/UnrealEd/Private/MaterialGraphNode_Custom.cpp b/Engine/Source/Editor/UnrealEd/Private/MaterialGraphNode_Custom.cpp
--- a/Engine/Source/Editor/UnrealEd/Private/MaterialGraphNode_Custom.cpp
+++ b/Engine/Source/Editor/UnrealEd/Private/MaterialGraphNode_Custom.cpp
@@ -1,7 +1,7 @@
 // Copyright Epic Games, Inc. All Rights Reserved.
 
 #include "MaterialGraph/MaterialGraphNode_Custom.h"
-#include "MaterialNodes/SGraphNodeMaterialCustom.h"
+#include "MaterialEditor/MaterialNodes/SGraphNodeMaterialCustom.h"
 
 #include "Materials/MaterialExpressionCustom.h"
 
diff --git a/Engine/Source/Editor/UnrealEd/Private/MaterialGraphNode_Root.cpp b/Engine/Source/Editor/UnrealEd/Private/MaterialGraphNode_Root.cpp
--- a/Engine/Source/Editor/UnrealEd/Private/MaterialGraphNode_Root.cpp
+++ b/Engine/Source/Editor/UnrealEd/Private/MaterialGraphNode_Root.cpp
@@ -12,7 +12,7 @@
 #include "Materials/MaterialAttributeDefinitionMap.h"
 #include "GraphEditorSettings.h"
 
-#include "MaterialNodes/SGraphNodeMaterialResult.h"
+#include "MaterialEditor/MaterialNodes/SGraphNodeMaterialResult.h"
 
 #define LOCTEXT_NAMESPACE "MaterialGraphNode_Root"
 
diff --git a/Engine/Source/Runtime/Engine/Private/Materials/MaterialExpressions.cpp b/Engine/Source/Runtime/Engine/Private/Materials/MaterialExpressions.cpp
--- a/Engine/Source/Runtime/Engine/Private/Materials/MaterialExpressions.cpp
+++ b/Engine/Source/Runtime/Engine/Private/Materials/MaterialExpressions.cpp
@@ -306,8 +306,8 @@
 #if WITH_EDITOR
 #include "MaterialGraph/MaterialGraphNode_Comment.h"
 #include "MaterialGraph/MaterialGraphNode_Composite.h"
-#include "MaterialNodes/SGraphNodeMaterialBase.h"
-#include "MaterialNodes/SGraphNodeMaterialConvert.h"
+#include "MaterialEditor/MaterialNodes/SGraphNodeMaterialBase.h"
+#include "MaterialEditor/MaterialNodes/SGraphNodeMaterialConvert.h"
 #include "Framework/Notifications/NotificationManager.h"
 #include "Widgets/Notifications/SNotificationList.h"
 #include "MaterialGraph/MaterialGraphSchema.h"
