When using the Material Editing Library of functions always refresh the material instance editor when changing material instance parameters.

- This fixes a bug where you would set a MI parameter via a python script and the MI Editor wouldn't update the properties panel to show the newly updated value.
- Automatically call UMaterialEditingLibrary::UpdateMaterialInstance() whenever a parameter on a instance is set via the library.
- UMaterialEditingLibrary::UpdateMaterialInstance should try and find any open material instance editors and tell them to that the material instance has changed.
- FMaterialInstanceEditor::NotifyExternalMaterialChange now calls `Refresh()` which will update the whole editor, including the property window, to properly reflect any parameter changes made by the library.

[RN] major | Rendering - Materials and Shaders

diff --git a/Engine/Source/Editor/MaterialEditor/Private/MaterialEditingLibrary.cpp b/Engine/Source/Editor/MaterialEditor/Private/MaterialEditingLibrary.cpp
--- a/Engine/Source/Editor/MaterialEditor/Private/MaterialEditingLibrary.cpp
+++ b/Engine/Source/Editor/MaterialEditor/Private/MaterialEditingLibrary.cpp
@@ -1146,6 +1146,7 @@ bool UMaterialEditingLibrary::SetMaterialInstanceScalarParameterValue(UMaterialI
 	if (Instance)
 	{
 		Instance->SetScalarParameterValueEditorOnly(FMaterialParameterInfo(ParameterName, Association, Association == EMaterialParameterAssociation::LayerParameter ? 0 : INDEX_NONE), Value);
+		UMaterialEditingLibrary::UpdateMaterialInstance(Instance);
 	}
 	return bResult;
 }
@@ -1167,6 +1168,7 @@ bool UMaterialEditingLibrary::SetMaterialInstanceTextureParameterValue(UMaterial
 	if (Instance)
 	{
 		Instance->SetTextureParameterValueEditorOnly(FMaterialParameterInfo(ParameterName, Association, Association == EMaterialParameterAssociation::LayerParameter ? 0 : INDEX_NONE), Value);
+		UMaterialEditingLibrary::UpdateMaterialInstance(Instance);
 	}
 	return bResult;
 }
@@ -1188,6 +1190,7 @@ bool UMaterialEditingLibrary::SetMaterialInstanceRuntimeVirtualTextureParameterV
 	if (Instance)
 	{
 		Instance->SetRuntimeVirtualTextureParameterValueEditorOnly(FMaterialParameterInfo(ParameterName, Association, Association == EMaterialParameterAssociation::LayerParameter ? 0 : INDEX_NONE), Value);
+		UMaterialEditingLibrary::UpdateMaterialInstance(Instance);
 	}
 	return bResult;
 }
@@ -1209,6 +1212,7 @@ bool UMaterialEditingLibrary::SetMaterialInstanceSparseVolumeTextureParameterVal
 	if (Instance)
 	{
 		Instance->SetSparseVolumeTextureParameterValueEditorOnly(FMaterialParameterInfo(ParameterName, Association, Association == EMaterialParameterAssociation::LayerParameter ? 0 : INDEX_NONE), Value);
+		UMaterialEditingLibrary::UpdateMaterialInstance(Instance);
 	}
 	return bResult;
 }
@@ -1230,6 +1234,7 @@ bool UMaterialEditingLibrary::SetMaterialInstanceVectorParameterValue(UMaterialI
 	if (Instance)
 	{
 		Instance->SetVectorParameterValueEditorOnly(FMaterialParameterInfo(ParameterName, Association, Association == EMaterialParameterAssociation::LayerParameter ? 0 : INDEX_NONE), Value);
+		UMaterialEditingLibrary::UpdateMaterialInstance(Instance);
 	}
 	return bResult;
 }
@@ -1277,6 +1282,12 @@ void UMaterialEditingLibrary::UpdateMaterialInstance(UMaterialInstanceConstant*
 		Instance->UpdateStaticPermutation();
 		Instance->UpdateParameterNames();
 
+		// Refresh the material instance editor to reflect parameter value changes in the property window.
+		if (IMaterialEditor* MaterialEditor = MaterialEditingLibraryImpl::FindMaterialEditorForAsset(Instance))
+		{
+			MaterialEditor->NotifyExternalMaterialChange();
+		}
+
 		// update the world's viewports
 		FEditorDelegates::RefreshEditor.Broadcast();
 		FEditorSupportDelegates::RedrawAllViewports.Broadcast();
diff --git a/Engine/Source/Editor/MaterialEditor/Private/MaterialInstanceEditor.cpp b/Engine/Source/Editor/MaterialEditor/Private/MaterialInstanceEditor.cpp
--- a/Engine/Source/Editor/MaterialEditor/Private/MaterialInstanceEditor.cpp
+++ b/Engine/Source/Editor/MaterialEditor/Private/MaterialInstanceEditor.cpp
@@ -1852,6 +1852,7 @@ void FMaterialInstanceEditor::PostRedo( bool bSuccess )
 void FMaterialInstanceEditor::NotifyExternalMaterialChange()
 {
 	MaterialStatsManager->SignalMaterialChanged();
+	Refresh();
 }
 
 void FMaterialInstanceEditor::NotifyUserSceneTextureLoadOrUnload()
