Add/Multiply allow Extra Pins for Operator Node in Material Graph New Translator

diff --git a/Engine/Source/Editor/MaterialEditor/Private/MaterialEditorDetailCustomization.cpp b/Engine/Source/Editor/MaterialEditor/Private/MaterialEditorDetailCustomization.cpp
--- a/Engine/Source/Editor/MaterialEditor/Private/MaterialEditorDetailCustomization.cpp
+++ b/Engine/Source/Editor/MaterialEditor/Private/MaterialEditorDetailCustomization.cpp
@@ -16,6 +16,7 @@
 #include "Materials/MaterialExpressionVectorParameter.h"
 #include "Materials/MaterialExpressionComposite.h"
 #include "Materials/MaterialExpressionColorRamp.h"
+#include "Materials/MaterialExpressionOperator.h"
 #include "Materials/MaterialExpressionPinBase.h"
 #include "Materials/MaterialExpressionTextureSampleParameter.h"
 #include "MaterialLayersFunctionsCustomization.h"
@@ -1137,4 +1138,60 @@ void FMaterialExpressionColorRampCustomization::CustomizeDetails(IDetailLayoutBu
 	}
 }
 
+TSharedRef<IDetailCustomization> FMaterialExpressionOperatorCustomization::MakeInstance()
+{
+	return MakeShareable(new FMaterialExpressionOperatorCustomization);
+}
+
+void FMaterialExpressionOperatorCustomization::CustomizeDetails(IDetailLayoutBuilder& DetailLayout)
+{
+	// Hides default detail panel property for input array
+    TSharedRef<IPropertyHandle> DynamicInputsHandle = DetailLayout.GetProperty(GET_MEMBER_NAME_CHECKED(UMaterialExpressionOperator, DynamicInputs));
+	DetailLayout.HideProperty(DynamicInputsHandle);
+
+	// Get Material Expression Operater ref
+    UMaterialExpressionOperator* MaterialExpOperator = nullptr;
+    TArray<TWeakObjectPtr<UObject>> Objects;
+	DetailLayout.GetObjectsBeingCustomized(Objects);
+    for (TWeakObjectPtr<UObject>& Obj : Objects)
+    {
+        if (auto* Test = Cast<UMaterialExpressionOperator>(Obj.Get()))
+        {
+			MaterialExpOperator = Test;
+            break;
+        }
+    }
+
+	if (MaterialExpOperator)
+	{
+		DetailLayout.EditCategory("MaterialExpressionOperator");
+		IDetailCategoryBuilder& InputCategory = DetailLayout.EditCategory("MaterialExpressionOperatorInputs", LOCTEXT("InputsCategory", "Material Expression Operator Inputs"));
+
+		// Add an entry for each input in the input array
+		uint32 NumInputs = 0;
+		DynamicInputsHandle->GetNumChildren(NumInputs);
+		for (uint32 Index = 0; Index < NumInputs; Index++)
+		{
+			TSharedPtr<IPropertyHandle> ElementHandle = DynamicInputsHandle->GetChildHandle(Index);
+			TSharedPtr<IPropertyHandle> ConstValueHandle = ElementHandle->GetChildHandle(GET_MEMBER_NAME_CHECKED(FMaterialExpressionOperatorInput, ConstValue));
+			FText DisplayName = FText::FromName(MaterialExpOperator->GetInputName(Index));
+
+			InputCategory.AddCustomRow(DisplayName)
+			.IsEnabled(true)
+			.NameContent()
+			[
+				SNew(STextBlock)
+				.Text(DisplayName)
+				.Font(IDetailLayoutBuilder::GetDetailFont())
+			]
+			.ValueContent()
+			.MinDesiredWidth(125.0f)
+			[
+				ConstValueHandle->CreatePropertyValueWidget()
+			];
+		}
+	}
+
+}
+
 #undef LOCTEXT_NAMESPACE
diff --git a/Engine/Source/Editor/MaterialEditor/Private/MaterialEditorDetailCustomization.h b/Engine/Source/Editor/MaterialEditor/Private/MaterialEditorDetailCustomization.h
--- a/Engine/Source/Editor/MaterialEditor/Private/MaterialEditorDetailCustomization.h
+++ b/Engine/Source/Editor/MaterialEditor/Private/MaterialEditorDetailCustomization.h
@@ -157,6 +157,19 @@ public:
 	virtual void CustomizeDetails(IDetailLayoutBuilder& DetailLayout) override;
 };
 
+/**
+ * Customizes the details of Operator node
+ */
+class FMaterialExpressionOperatorCustomization : public IDetailCustomization
+{
+public:
+	/** Makes a new instance of this detail layout class for a specific detail view requesting it */
+	static TSharedRef<IDetailCustomization> MakeInstance();
+
+	/** IDetailCustomization interface */
+	virtual void CustomizeDetails(IDetailLayoutBuilder& DetailLayout) override;
+};
+
 class FMaterialDetailCustomization : public IDetailCustomization
 {
 public:
diff --git a/Engine/Source/Editor/MaterialEditor/Private/MaterialEditorModule.cpp b/Engine/Source/Editor/MaterialEditor/Private/MaterialEditorModule.cpp
--- a/Engine/Source/Editor/MaterialEditor/Private/MaterialEditorModule.cpp
+++ b/Engine/Source/Editor/MaterialEditor/Private/MaterialEditorModule.cpp
@@ -77,11 +77,16 @@ public:
 				GetMutableDefault<UMaterialEditorSettings>());
 		}
 
+		// Register customizaed detail panel layouts
 		FPropertyEditorModule& PropertyModule = FModuleManager::LoadModuleChecked<FPropertyEditorModule>("PropertyEditor");
 		PropertyModule.RegisterCustomClassLayout(
 			"MaterialExpressionColorRamp",
 			FOnGetDetailCustomizationInstance::CreateStatic(&FMaterialExpressionColorRampCustomization::MakeInstance)
 		);
+		PropertyModule.RegisterCustomClassLayout(
+			"MaterialExpressionOperator",
+			FOnGetDetailCustomizationInstance::CreateStatic(&FMaterialExpressionOperatorCustomization::MakeInstance)
+		);
 
 		GraphPanelPinFactory = MakeShared<FMaterialEditorGraphPanelPinFactory>();
 		FEdGraphUtilities::RegisterVisualPinFactory(GraphPanelPinFactory);
diff --git a/Engine/Source/Editor/MaterialEditor/Private/MaterialNodes/SGraphNodeMaterialOperator.cpp b/Engine/Source/Editor/MaterialEditor/Private/MaterialNodes/SGraphNodeMaterialOperator.cpp
new file mode 100644
--- /dev/null
+++ b/Engine/Source/Editor/MaterialEditor/Private/MaterialNodes/SGraphNodeMaterialOperator.cpp
@@ -0,0 +1,260 @@
+// Copyright Epic Games, Inc. All Rights Reserved.
+
+#include "MaterialEditor/MaterialNodes/SGraphNodeMaterialOperator.h"
+#include "MaterialGraph/MaterialGraphNode_Operator.h"
+#include "Materials/MaterialExpressionOperator.h"
+
+#include "Widgets/Text/STextBlock.h" 
+#include "Widgets/Input/SEditableTextBox.h"
+#include "EditorStyleSet.h"
+#include "Framework/MultiBox/MultiBoxBuilder.h"
+
+#define LOCTEXT_NAMESPACE "GraphNodeMaterialOperator"
+
+///////////////////////////////////////////////////////////////////////////////
+// SMaterialExpressionOperatorGraphPin Implementation
+///////////////////////////////////////////////////////////////////////////////
+
+// Instantiation for float
+template class SMaterialExpressionOperatorGraphPin<float>;
+
+template<typename NumericType>
+void SMaterialExpressionOperatorGraphPin<NumericType>::Construct(const FArguments& InArgs, UEdGraphPin* InGraphPinObj)
+{
+	ArrayIndex = InArgs._ArrayIndex;
+	OperatorExpression = InArgs._OperatorExpression;
+	SGraphPin::Construct(SGraphPin::FArguments(), InGraphPinObj);
+}
+
+template<typename NumericType>
+TSharedRef<SWidget> SMaterialExpressionOperatorGraphPin<NumericType>::GetDefaultValueWidget()
+{
+	return SNew(SBox)
+		.MinDesiredWidth(18.0f)
+		.MaxDesiredWidth(400.0f)
+		[
+			SNew(SNumericEntryBox<NumericType>)
+				.EditableTextBoxStyle(FAppStyle::Get(), "Graph.EditableTextBox")
+				.BorderForegroundColor(FSlateColor::UseForeground())
+				.Visibility(this, &SMaterialExpressionOperatorGraphPin::GetDefaultValueVisibility)
+				.IsEnabled(this, &SMaterialExpressionOperatorGraphPin::GetDefaultValueIsEditable)
+				.Value(this, &SMaterialExpressionOperatorGraphPin::GetNumericValue)
+				.OnValueCommitted(this, &SMaterialExpressionOperatorGraphPin::SetNumericValue)
+		];
+}
+
+template<typename NumericType>
+TOptional<NumericType> SMaterialExpressionOperatorGraphPin<NumericType>::GetNumericValue() const
+{
+	if (OperatorExpression && OperatorExpression->DynamicInputs.IsValidIndex(ArrayIndex))
+	{
+		return OperatorExpression->DynamicInputs[ArrayIndex].ConstValue;
+	}
+	return {};
+}
+
+template<typename NumericType>
+void SMaterialExpressionOperatorGraphPin<NumericType>::SetNumericValue(NumericType NewValue, ETextCommit::Type)
+{
+	if (OperatorExpression && OperatorExpression->DynamicInputs.IsValidIndex(ArrayIndex))
+	{
+		const FScopedTransaction Transaction(LOCTEXT("EditNodeConstValue", "Edit Const Value"));
+		OperatorExpression->Modify();
+		OperatorExpression->DynamicInputs[ArrayIndex].ConstValue = NewValue;
+		OperatorExpression->RefreshNode();
+		if (UEdGraph* Graph = GraphPinObj->GetOwningNode()->GetGraph())
+		{
+			Graph->NotifyGraphChanged();
+		}
+	}
+}
+
+///////////////////////////////////////////////////////////////////////////////
+// SGraphNodeMaterialOperator Implementation
+///////////////////////////////////////////////////////////////////////////////
+
+void SGraphNodeMaterialOperator::Construct(const FArguments& InArgs, class UMaterialGraphNode* InNode)
+{
+	this->GraphNode = InNode;
+	this->MaterialNode = InNode;
+	this->UpdateGraphNode();
+}
+
+UMaterialExpressionOperator* SGraphNodeMaterialOperator::GetMaterialExpression() const
+{
+	if (auto GraphOp = Cast<UMaterialGraphNode_Operator>(GraphNode))
+	{
+		return Cast<UMaterialExpressionOperator>(GraphOp->MaterialExpression);
+	}
+	return nullptr;
+}
+
+/**
+* Create "Add Pin" Button on the material expression node
+*/
+void SGraphNodeMaterialOperator::CreateBelowPinControls(TSharedPtr<SVerticalBox> MainBox)
+{
+	UMaterialExpressionOperator* MaterialExpOperator = GetMaterialExpression();
+	if (MaterialExpOperator && MaterialExpOperator->bAllowAddPin)
+	{
+		TSharedRef<SButton> AddPinButton = StaticCastSharedRef<SButton>(AddPinButtonContent(
+			LOCTEXT("GraphNodeMaterial_Operator", "Add Input"),
+			LOCTEXT("GraphNodeMaterial_Operator_Tooltip", "Add an input to this convert node."),
+			false // bRightSide
+		));
+
+		AddPinButton->SetOnClicked(FOnClicked::CreateSP(this, &SGraphNodeMaterialOperator::OnAddPinClicked));
+
+		MainBox->AddSlot()
+			.FillHeight(1.f)
+			.Padding(5.f)
+			.VAlign(VAlign_Bottom)
+			.HAlign(HAlign_Left)
+			[
+				AddPinButton
+			];
+	}
+}
+
+/**
+* Create SMaterialExpressionOperatorGraphPin for each input
+*/
+TSharedPtr<SGraphPin> SGraphNodeMaterialOperator::CreatePinWidget(UEdGraphPin* InPin) const
+{
+	TArray<UEdGraphPin*> NewInputPins;
+	for (UEdGraphPin* Pin : GraphNode->Pins)
+	{
+		if (Pin->Direction == EGPD_Input)
+		{
+			NewInputPins.Add(Pin);
+		}
+	}
+
+	const int32 Index = NewInputPins.IndexOfByKey(InPin);
+	UMaterialExpressionOperator* MaterialExpOperator = GetMaterialExpression();
+
+	if (Index != INDEX_NONE && MaterialExpOperator)
+	{
+		FName PinName = MaterialExpOperator->GetInputName(Index);
+		InPin->PinFriendlyName = FText::FromName(PinName);
+	}
+
+	return SNew(SMaterialExpressionOperatorGraphPin<float>, InPin)
+		.ArrayIndex(Index)
+		.OperatorExpression(MaterialExpOperator);
+}
+
+/**
+* Create Operator DropDown Menu Button on the material expression node
+*/
+void SGraphNodeMaterialOperator::SetDefaultTitleAreaWidget(TSharedRef<SOverlay> DefaultTitleAreaWidget)
+{
+	DefaultTitleAreaWidget->AddSlot()
+		.HAlign(HAlign_Right)
+		.VAlign(VAlign_Top)
+		.Padding(FMargin(5, 5, 40, 5))
+		[
+			SAssignNew(OperatorButton, SButton)
+				.ButtonStyle(FAppStyle::Get(), "HoverHintOnly")
+				.OnClicked(this, &SGraphNodeMaterialOperator::OnOperatorMenuClicked)
+				.ContentPadding(FMargin(4, 2))
+				[
+					SNew(SImage)
+						.Image(FAppStyle::GetBrush(TEXT("Icons.Dropdown")))
+				]
+		];
+	SGraphNodeMaterialBase::SetDefaultTitleAreaWidget(DefaultTitleAreaWidget);
+}
+
+FReply SGraphNodeMaterialOperator::OnAddPinClicked()
+{
+	UMaterialExpressionOperator* MaterialExpOperator = GetMaterialExpression();
+	if (MaterialExpOperator)
+	{
+		const FScopedTransaction Transaction(LOCTEXT("AddPinOperatorNode", "Add Pin on Operator Node"));
+		MaterialExpOperator->Modify();
+		MaterialExpOperator->AddInputPin();
+		MaterialExpOperator->PostEditChange();
+		MaterialExpOperator->RefreshNode();
+		this->UpdateGraphNode();
+		MaterialNode->ReconstructNode();
+	}
+	return FReply::Handled();
+}
+
+TSharedRef<SWidget> SGraphNodeMaterialOperator::BuildOperatorMenu()
+{
+	FMenuBuilder MenuBuilder(true, nullptr);
+	static const UEnum* EnumPtr = StaticEnum<EMaterialExpressionOperatorKind>();
+	check(EnumPtr);
+	for (int32 Idx = 0; Idx < EnumPtr->NumEnums() - 1; ++Idx)
+	{
+		const EMaterialExpressionOperatorKind Kind = static_cast<EMaterialExpressionOperatorKind>(EnumPtr->GetValueByIndex(Idx));
+		const FText Display = EnumPtr->GetDisplayNameTextByIndex(Idx);
+		MenuBuilder.AddMenuEntry(
+			Display,
+			FText::GetEmpty(),
+			FSlateIcon(),
+			FUIAction(FExecuteAction::CreateSP(this, &SGraphNodeMaterialOperator::OnOperatorMenuSelected, Kind))
+		);
+	}
+	return MenuBuilder.MakeWidget();
+}
+
+FReply SGraphNodeMaterialOperator::OnOperatorMenuClicked()
+{
+	if (auto Button = OperatorButton.Pin())
+	{
+		TSharedRef<SWidget> MenuContent = BuildOperatorMenu();
+		FWidgetPath ButtonPath;
+		FSlateApplication::Get().GeneratePathToWidgetUnchecked(Button.ToSharedRef(), ButtonPath);
+		FSlateApplication::Get().PushMenu(
+			Button.ToSharedRef(),
+			ButtonPath,
+			MenuContent,
+			FSlateApplication::Get().GetCursorPos(),
+			FPopupTransitionEffect(FPopupTransitionEffect::TypeInPopup)
+		);
+
+		return FReply::Handled();
+	}
+	return FReply::Unhandled();
+}
+
+void SGraphNodeMaterialOperator::OnOperatorMenuSelected(EMaterialExpressionOperatorKind NewKind)
+{
+	if (UMaterialExpressionOperator* Expr = GetMaterialExpression())
+	{
+		if (Expr->Operator != NewKind)
+		{
+			const FScopedTransaction Transaction(LOCTEXT("ChangeOperatorKind", "Change Operator Kind"));
+			Expr->Modify();
+			Expr->Operator = NewKind;
+
+			static const FName OperatorPropName = GET_MEMBER_NAME_CHECKED(UMaterialExpressionOperator, Operator);
+			FProperty* ChangedProp = FindFProperty<FProperty>(UMaterialExpressionOperator::StaticClass(), *OperatorPropName.ToString());
+			FPropertyChangedEvent PropEvent(ChangedProp);
+			Expr->PostEditChangeProperty(PropEvent);
+
+			Expr->RefreshNode();
+			this->UpdateGraphNode();
+			MaterialNode->ReconstructNode();
+		}
+	}
+}
+
+FText SGraphNodeMaterialOperator::GetOperatorDisplayName() const
+{
+	if (UMaterialExpressionOperator* Expr = GetMaterialExpression())
+	{
+		static const UEnum* EnumPtr = StaticEnum<EMaterialExpressionOperatorKind>();
+		if (EnumPtr)
+		{
+			return EnumPtr->GetDisplayNameTextByValue((int64)Expr->Operator);
+		}
+	}
+	return LOCTEXT("SelectOperator", "Select…");
+}
+
+#undef LOCTEXT_NAMESPACE
+
diff --git a/Engine/Source/Editor/MaterialEditor/Public/MaterialEditor/MaterialNodes/SGraphNodeMaterialOperator.h b/Engine/Source/Editor/MaterialEditor/Public/MaterialEditor/MaterialNodes/SGraphNodeMaterialOperator.h
new file mode 100644
--- /dev/null
+++ b/Engine/Source/Editor/MaterialEditor/Public/MaterialEditor/MaterialNodes/SGraphNodeMaterialOperator.h
@@ -0,0 +1,70 @@
+// Copyright Epic Games, Inc. All Rights Reserved.
+
+#pragma once
+#include "MaterialEditor/MaterialNodes/SGraphNodeMaterialBase.h"
+#include "Materials/MaterialExpressionOperator.h"
+
+#include "Components/HorizontalBox.h"
+#include "KismetPins/SGraphPinNum.h"
+#include "Widgets/Input/SButton.h"
+
+template<typename NumericType>
+/** Custom SGraphPin for handling array of inputs */
+class SMaterialExpressionOperatorGraphPin : public SGraphPin
+{
+public:
+	SLATE_BEGIN_ARGS(SMaterialExpressionOperatorGraphPin) {}
+		SLATE_ARGUMENT(int32, ArrayIndex)
+		SLATE_ARGUMENT(UMaterialExpressionOperator*, OperatorExpression)
+	SLATE_END_ARGS()
+
+	void Construct(const FArguments& InArgs, UEdGraphPin* InGraphPinObj);
+
+protected:
+	virtual TSharedRef<SWidget> GetDefaultValueWidget() override;
+	TOptional<NumericType> GetNumericValue() const;
+	void SetNumericValue(NumericType NewValue, ETextCommit::Type CommitInfo);
+
+private:
+	int32 ArrayIndex;
+	UMaterialExpressionOperator* OperatorExpression;
+};
+
+
+class MATERIALEDITOR_API SGraphNodeMaterialOperator : public SGraphNodeMaterialBase
+{
+public:
+	SLATE_BEGIN_ARGS(SGraphNodeMaterialOperator) {}
+	SLATE_END_ARGS()
+
+	void Construct(const FArguments& InArgs, class UMaterialGraphNode* InNode);
+
+protected:
+	// SGraphNode interface
+	virtual void CreateBelowPinControls(TSharedPtr<SVerticalBox> MainBox) override;
+	virtual TSharedPtr<SGraphPin> CreatePinWidget(UEdGraphPin* InPin) const override;
+	virtual void SetDefaultTitleAreaWidget(TSharedRef<SOverlay> DefaultTitleAreaWidget) override;
+	// End of SGraphNode interface
+
+private:
+	/** Get the Material Expression */
+	UMaterialExpressionOperator* GetMaterialExpression() const;
+
+	/** Event for adding a pin */
+	FReply OnAddPinClicked();
+
+	/** Button on Operator Graph Node for Operator Selection Menu */
+	TWeakPtr<SButton> OperatorButton;
+
+	/** Event for Operator Selection Menu */
+	FReply OnOperatorMenuClicked();
+
+	/** Create Operator Selection Menu */
+	TSharedRef<SWidget> BuildOperatorMenu();
+
+	/** Event for selecting operator */
+	void OnOperatorMenuSelected(EMaterialExpressionOperatorKind NewKind);
+
+	/** Get Operator Display Name */
+	FText GetOperatorDisplayName() const;
+};
\ No newline at end of file
diff --git a/Engine/Source/Editor/UnrealEd/Classes/MaterialGraph/MaterialGraphNode_Operator.h b/Engine/Source/Editor/UnrealEd/Classes/MaterialGraph/MaterialGraphNode_Operator.h
new file mode 100644
--- /dev/null
+++ b/Engine/Source/Editor/UnrealEd/Classes/MaterialGraph/MaterialGraphNode_Operator.h
@@ -0,0 +1,22 @@
+// Copyright Epic Games, Inc. All Rights Reserved.
+
+#pragma once
+
+#include "MaterialGraph/MaterialGraph.h"
+#include "MaterialGraph/MaterialGraphNode.h"
+#include "MaterialGraph/MaterialGraphNode_Base.h"
+
+#include "MaterialGraphNode_Operator.generated.h"
+
+class UEdGraphPin;
+
+UCLASS(MinimalAPI)
+class UMaterialGraphNode_Operator : public UMaterialGraphNode
+{
+	GENERATED_UCLASS_BODY()
+
+public:
+	//~ Begin UEdGraphNode Interface.
+	virtual TSharedPtr<SGraphNode> CreateVisualWidget() override;
+	//~ End UEdGraphNode Interface.
+};
\ No newline at end of file
diff --git a/Engine/Source/Editor/UnrealEd/Private/MaterialGraph.cpp b/Engine/Source/Editor/UnrealEd/Private/MaterialGraph.cpp
--- a/Engine/Source/Editor/UnrealEd/Private/MaterialGraph.cpp
+++ b/Engine/Source/Editor/UnrealEd/Private/MaterialGraph.cpp
@@ -4,6 +4,7 @@
 // UMaterialGraph
 
 #include "MaterialGraph/MaterialGraph.h"
+#include "MaterialGraph/MaterialGraphNode_Operator.h"
 #include "MaterialGraph/MaterialGraphNode_Comment.h"
 #include "MaterialGraph/MaterialGraphNode_Custom.h"
 #include "MaterialGraph/MaterialGraphNode_Composite.h"
@@ -17,6 +18,7 @@
 #include "Materials/MaterialExpressionComposite.h"
 #include "Materials/MaterialExpressionCustom.h"
 #include "Materials/MaterialExpressionPinBase.h"
+#include "Materials/MaterialExpressionOperator.h"
 #include "Materials/MaterialExpressionFunctionOutput.h"
 #include "Materials/MaterialExpressionCustomOutput.h"
 #include "Materials/MaterialExpressionReroute.h"
@@ -289,6 +291,10 @@ UMaterialGraphNode* UMaterialGraph::AddExpression(UMaterialExpression* Expressio
 		{
 			Node = InitExpressionNewNode<UMaterialGraphNode_PinBase>(this, Expression, false);
 		}
+		else if (Expression->IsA(UMaterialExpressionOperator::StaticClass()))
+		{
+			Node = InitExpressionNewNode<UMaterialGraphNode_Operator>(this, Expression, false);
+		}
 		else if (Expression->IsA(UMaterialExpressionCustom::StaticClass()))
 		{
 			// This is for backward compatibility. We don't want people's custom HLSL nodes to suddenly explode
diff --git a/Engine/Source/Editor/UnrealEd/Private/MaterialGraphNode_Operator.cpp b/Engine/Source/Editor/UnrealEd/Private/MaterialGraphNode_Operator.cpp
new file mode 100644
--- /dev/null
+++ b/Engine/Source/Editor/UnrealEd/Private/MaterialGraphNode_Operator.cpp
@@ -0,0 +1,14 @@
+// Copyright Epic Games, Inc. All Rights Reserved.
+
+#include "MaterialGraph/MaterialGraphNode_Operator.h"
+#include "MaterialEditor/MaterialNodes/SGraphNodeMaterialOperator.h"
+
+UMaterialGraphNode_Operator::UMaterialGraphNode_Operator(const FObjectInitializer& ObjectInitializer)
+	: Super(ObjectInitializer)
+{
+}
+
+TSharedPtr<SGraphNode> UMaterialGraphNode_Operator::CreateVisualWidget()
+{
+	return SNew(SGraphNodeMaterialOperator, this);
+}
diff --git a/Engine/Source/Runtime/Engine/Private/Materials/MaterialExpressions.cpp b/Engine/Source/Runtime/Engine/Private/Materials/MaterialExpressions.cpp
--- a/Engine/Source/Runtime/Engine/Private/Materials/MaterialExpressions.cpp
+++ b/Engine/Source/Runtime/Engine/Private/Materials/MaterialExpressions.cpp
@@ -23931,6 +23931,8 @@ int32 UMaterialExpressionRecordTextureStreamingInfo::Compile(class FMaterialComp
 	return TextureCodeIndex;
 }
 
+#endif // WITH_EDITOR
+
 ///////////////////////////////////////////////////////////////////////////////
 // MaterialExpressionOperator
 ///////////////////////////////////////////////////////////////////////////////
@@ -23938,44 +23940,129 @@ int32 UMaterialExpressionRecordTextureStreamingInfo::Compile(class FMaterialComp
 // defined in MaterialExpressionsIR.cpp
 uint32 GetMaterialExpressionOperatorArity(EMaterialExpressionOperatorKind Operator);
 
-void UMaterialExpressionOperator::GetCaption(TArray<FString>& OutCaptions) const
+/** Array of Operators that allow adding inputs */
+static const EMaterialExpressionOperatorKind VariableInputsExpressionOperators[] = {
+	EMaterialExpressionOperatorKind::Add,
+	EMaterialExpressionOperatorKind::Multiply
+};
+
+UMaterialExpressionOperator::UMaterialExpressionOperator(const FObjectInitializer& ObjectInitializer)
+	: Super(ObjectInitializer)
 {
-	OutCaptions.Add(UEnum::GetDisplayValueAsText(Operator).ToString());
+	// Initialize default inputs
+	for (uint32 i = 0; i < Arity; i++)
+	{
+		DynamicInputs.Add(FMaterialExpressionOperatorInput());
+	}
 }
 
-FText UMaterialExpressionOperator::GetKeywords() const
+#if WITH_EDITOR
+
+bool UMaterialExpressionOperator::OperatorSupportsVariablePinCount(EMaterialExpressionOperatorKind InOperator)
 {
-	return FText::FromString(TEXT("operator"));
+	auto begin = std::begin(VariableInputsExpressionOperators);
+	auto end = std::end(VariableInputsExpressionOperators);
+	return std::find(begin, end, InOperator) != end;
 }
 
-FText UMaterialExpressionOperator::GetCreationName() const
+void UMaterialExpressionOperator::AddInputPin()
 {
-	return FText::FromString(TEXT("Operator"));
+	DynamicInputs.Add(FMaterialExpressionOperatorInput());
+	PreEditChange(nullptr);
+	PostEditChange();
+}
+
+bool UMaterialExpressionOperator::CanDeletePin(EEdGraphPinDirection PinDirection, int32 PinIndex) const
+{
+	return (uint32(DynamicInputs.Num()) > Arity);
+}
+
+void UMaterialExpressionOperator::DeletePin(EEdGraphPinDirection PinDirection, int32 PinIndex)
+{
+	if (DynamicInputs.IsValidIndex(PinIndex))
+	{
+		DynamicInputs.RemoveAt(PinIndex);
+	}
+	if (GraphNode)
+	{
+		GraphNode->ReconstructNode();
+		this->RefreshNode();
+	}
+}
+
+int32 UMaterialExpressionOperator::CountInputs() const
+{
+	return DynamicInputs.Num();
 }
 
 FExpressionInput* UMaterialExpressionOperator::GetInput(int32 InputIndex) 
 {
-	switch (InputIndex)
+	return (DynamicInputs.IsValidIndex(InputIndex)) ? &DynamicInputs[InputIndex].ExpressionInput : nullptr;
+}
+
+FName UMaterialExpressionOperator::GetInputName(int32 InputIndex) const
+{
+	// if InputIndex is 0-25, return name: A-Z
+	if (InputIndex < 26)
 	{
-		case 0: return &A;
-		case 1: return Arity >= 2 ? &B : nullptr;
-		case 2: return Arity >= 3 ? &C : nullptr;
-		default: return nullptr;
+		TCHAR Letter = TCHAR('A' + InputIndex);
+		return FName(*FString::Printf(TEXT("%c"), Letter));
 	}
+	// return name: Input N
+	return FName(*FString::Printf(TEXT("Input %d"), InputIndex + 1));
+}
+
+void UMaterialExpressionOperator::GetCaption(TArray<FString>& OutCaptions) const
+{
+	OutCaptions.Add(UEnum::GetDisplayValueAsText(Operator).ToString());
+}
+
+FText UMaterialExpressionOperator::GetKeywords() const
+{
+	return FText::FromString(TEXT("operator"));
+}
+
+FText UMaterialExpressionOperator::GetCreationName() const
+{
+	return FText::FromString(TEXT("Operator"));
 }
 
 void UMaterialExpressionOperator::PostEditChangeProperty(FPropertyChangedEvent& PropertyChangedEvent)
 {
+	Super::PostEditChangeProperty(PropertyChangedEvent);
+	UMaterialExpression::PostEditChangeProperty(PropertyChangedEvent);
+
+	bool bNeedsUpdate = false;
+
+	if (PropertyChangedEvent.GetPropertyName() == GET_MEMBER_NAME_CHECKED(ThisClass, DynamicInputs))
+	{
+		bNeedsUpdate = true;
+	}	
+	
 	if (PropertyChangedEvent.GetPropertyName() == TEXT("Operator"))
 	{
-		uint32 OldArity = Arity;
-		Arity = GetMaterialExpressionOperatorArity(Operator);
-		if (Arity != OldArity)
+		const uint32 ArityNew = GetMaterialExpressionOperatorArity(Operator);
+		const bool bAllowAddPinNew = OperatorSupportsVariablePinCount(Operator);
+
+		// Resize if switching between fixed and flexible operators, OR if the arity change
+		if (Arity != ArityNew || bAllowAddPin != bAllowAddPinNew)
 		{
-			GraphNode->ReconstructNode();
+			// If the operator remains flexible before and after, and DynamicInputs don't need to be updated according to Arity
+			// Otherwise, resize
+			if (!bAllowAddPinNew || !bAllowAddPin)
+			{
+				DynamicInputs.SetNum(ArityNew, EAllowShrinking::Yes);
+			}
+			Arity = ArityNew;
+			bAllowAddPin = bAllowAddPinNew;
+			bNeedsUpdate = true;
 		}
 	}
-	UMaterialExpression::PostEditChangeProperty(PropertyChangedEvent);
+
+	if (bNeedsUpdate)
+	{
+		GraphNode->ReconstructNode();
+	}
 }
 
 void UMaterialExpressionOperator::PostLoad()
diff --git a/Engine/Source/Runtime/Engine/Private/Materials/MaterialExpressionsIR.cpp b/Engine/Source/Runtime/Engine/Private/Materials/MaterialExpressionsIR.cpp
--- a/Engine/Source/Runtime/Engine/Private/Materials/MaterialExpressionsIR.cpp
+++ b/Engine/Source/Runtime/Engine/Private/Materials/MaterialExpressionsIR.cpp
@@ -1331,14 +1331,32 @@ uint32 GetMaterialExpressionOperatorArity(EMaterialExpressionOperatorKind Operat
 
 void UMaterialExpressionOperator::Build(MIR::FEmitter& Em)
 {
-	MIR::EOperator OpMIR = MaterialExpressionOperatorToMIR(Operator);
-	int32 OperatorArity = MIR::GetOperatorArity(OpMIR);
+	if (bAllowAddPin)
+	{
+		MIR::EOperator OpMIR = MaterialExpressionOperatorToMIR(Operator);
+
+		// Apply operation to iteratively to all input values
+		FValueRef Value = Em.InputDefaultFloat(&DynamicInputs[0].ExpressionInput, DynamicInputs[0].ConstValue);
+		for (int32 i = 1; i < DynamicInputs.Num(); i++)
+		{
+			FValueRef CurValue = Em.InputDefaultFloat(&DynamicInputs[i].ExpressionInput, DynamicInputs[i].ConstValue);
+			Value = Em.Operator(OpMIR, Value, CurValue);
+		}
+
+		Em.Output(0, Value);
+	}
+	else
+	{
+		MIR::EOperator OpMIR = MaterialExpressionOperatorToMIR(Operator);
+		int32 OperatorArity = MIR::GetOperatorArity(OpMIR);
 
-	FValueRef AValue = Em.InputDefaultFloat(&A, ConstA);
-	FValueRef BValue = OperatorArity >= 2 ? Em.InputDefaultFloat(&B, ConstB) : FValueRef{};
-	FValueRef CValue = OperatorArity >= 3 ? Em.InputDefaultFloat(&C, ConstC) : FValueRef{};
+		FValueRef AValue = Em.InputDefaultFloat(&DynamicInputs[0].ExpressionInput, DynamicInputs[0].ConstValue);
+		FValueRef BValue = OperatorArity >= 2 ? Em.InputDefaultFloat(&DynamicInputs[1].ExpressionInput, DynamicInputs[1].ConstValue) : FValueRef{};
+		FValueRef CValue = OperatorArity >= 3 ? Em.InputDefaultFloat(&DynamicInputs[2].ExpressionInput, DynamicInputs[2].ConstValue) : FValueRef{};
+
+		Em.Output(0, Em.Operator(OpMIR, AValue, BValue, CValue));
+	}
 
-	Em.Output(0, Em.Operator(OpMIR, AValue, BValue, CValue));
 }
 
 #endif // WITH_EDITOR
diff --git a/Engine/Source/Runtime/Engine/Public/Materials/MaterialExpressionOperator.h b/Engine/Source/Runtime/Engine/Public/Materials/MaterialExpressionOperator.h
--- a/Engine/Source/Runtime/Engine/Public/Materials/MaterialExpressionOperator.h
+++ b/Engine/Source/Runtime/Engine/Public/Materials/MaterialExpressionOperator.h
@@ -81,50 +81,64 @@ enum class EMaterialExpressionOperatorKind
 	Smoothstep,
 };
 
+USTRUCT()
+struct FMaterialExpressionOperatorInput
+{
+	GENERATED_USTRUCT_BODY()
+
+	UPROPERTY()
+	FExpressionInput ExpressionInput;
+
+	UPROPERTY(EditAnywhere, Category = "AddInput")
+	float ConstValue;
+
+	FMaterialExpressionOperatorInput()
+		: ConstValue(1.0f)
+	{
+	}
+};
+ 
 UCLASS(MinimalAPI)
 class UMaterialExpressionOperator : public UMaterialExpression
 {
 	GENERATED_BODY()
 
 public:
-
+	/** Array of Inputs */
+	UPROPERTY(EditAnywhere, Category = "Dynamic Inputs")
+	TArray<FMaterialExpressionOperatorInput> DynamicInputs;	
+	
+	/** Default Operator */
 	UPROPERTY(EditAnywhere, Category=MaterialExpressionOperator)
 	EMaterialExpressionOperatorKind Operator = EMaterialExpressionOperatorKind::Add;
 
-	UPROPERTY(meta = (RequiredInput = "false", ToolTip = "Defaults to 'ConstA' if not specified"))
-	FExpressionInput A;
-
-	UPROPERTY(meta = (RequiredInput = "false", ToolTip = "Defaults to 'ConstB' if not specified"))
-	FExpressionInput B;
-
-	UPROPERTY(meta = (RequiredInput = "false", ToolTip = "Defaults to 'ConstB' if not specified"))
-	FExpressionInput C;
-
-	// only used if A is not hooked up
-	UPROPERTY(EditAnywhere, Category=MaterialExpressionOperator, meta=(OverridingInputProperty = "A"))
-	float ConstA = 0.0f;
+	/** Default Number of Inputs */
+	uint32 Arity = 2;
 
-	// only used if B is not hooked up
-	UPROPERTY(EditAnywhere, Category=MaterialExpressionOperator, meta=(OverridingInputProperty = "B"))
-	float ConstB = 1.0f;
-	
-	// only used if B is not hooked up
-	UPROPERTY(EditAnywhere, Category=MaterialExpressionOperator, meta=(OverridingInputProperty = "C"))
-	float ConstC = 1.0f;
+	/** If the current operator allows Add a Pin feature */
+	bool bAllowAddPin = true;
 
-	//
-	uint32 Arity = 2;
+	UMaterialExpressionOperator(const FObjectInitializer& ObjectInitializer);
 
 	//~ Begin UMaterialExpression Interface
 #if WITH_EDITOR
+	static bool OperatorSupportsVariablePinCount(EMaterialExpressionOperatorKind InOperator);
+	virtual void AddInputPin();	
+	virtual bool CanDeletePin(EEdGraphPinDirection PinDirection, int32 PinIndex) const override;
+	virtual void DeletePin(EEdGraphPinDirection PinDirection, int32 PinIndex) override;	
+
+	virtual FExpressionInput* GetInput(int32 InputIndex) override;
+	virtual FName GetInputName(int32 InputIndex) const override;
+	virtual int32 CountInputs() const override;	
+
 	virtual void GetCaption(TArray<FString>& OutCaptions) const override;
 	virtual FText GetKeywords() const override;
 	virtual FText GetCreationName() const override;
-	virtual FExpressionInput* GetInput(int32 InputIndex) override;
+
 	virtual void PostEditChangeProperty(FPropertyChangedEvent& PropertyChangedEvent) override;
 	virtual void PostLoad() override;
-	virtual void Build(MIR::FEmitter& Emitter) override;
 
+	virtual void Build(MIR::FEmitter& Emitter) override;
 #endif // WITH_EDITOR
 	//~ End UMaterialExpression Interface
 };
diff --git a/Engine/Source/Runtime/SlateCore/Private/Styling/StarshipCoreStyle.cpp b/Engine/Source/Runtime/SlateCore/Private/Styling/StarshipCoreStyle.cpp
--- a/Engine/Source/Runtime/SlateCore/Private/Styling/StarshipCoreStyle.cpp
+++ b/Engine/Source/Runtime/SlateCore/Private/Styling/StarshipCoreStyle.cpp
@@ -434,7 +434,10 @@ TSharedRef<ISlateStyle> FStarshipCoreStyle::Create()
 
 		Style->Set("Icons.Calendar", new IMAGE_BRUSH_SVG("Starship/Common/Calendar", Icon16x16));
 
+		// Material Graph Node Icons
 		Style->Set("Icons.Preview", new IMAGE_BRUSH_SVG("Starship/Common/preview-default", Icon10x10));
+		Style->Set("Icons.DropDown", new IMAGE_BRUSH_SVG("Starship/Common/caret-down", Icon10x10));
+
 	}
 
 	// Tool panels
