Added function APIs to MaterialEditingLibrary to manipulate the new Enumeration material parameter.

diff --git a/Engine/Source/Editor/MaterialEditor/Private/MaterialEditingLibrary.cpp b/Engine/Source/Editor/MaterialEditor/Private/MaterialEditingLibrary.cpp
--- a/Engine/Source/Editor/MaterialEditor/Private/MaterialEditingLibrary.cpp
+++ b/Engine/Source/Editor/MaterialEditor/Private/MaterialEditingLibrary.cpp
@@ -24,6 +24,7 @@
 #include "Materials/MaterialExpressionDynamicParameter.h"
 #include "Materials/MaterialParameterCollection.h"
 #include "Materials/MaterialExpressionMaterialFunctionCall.h"
+#include "Materials/MaterialEnumeration.h"
 #include "MaterialEditor/MaterialEditorInstanceConstant.h"
 #include "MaterialStatsCommon.h"
 #include "Particles/ParticleSystemComponent.h"
@@ -35,6 +36,7 @@
 #include "Subsystems/AssetEditorSubsystem.h"
 #include "ShaderCompiler.h"
 #include "UObject/UObjectIterator.h"
+#include "MaterialCachedData.h"
 
 #include UE_INLINE_GENERATED_CPP_BY_NAME(MaterialEditingLibrary)
 
@@ -1271,6 +1273,27 @@ bool UMaterialEditingLibrary::SetMaterialInstanceStaticSwitchParameterValue(UMat
 	return bResult;
 }
 
+FName UMaterialEditingLibrary::GetMaterialInstanceEnumerationParameterValue(UMaterialInstanceConstant* Instance, FName ParameterName, EMaterialParameterAssociation Association)
+{
+	FName Result = {};
+	if (Instance)
+	{
+		Instance->GetEnumerationParameterValue(FHashedMaterialParameterInfo(ParameterName, Association, Association == EMaterialParameterAssociation::LayerParameter ? 0 : INDEX_NONE), Result);
+	}
+	return Result;
+}
+
+bool UMaterialEditingLibrary::SetMaterialInstanceEnumerationParameterValue(UMaterialInstanceConstant* Instance, FName ParameterName, FName Value, EMaterialParameterAssociation Association)
+{
+	bool bResult = false;
+	if (Instance)
+	{
+		Instance->SetEnumerationParameterValueEditorOnly(FMaterialParameterInfo(ParameterName, Association, Association == EMaterialParameterAssociation::LayerParameter ? 0 : INDEX_NONE), Value);
+		UMaterialEditingLibrary::UpdateMaterialInstance(Instance);
+	}
+	return bResult;
+}
+
 void UMaterialEditingLibrary::UpdateMaterialInstance(UMaterialInstanceConstant* Instance)
 {
 	if (Instance)
@@ -1373,7 +1396,23 @@ void UMaterialEditingLibrary::GetStaticSwitchParameterNames(UMaterialInterface*
 	}
 }
 
-bool GetParameterSource(UMaterialInterface* Material, const TArray<FMaterialParameterInfo>& Info, const TArray<FGuid>& Guids, const FName& ParameterName, FSoftObjectPath& OutParameterSource)
+void UMaterialEditingLibrary::GetEnumerationParameterNames(UMaterialInterface* Material, TArray<FName>& ParameterNames)
+{
+	ParameterNames.Empty();
+	if (Material)
+	{
+		TArray<FMaterialParameterInfo> MaterialInfo;
+		TArray<FGuid> MaterialGuids;
+		Material->GetAllEnumerationParameterInfo(MaterialInfo, MaterialGuids);
+
+		for (const FMaterialParameterInfo& Info : MaterialInfo)
+		{
+			ParameterNames.Add(Info.Name);
+		}
+	}
+}
+
+static bool GetParameterSource(UMaterialInterface* Material, const TArray<FMaterialParameterInfo>& Info, const TArray<FGuid>& Guids, const FName& ParameterName, FSoftObjectPath& OutParameterSource)
 {
 	bool bResult = false;
 	for (int32 Index = 0; Index < Info.Num(); ++Index)
@@ -1441,6 +1480,46 @@ bool UMaterialEditingLibrary::GetStaticSwitchParameterSource(UMaterialInterface*
 	return false;
 }
 
+bool UMaterialEditingLibrary::GetEnumerationParameterSource(UMaterialInterface* Material, const FName ParameterName, FSoftObjectPath& ParameterSource)
+{
+	if (Material)
+	{
+		TArray<FMaterialParameterInfo> MaterialInfo;
+		TArray<FGuid> MaterialGuids;
+		Material->GetAllEnumerationParameterInfo(MaterialInfo, MaterialGuids);
+		return GetParameterSource(Material, MaterialInfo, MaterialGuids, ParameterName, ParameterSource);
+	}
+	return false;
+}
+
+bool UMaterialEditingLibrary::GetEnumerationParameterEntryNames(UMaterialInterface* Material, FName ParameterName, EMaterialParameterAssociation Association, TArray<FName>& OutEntryNames)
+{
+	if (!Material)
+	{
+		return false;
+	}
+
+	FHashedMaterialParameterInfo ParameterInfo(ParameterName, Association, Association == EMaterialParameterAssociation::LayerParameter ? 0 : INDEX_NONE);
+	
+	FMaterialParameterMetadata ParameterMetadata;
+	if (!Material->GetCachedExpressionData().GetParameterValue(EMaterialParameterType::Enumeration, ParameterInfo, ParameterMetadata))
+	{
+		return false;
+	}
+
+	if (!ParameterMetadata.Enumeration)
+	{
+		return false;
+	}
+
+	Cast<IMaterialEnumerationProvider>(ParameterMetadata.Enumeration)->ForEachEntry([&OutEntryNames](FName EntryName, int32 EntryValue)
+	{
+		OutEntryNames.Add(EntryName);
+	});
+
+	return true;
+}
+
 FMaterialStatistics UMaterialEditingLibrary::GetStatistics(class UMaterialInterface* Material)
 {
 	FMaterialStatistics Result;
diff --git a/Engine/Source/Editor/MaterialEditor/Public/MaterialEditingLibrary.h b/Engine/Source/Editor/MaterialEditor/Public/MaterialEditingLibrary.h
--- a/Engine/Source/Editor/MaterialEditor/Public/MaterialEditingLibrary.h
+++ b/Engine/Source/Editor/MaterialEditor/Public/MaterialEditingLibrary.h
@@ -342,6 +342,14 @@ public:
 	UFUNCTION(BlueprintCallable, Category = "MaterialEditing")
 	static bool SetMaterialInstanceStaticSwitchParameterValue(UMaterialInstanceConstant* Instance, FName ParameterName, bool Value, EMaterialParameterAssociation Association = EMaterialParameterAssociation::GlobalParameter, bool bUpdateMaterialInstance = true);
 
+	/** Get the current enumeration parameter value from a Material Instance */
+	UFUNCTION(BlueprintPure, Category = "MaterialEditing")
+	static FName GetMaterialInstanceEnumerationParameterValue(UMaterialInstanceConstant* Instance, FName ParameterName, EMaterialParameterAssociation Association = EMaterialParameterAssociation::GlobalParameter);
+
+	/** Set the enumeration parameter value for a Material Instance */
+	UFUNCTION(BlueprintCallable, Category = "MaterialEditing")
+	static bool SetMaterialInstanceEnumerationParameterValue(UMaterialInstanceConstant* Instance, FName ParameterName, FName Value, EMaterialParameterAssociation Association = EMaterialParameterAssociation::GlobalParameter);
+
 	/** Called after making modifications to a Material Instance to recompile shaders etc. */
 	UFUNCTION(BlueprintCallable, Category = "MaterialEditing")
 	static void UpdateMaterialInstance(UMaterialInstanceConstant* Instance);
@@ -365,6 +373,10 @@ public:
 	/** Gets all static switch parameter names */
 	UFUNCTION(BlueprintCallable, Category = "MaterialEditing")
 	static void GetStaticSwitchParameterNames(UMaterialInterface* Material, TArray<FName>& ParameterNames);
+	
+	/** Gets all enumeration parameter names */
+	UFUNCTION(BlueprintCallable, Category = "MaterialEditing")
+	static void GetEnumerationParameterNames(UMaterialInterface* Material, TArray<FName>& ParameterNames);
 
 	/**
 	*	Returns the path of the asset where the parameter originated, as well as true/false if it was found
@@ -405,6 +417,27 @@ public:
 	*/
 	UFUNCTION(BlueprintCallable, Category = "MaterialEditing")
 	static bool GetStaticSwitchParameterSource(UMaterialInterface* Material, const FName ParameterName, FSoftObjectPath& ParameterSource);
+	
+	/**
+	*	Returns the path of the asset where the parameter originated, as well as true/false if it was found
+	*	@param	Material	The material or material instance you want to look up a parameter from
+	*	@param	ParameterName		The parameter name
+	*	@param	ParameterSource		The soft object path of the asset the parameter originates in
+	*	@return	Whether or not the parameter was found in this material
+	*/
+	UFUNCTION(BlueprintCallable, Category = "MaterialEditing")
+	static bool GetEnumerationParameterSource(UMaterialInterface* Material, const FName ParameterName, FSoftObjectPath& ParameterSource);
+	
+	/**
+	*	Returns the array of entry names in the enumeration definition referenced by specified enumeration parameter.
+	*	@param	Material				The material or material instance you want to look up a parameter from
+	*	@param	ParameterName			The parameter name
+	*   @param  Association				The parameter association
+	*	@param	OutEntryNames	The output array of enumeration definition entry names.
+	*	@return	Whether or not the parameter was found in this material
+	*/
+	UFUNCTION(BlueprintCallable, Category = "MaterialEditing")
+	static bool GetEnumerationParameterEntryNames(UMaterialInterface* Material, FName ParameterName, EMaterialParameterAssociation Association, TArray<FName>& OutEntryNames);
 
 	/** Returns statistics about the given material */
 	UFUNCTION(BlueprintCallable, Category = "MaterialEditing")
diff --git a/Engine/Source/Runtime/Engine/Private/Materials/MaterialInstanceConstant.cpp b/Engine/Source/Runtime/Engine/Private/Materials/MaterialInstanceConstant.cpp
--- a/Engine/Source/Runtime/Engine/Private/Materials/MaterialInstanceConstant.cpp
+++ b/Engine/Source/Runtime/Engine/Private/Materials/MaterialInstanceConstant.cpp
@@ -154,7 +154,13 @@ void UMaterialInstanceConstant::SetSparseVolumeTextureParameterValueEditorOnly(c
 void UMaterialInstanceConstant::SetFontParameterValueEditorOnly(const FMaterialParameterInfo& ParameterInfo,class UFont* FontValue,int32 FontPage)
 {
 	check(GIsEditor || IsRunningCommandlet());
-	SetFontParameterValueInternal(ParameterInfo,FontValue,FontPage);
+	SetFontParameterValueInternal(ParameterInfo, FontValue, FontPage);
+}
+
+void UMaterialInstanceConstant::SetEnumerationParameterValueEditorOnly(const FMaterialParameterInfo& ParameterInfo, FName Value)
+{
+	check(GIsEditor || IsRunningCommandlet());
+	SetEnumerationParameterValueInternal(ParameterInfo, Value);
 }
 
 void UMaterialInstanceConstant::ClearParameterValuesEditorOnly()
diff --git a/Engine/Source/Runtime/Engine/Private/Materials/MaterialInterface.cpp b/Engine/Source/Runtime/Engine/Private/Materials/MaterialInterface.cpp
--- a/Engine/Source/Runtime/Engine/Private/Materials/MaterialInterface.cpp
+++ b/Engine/Source/Runtime/Engine/Private/Materials/MaterialInterface.cpp
@@ -1018,6 +1018,17 @@ bool UMaterialInterface::GetFontParameterValue(const FHashedMaterialParameterInf
 	return false;
 }
 
+bool UMaterialInterface::GetEnumerationParameterValue(const FHashedMaterialParameterInfo& ParameterInfo, FName& OutValue, bool bOveriddenOnly) const
+{
+	FMaterialParameterMetadata Result;
+	if (GetParameterValue(EMaterialParameterType::Enumeration, ParameterInfo, Result, MakeParameterValueFlags(bOveriddenOnly)))
+	{
+		OutValue = Result.Value.EnumEntry;
+		return true;
+	}
+	return false;
+}
+
 bool UMaterialInterface::GetParameterDefaultValue(EMaterialParameterType Type, const FMemoryImageMaterialParameterInfo& ParameterInfo, FMaterialParameterMetadata& OutValue) const
 {
 	return GetParameterValue(Type, ParameterInfo, OutValue, EMaterialGetParameterValueFlags::CheckNonOverrides);
@@ -1113,6 +1124,16 @@ bool UMaterialInterface::GetFontParameterDefaultValue(const FHashedMaterialParam
 	return false;
 }
 
+bool UMaterialInterface::GetEnumerationParameterDefaultValue(const FHashedMaterialParameterInfo& ParameterInfo, FName& OutValue) const
+{
+	FMaterialParameterMetadata Result;
+	if (GetParameterDefaultValue(EMaterialParameterType::Enumeration, ParameterInfo, Result))
+	{
+		OutValue = Result.Value.EnumEntry;
+		return true;
+	}
+	return false;
+}
 
 #if WITH_EDITOR
 bool UMaterialInterface::GetStaticSwitchParameterDefaultValue(const FHashedMaterialParameterInfo& ParameterInfo, bool& OutValue, FGuid& OutExpressionGuid) const
@@ -1197,6 +1218,11 @@ void UMaterialInterface::GetAllFontParameterInfo(TArray<FMaterialParameterInfo>&
 	GetAllParameterInfoOfType(EMaterialParameterType::Font, OutParameterInfo, OutParameterIds);
 }
 
+void UMaterialInterface::GetAllEnumerationParameterInfo(TArray<FMaterialParameterInfo>& OutParameterInfo, TArray<FGuid>& OutParameterIds) const
+{
+	GetAllParameterInfoOfType(EMaterialParameterType::Enumeration, OutParameterInfo, OutParameterIds);
+}
+
 #if WITH_EDITORONLY_DATA
 void UMaterialInterface::GetAllStaticSwitchParameterInfo(TArray<FMaterialParameterInfo>&OutParameterInfo, TArray<FGuid>&OutParameterIds) const
 {
diff --git a/Engine/Source/Runtime/Engine/Public/Materials/MaterialInstanceConstant.h b/Engine/Source/Runtime/Engine/Public/Materials/MaterialInstanceConstant.h
--- a/Engine/Source/Runtime/Engine/Public/Materials/MaterialInstanceConstant.h
+++ b/Engine/Source/Runtime/Engine/Public/Materials/MaterialInstanceConstant.h
@@ -100,6 +100,7 @@ class UMaterialInstanceConstant : public UMaterialInstance
 	ENGINE_API void SetRuntimeVirtualTextureParameterValueEditorOnly(const FMaterialParameterInfo& ParameterInfo, class URuntimeVirtualTexture* Value);
 	ENGINE_API void SetSparseVolumeTextureParameterValueEditorOnly(const FMaterialParameterInfo& ParameterInfo, class USparseVolumeTexture* Value);
 	ENGINE_API void SetFontParameterValueEditorOnly(const FMaterialParameterInfo& ParameterInfo, class UFont* FontValue, int32 FontPage);
+	ENGINE_API void SetEnumerationParameterValueEditorOnly(const FMaterialParameterInfo& ParameterInfo, FName Value);
 
 	/**
 	 * Clear all parameter overrides on this material instance. This function
diff --git a/Engine/Source/Runtime/Engine/Public/Materials/MaterialInterface.h b/Engine/Source/Runtime/Engine/Public/Materials/MaterialInterface.h
--- a/Engine/Source/Runtime/Engine/Public/Materials/MaterialInterface.h
+++ b/Engine/Source/Runtime/Engine/Public/Materials/MaterialInterface.h
@@ -715,6 +715,7 @@ public:
 	ENGINE_API void GetAllRuntimeVirtualTextureParameterInfo(TArray<FMaterialParameterInfo>& OutParameterInfo, TArray<FGuid>& OutParameterIds) const;
 	ENGINE_API void GetAllSparseVolumeTextureParameterInfo(TArray<FMaterialParameterInfo>& OutParameterInfo, TArray<FGuid>& OutParameterIds) const;
 	ENGINE_API void GetAllFontParameterInfo(TArray<FMaterialParameterInfo>& OutParameterInfo, TArray<FGuid>& OutParameterIds) const;
+	ENGINE_API void GetAllEnumerationParameterInfo(TArray<FMaterialParameterInfo>& OutParameterInfo, TArray<FGuid>& OutParameterIds) const;
 
 #if WITH_EDITORONLY_DATA
 	ENGINE_API void GetAllStaticSwitchParameterInfo(TArray<FMaterialParameterInfo>& OutParameterInfo, TArray<FGuid>& OutParameterIds) const;
@@ -735,6 +736,7 @@ public:
 	ENGINE_API bool GetRuntimeVirtualTextureParameterDefaultValue(const FHashedMaterialParameterInfo& ParameterInfo, class URuntimeVirtualTexture*& OutValue) const;
 	ENGINE_API bool GetSparseVolumeTextureParameterDefaultValue(const FHashedMaterialParameterInfo& ParameterInfo, class USparseVolumeTexture*& OutValue) const;
 	ENGINE_API bool GetFontParameterDefaultValue(const FHashedMaterialParameterInfo& ParameterInfo, class UFont*& OutFontValue, int32& OutFontPage) const;
+	ENGINE_API bool GetEnumerationParameterDefaultValue(const FHashedMaterialParameterInfo& ParameterInfo, FName& OutValue) const;
 	
 #if WITH_EDITOR
 	ENGINE_API bool GetStaticSwitchParameterDefaultValue(const FHashedMaterialParameterInfo& ParameterInfo, bool& OutValue, FGuid& OutExpressionGuid) const;
@@ -965,6 +967,8 @@ public:
 	ENGINE_API bool GetTextureParameterChannelNames(const FHashedMaterialParameterInfo& ParameterInfo, FParameterChannelNames& OutValue) const;
 #endif
 	ENGINE_API bool GetFontParameterValue(const FHashedMaterialParameterInfo& ParameterInfo,class UFont*& OutFontValue, int32& OutFontPage, bool bOveriddenOnly = false) const;
+	ENGINE_API bool GetEnumerationParameterValue(const FHashedMaterialParameterInfo& ParameterInfo, FName& OutValue, bool bOveriddenOnly = false) const;
+
 	ENGINE_API virtual bool GetRefractionSettings(float& OutBiasValue) const;
 	ENGINE_API virtual bool GetUserSceneTextureOverride(FName& InOutValue) const;
 	ENGINE_API FName GetUserSceneTextureOutput(const UMaterial* Base) const;
