[MaterialEditor] - Fix bug where incorrect sampler counts were reported in the Material Stats window in the Material Editor.

- This happened because we were using ODSC to only compile a small subset to decrease compile times to report stats.  This has the side effect of not compiling all shader types and vertex factory types.  Some of these types use more samplers as a baseline.
- This resulted in errors at cook time due to sampler limits and users not able to reproduce the issue in editor.
- The fix here is to revert the ODSC code and compile all shaders for a given material when generating platform stats
- We would rather pay for an increase in compile time in order to get accurate information.
- More work could be done to provide better estimates in the future.
- This code compiles all the shaders in the background so the user can still work in the editor.

[RN] major | Rendering - Materials and Shaders

diff --git a/Engine/Source/Editor/MaterialEditor/Private/MaterialStats.cpp b/Engine/Source/Editor/MaterialEditor/Private/MaterialStats.cpp
--- a/Engine/Source/Editor/MaterialEditor/Private/MaterialStats.cpp
+++ b/Engine/Source/Editor/MaterialEditor/Private/MaterialStats.cpp
@@ -293,49 +293,8 @@ bool FShaderPlatformSettings::CheckShaders(bool bIgnoreCooldown)
 						MaterialInstance->UpdateCachedData();
 					}
 
-					TMap<FName, TArray<FMaterialStatsUtils::FRepresentativeShaderInfo>> ShaderTypeNamesAndDescriptions;
-					FMaterialStatsUtils::GetRepresentativeShaderTypesAndDescriptions(ShaderTypeNamesAndDescriptions, Instance.MaterialResourcesStats);
-
-					TArray<const FVertexFactoryType*> VFTypes;
-					TArray<const FShaderPipelineType*> PipelineTypes;
-					TArray<const FShaderType*> ShaderTypes;
-
-					for (auto& DescriptionPair : ShaderTypeNamesAndDescriptions)
-					{
-						const FVertexFactoryType* VFType = FindVertexFactoryType(DescriptionPair.Key);
-						check(VFType);
-
-						auto& DescriptionArray = DescriptionPair.Value;
-						for (const FMaterialStatsUtils::FRepresentativeShaderInfo& ShaderInfo : DescriptionArray)
-						{
-							const FShaderType* ShaderType = FindShaderTypeByName(ShaderInfo.ShaderName);
-
-							// in compile only setting we only care if derived MI's are compiling at all
-							// so we only take one shader combination not to slow down material stats window too much
-							// it's likely sufficient to treat first shader as most complex, otherwise we need to refactor FindShaderTypeByName
-							if (Instance.bOnlyCompileMostComplexShader && VFTypes.Num() >= 1)
-							{
-								break;
-							}
-
-							if (ShaderType && VFType)
-							{
-								VFTypes.Add(VFType);
-								ShaderTypes.Add(ShaderType);
-								PipelineTypes.Add(nullptr);
-							}
-						}
-					}
-
 					// Prepare the resource for compilation, but don't compile the completed shader map.
-					const bool bSuccess = Instance.MaterialResourcesStats->CacheShaders(PlatformShaderID, EMaterialShaderPrecompileMode::None);
-
-					if (bSuccess)
-					{
-						// Compile just the types we want.
-						Instance.MaterialResourcesStats->CacheGivenTypes(PlatformShaderID, VFTypes, PipelineTypes, ShaderTypes);
-					}
-
+					Instance.MaterialResourcesStats->CacheShaders(PlatformShaderID);
 					Instance.bCompilingShaders = true;
 					Instance.bUpdateShaderCode = true;
 					Instance.bNeedShaderRecompilation = false;
