[Materials] - Material Functions now properly display generated Material code in the HLSL Code window in the Material Editor.

- When the Material Editor is editing a Material Function it previews the Function Output to visualize the material function.
- This is a separate preview material and that is different than the `Material` being edited _or_ the `OriginalMaterial`.
- If we are editing a Material Function we need to tell the stats manager to use the `ExpressionPreviewMaterial` material (which is previewing the mat function output).
- This way the stats manager is pulling the HLSL code from the correct material.

[RN] major | Rendering - Materials and Shaders

diff --git a/Engine/Source/Editor/MaterialEditor/Private/MaterialEditor.cpp b/Engine/Source/Editor/MaterialEditor/Private/MaterialEditor.cpp
--- a/Engine/Source/Editor/MaterialEditor/Private/MaterialEditor.cpp
+++ b/Engine/Source/Editor/MaterialEditor/Private/MaterialEditor.cpp
@@ -2756,7 +2756,14 @@ void FMaterialEditor::UpdatePreviewMaterial( bool bForce )
 		CreateDerivedMaterialInstancesPreviews();
 	}
 
-	MaterialStatsManager->SetMaterial(bStatsFromPreviewMaterial ? Material : OriginalMaterial, bStatsFromPreviewMaterial ? DerivedMaterialInstances : OriginalDerivedMaterialInstances);
+	if (MaterialFunction && ExpressionPreviewMaterial)
+	{
+		MaterialStatsManager->SetMaterial(ExpressionPreviewMaterial, bStatsFromPreviewMaterial ? DerivedMaterialInstances : OriginalDerivedMaterialInstances);
+	}
+	else
+	{
+		MaterialStatsManager->SetMaterial(bStatsFromPreviewMaterial ? Material : OriginalMaterial, bStatsFromPreviewMaterial ? DerivedMaterialInstances : OriginalDerivedMaterialInstances);
+	}
 	MaterialStatsManager->SignalMaterialChanged();
 
 	// Reregister all components that use the preview material, since UMaterial::PEC does not reregister components using a bIsPreviewMaterial=true material
