In LayersFunctions instance parameter, implement the new enumeration ui feature

diff --git a/Engine/Source/Editor/MaterialEditor/Private/SMaterialLayersFunctionsTree.cpp b/Engine/Source/Editor/MaterialEditor/Private/SMaterialLayersFunctionsTree.cpp
--- a/Engine/Source/Editor/MaterialEditor/Private/SMaterialLayersFunctionsTree.cpp
+++ b/Engine/Source/Editor/MaterialEditor/Private/SMaterialLayersFunctionsTree.cpp
@@ -619,36 +619,70 @@ void SMaterialLayersFunctionsInstanceTreeItem::Construct(const FArguments& InArg
 				]
 			];
 		}
-		else if (ScalarParam && ScalarParam->AtlasData.bIsUsedAsAtlasPosition)
+		else if (ScalarParam)
 		{
-			const FText ParameterName = FText::FromName(StackParameterData->Parameter->ParameterInfo.Name);
+			if (ScalarParam->AtlasData.bIsUsedAsAtlasPosition)
+			{
+				const FText ParameterName = FText::FromName(StackParameterData->Parameter->ParameterInfo.Name);
 
-			FDetailWidgetRow& CustomWidget = Row.CustomWidget();
-			CustomWidget
-				.FilterString(ParameterName)
-				.NameContent()
-				[
-					SNew(STextBlock)
-					.Text(ParameterName)
-					.ToolTipText(FMaterialPropertyHelpers::GetParameterExpressionDescription(StackParameterData->Parameter, MaterialEditorInstance))
-					.Font(FAppStyle::GetFontStyle(TEXT("PropertyWindow.NormalFont")))
-				]
-				.ValueContent()
-				.HAlign(HAlign_Fill)
-				.MaxDesiredWidth(400.0f)
-				[
-					SNew(SObjectPropertyEntryBox)
-					.ObjectPath(this, &SMaterialLayersFunctionsInstanceTreeItem::GetCurvePath, ScalarParam)
-					.AllowedClass(UCurveLinearColor::StaticClass())
-					.NewAssetFactories(TArray<UFactory*>())
-					.DisplayThumbnail(true)
-					.ThumbnailPool(InArgs._InTree->GetTreeThumbnailPool())
-					.OnShouldSetAsset(FOnShouldSetAsset::CreateStatic(&FMaterialPropertyHelpers::OnShouldSetCurveAsset, ScalarParam->AtlasData.Atlas))
-					.OnShouldFilterAsset(FOnShouldFilterAsset::CreateStatic(&FMaterialPropertyHelpers::OnShouldFilterCurveAsset, ScalarParam->AtlasData.Atlas))
-					.OnObjectChanged(FOnSetObject::CreateStatic(&FMaterialPropertyHelpers::SetPositionFromCurveAsset, ScalarParam->AtlasData.Atlas, ScalarParam, StackParameterData->ParameterHandle, (UObject*)MaterialEditorInstance))
-					.DisplayCompactSize(true)
-				];
-			
+				FDetailWidgetRow& CustomWidget = Row.CustomWidget();
+				CustomWidget
+					.FilterString(ParameterName)
+					.NameContent()
+					[
+						SNew(STextBlock)
+						.Text(ParameterName)
+						.ToolTipText(FMaterialPropertyHelpers::GetParameterExpressionDescription(StackParameterData->Parameter, MaterialEditorInstance))
+						.Font(FAppStyle::GetFontStyle(TEXT("PropertyWindow.NormalFont")))
+					]
+					.ValueContent()
+					.HAlign(HAlign_Fill)
+					.MaxDesiredWidth(400.0f)
+					[
+						SNew(SObjectPropertyEntryBox)
+						.ObjectPath(this, &SMaterialLayersFunctionsInstanceTreeItem::GetCurvePath, ScalarParam)
+						.AllowedClass(UCurveLinearColor::StaticClass())
+						.NewAssetFactories(TArray<UFactory*>())
+						.DisplayThumbnail(true)
+						.ThumbnailPool(InArgs._InTree->GetTreeThumbnailPool())
+						.OnShouldSetAsset(FOnShouldSetAsset::CreateStatic(&FMaterialPropertyHelpers::OnShouldSetCurveAsset, ScalarParam->AtlasData.Atlas))
+						.OnShouldFilterAsset(FOnShouldFilterAsset::CreateStatic(&FMaterialPropertyHelpers::OnShouldFilterCurveAsset, ScalarParam->AtlasData.Atlas))
+						.OnObjectChanged(FOnSetObject::CreateStatic(&FMaterialPropertyHelpers::SetPositionFromCurveAsset, ScalarParam->AtlasData.Atlas, ScalarParam, StackParameterData->ParameterHandle, (UObject*)MaterialEditorInstance))
+						.DisplayCompactSize(true)
+					];
+			}
+			else if (!ScalarParam->bUseCustomPrimitiveData)
+			{
+				// Try to find a valid enumeration
+				TSoftObjectPtr<UObject> Enumeration;
+				if (!ScalarParam->Enumeration.IsNull())
+				{
+					Enumeration = ScalarParam->Enumeration;
+				}
+				else if (ScalarParam->EnumerationIndex >= 0)
+				{
+					Enumeration = MaterialEditorInstance->SourceInstance->GetEnumerationObject(ScalarParam->EnumerationIndex);
+				}
+				// If found a valid Enumeration, customize the widget, else fallback todefault scalar numeric edit box
+				if (!Enumeration.IsNull())
+				{
+					const FText ParameterName = FText::FromName(StackParameterData->Parameter->ParameterInfo.Name);
+					TSharedPtr<IPropertyHandle> ParameterValueProperty = StackParameterData->ParameterNode->CreatePropertyHandle();
+
+					Row.CustomWidget()
+						.FilterString(ParameterName)
+						.NameContent()
+						[
+							ParameterValueProperty->CreatePropertyNameWidget()
+						]
+						.ValueContent()
+						.HAlign(HAlign_Fill)
+						.MaxDesiredWidth(400.0f)
+						[
+							FMaterialPropertyHelpers::CreateScalarEnumerationParameterValueWidget(Enumeration, ParameterValueProperty)
+						];
+				}
+			}
 		}
 		else if (TextureParam)
 		{
@@ -2352,33 +2386,65 @@ void SMaterialLayersFunctionsMaterialTreeItem::Construct(const FArguments& InArg
 					]
 				];
 		}
-		else if (ScalarParam && ScalarParam->AtlasData.bIsUsedAsAtlasPosition)
+		else if (ScalarParam)
 		{
-			const FText ParameterName = FText::FromName(StackParameterData->Parameter->ParameterInfo.Name);
+			if (ScalarParam->AtlasData.bIsUsedAsAtlasPosition)
+			{
+				const FText ParameterName = FText::FromName(StackParameterData->Parameter->ParameterInfo.Name);
 
-			FDetailWidgetRow& CustomWidget = Row.CustomWidget();
-			CustomWidget
-				.FilterString(ParameterName)
-				.NameContent()
-				[
-					SNew(STextBlock)
-					.Text(ParameterName)
-					.ToolTipText(FMaterialPropertyHelpers::GetParameterExpressionDescription(StackParameterData->Parameter, MaterialEditorInstance))
-					.Font(FAppStyle::GetFontStyle(TEXT("PropertyWindow.NormalFont")))
-				]
-			.ValueContent()
-				.HAlign(HAlign_Fill)
-				.MaxDesiredWidth(400.0f)
-				[
-					SNew(SObjectPropertyEntryBox)
-					.ObjectPath(this, &SMaterialLayersFunctionsMaterialTreeItem::GetCurvePath, ScalarParam)
-					.AllowedClass(UCurveLinearColor::StaticClass())
-					.NewAssetFactories(TArray<UFactory*>())
-					.DisplayThumbnail(true)
-					.ThumbnailPool(InArgs._InTree->GetTreeThumbnailPool())
-					.DisplayCompactSize(true)
-				];
+				FDetailWidgetRow& CustomWidget = Row.CustomWidget();
+				CustomWidget
+					.FilterString(ParameterName)
+					.NameContent()
+					[
+						SNew(STextBlock)
+						.Text(ParameterName)
+						.ToolTipText(FMaterialPropertyHelpers::GetParameterExpressionDescription(StackParameterData->Parameter, MaterialEditorInstance))
+						.Font(FAppStyle::GetFontStyle(TEXT("PropertyWindow.NormalFont")))
+					]
+				.ValueContent()
+					.HAlign(HAlign_Fill)
+					.MaxDesiredWidth(400.0f)
+					[
+						SNew(SObjectPropertyEntryBox)
+						.ObjectPath(this, &SMaterialLayersFunctionsMaterialTreeItem::GetCurvePath, ScalarParam)
+						.AllowedClass(UCurveLinearColor::StaticClass())
+						.NewAssetFactories(TArray<UFactory*>())
+						.DisplayThumbnail(true)
+						.ThumbnailPool(InArgs._InTree->GetTreeThumbnailPool())
+						.DisplayCompactSize(true)
+					];
+			}
+			else if (!ScalarParam->bUseCustomPrimitiveData)
+			{
+				// Try to find a valid enumeration
+				TSoftObjectPtr<UObject> Enumeration;
+				if (!ScalarParam->Enumeration.IsNull())
+				{
+					Enumeration = ScalarParam->Enumeration;
+				}
+				// EnumerationIndex is not supported in a Material
 
+				// If found a valid Enumeration, customize the widget, else fallback todefault scalar numeric edit box
+				if (!Enumeration.IsNull())
+				{
+					const FText ParameterName = FText::FromName(StackParameterData->Parameter->ParameterInfo.Name);
+					TSharedPtr<IPropertyHandle> ParameterValueProperty = StackParameterData->ParameterNode->CreatePropertyHandle();
+
+					Row.CustomWidget()
+						.FilterString(ParameterName)
+						.NameContent()
+						[
+							ParameterValueProperty->CreatePropertyNameWidget()
+						]
+						.ValueContent()
+						.HAlign(HAlign_Fill)
+						.MaxDesiredWidth(400.0f)
+						[
+							FMaterialPropertyHelpers::CreateScalarEnumerationParameterValueWidget(Enumeration, ParameterValueProperty)
+						];
+				}
+			}
 		}
 		else if (TextureParam)
 		{
