[Materials] - Refactor common logic into a helper function UMaterialExpression::ShouldShowPreview()

#rnx

diff --git a/Engine/Source/Editor/MaterialEditor/Private/MaterialEditor.cpp b/Engine/Source/Editor/MaterialEditor/Private/MaterialEditor.cpp
--- a/Engine/Source/Editor/MaterialEditor/Private/MaterialEditor.cpp
+++ b/Engine/Source/Editor/MaterialEditor/Private/MaterialEditor.cpp
@@ -7550,7 +7550,7 @@ void FMaterialEditor::RefreshExpressionPreview(UMaterialExpression* MaterialExpr
 FMatExpressionPreview* FMaterialEditor::GetExpressionPreview(UMaterialExpression* MaterialExpression, bool& bNewlyCreated)
 {
 	bNewlyCreated = false;
-	if (!MaterialExpression->bHidePreviewWindow && !MaterialExpression->bCollapsed && !MaterialExpression->IsA<UMaterialExpressionCustomOutput>())
+	if (MaterialExpression->ShouldShowPreview() && !MaterialExpression->IsA<UMaterialExpressionCustomOutput>())
 	{
 		FMatExpressionPreview* Preview = NULL;
 		for( int32 PreviewIndex = 0 ; PreviewIndex < ExpressionPreviews.Num() ; ++PreviewIndex )
diff --git a/Engine/Source/Editor/MaterialEditor/Private/MaterialNodes/SGraphNodeMaterialBase.cpp b/Engine/Source/Editor/MaterialEditor/Private/MaterialNodes/SGraphNodeMaterialBase.cpp
--- a/Engine/Source/Editor/MaterialEditor/Private/MaterialNodes/SGraphNodeMaterialBase.cpp
+++ b/Engine/Source/Editor/MaterialEditor/Private/MaterialNodes/SGraphNodeMaterialBase.cpp
@@ -637,8 +637,7 @@ TSharedRef<SWidget> SGraphNodeMaterialBase::CreatePreviewWidget()
 {
 	PreviewViewport.Reset();
 
-	// if this node should currently show a preview
-	if (!MaterialNode->MaterialExpression->bHidePreviewWindow && !MaterialNode->MaterialExpression->bCollapsed)
+	if (MaterialNode->MaterialExpression->ShouldShowPreview())
 	{
 		const float ExpressionPreviewSize = 106.0f;
 		const float CentralPadding = 5.0f;
@@ -694,9 +693,7 @@ FText SGraphNodeMaterialBase::ExpressionPreviewOverlayText() const
 
 EVisibility SGraphNodeMaterialBase::ExpressionPreviewVisibility() const
 {
-	UMaterialExpression* MaterialExpression = MaterialNode->MaterialExpression;
-	const bool bShowPreview = !MaterialExpression->bHidePreviewWindow && !MaterialExpression->bCollapsed;
-	return bShowPreview ? EVisibility::Visible : EVisibility::Collapsed;
+	return MaterialNode->MaterialExpression->ShouldShowPreview() ? EVisibility::Visible : EVisibility::Collapsed;
 }
 
 FReply SGraphNodeMaterialBase::OnTogglePreviewClicked()
diff --git a/Engine/Source/Runtime/Engine/Private/Materials/MaterialExpressions.cpp b/Engine/Source/Runtime/Engine/Private/Materials/MaterialExpressions.cpp
--- a/Engine/Source/Runtime/Engine/Private/Materials/MaterialExpressions.cpp
+++ b/Engine/Source/Runtime/Engine/Private/Materials/MaterialExpressions.cpp
@@ -2163,6 +2163,12 @@ void UMaterialExpression::ValidateState()
 }
 
 #if WITH_EDITOR
+
+bool UMaterialExpression::ShouldShowPreview() const
+{
+	return !bHidePreviewWindow && !bCollapsed;
+}
+
 bool UMaterialExpression::GetAllInputExpressions(TArray<UMaterialExpression*>& InputExpressions)
 {
 	// Make sure we don't end up in a loop
diff --git a/Engine/Source/Runtime/Engine/Public/Materials/MaterialExpression.h b/Engine/Source/Runtime/Engine/Public/Materials/MaterialExpression.h
--- a/Engine/Source/Runtime/Engine/Public/Materials/MaterialExpression.h
+++ b/Engine/Source/Runtime/Engine/Public/Materials/MaterialExpression.h
@@ -529,6 +529,10 @@ class UMaterialExpression : public UObject
 	ENGINE_API void ValidateState();
 
 #if WITH_EDITOR
+
+	/** Returns whether or not this expression should show the preview widget */
+	ENGINE_API bool ShouldShowPreview() const;
+
 	/** Returns the keywords that should be used when searching for this expression */
 	virtual FText GetKeywords() const {return FText::GetEmpty();}
 
