Restrict Material Parameter Collection parameter dropdown to only child collections.

Lightly refactor code that makes material instance editor property rows. Fixes a couple inconsistencies across parameter types, most notably the "Reset to Default" arrow is visible on some properties even when set to the default value.

#rnx

diff --git a/Engine/Source/Editor/MaterialEditor/Private/MaterialEditorInstanceDetailCustomization.cpp b/Engine/Source/Editor/MaterialEditor/Private/MaterialEditorInstanceDetailCustomization.cpp
--- a/Engine/Source/Editor/MaterialEditor/Private/MaterialEditorInstanceDetailCustomization.cpp
+++ b/Engine/Source/Editor/MaterialEditor/Private/MaterialEditorInstanceDetailCustomization.cpp
@@ -48,6 +48,7 @@
 #include "Factories/MaterialInstanceConstantFactoryNew.h"
 #include "Modules/ModuleManager.h"
 #include "AssetToolsModule.h"
+#include "AssetRegistry/IAssetRegistry.h"
 #include "Materials/MaterialFunctionInterface.h"
 #include "Materials/MaterialFunction.h"
 #include "Materials/MaterialFunctionInstance.h"
@@ -524,8 +525,10 @@ void FMaterialInstanceParameterDetails::CreateSingleGroupWidget(FEditorParameter
 	for (int32 ParamIdx = 0; ParamIdx < ParameterGroup.Parameters.Num(); ++ParamIdx)
 	{
 		TSharedPtr<IPropertyHandle> ParameterProperty = ParametersArrayProperty->GetChildHandle(ParamIdx);
+		TSharedPtr<IPropertyHandle> ParameterValueProperty = ParameterProperty.IsValid() ? ParameterProperty->GetChildHandle("ParameterValue") : nullptr;
+
 		UDEditorParameterValue* Parameter = ParameterGroup.Parameters[ParamIdx];
-		if (ParameterProperty.IsValid() &&
+		if (ParameterProperty.IsValid() && ParameterValueProperty.IsValid() && ParameterValueProperty->IsValidHandle() &&
 			(GroupIndex == INDEX_NONE || Parameter->ParameterInfo.Index == GroupIndex))
 		{
 			UDEditorFontParameterValue* FontParam = Cast<UDEditorFontParameterValue>(Parameter);
@@ -572,11 +575,11 @@ void FMaterialInstanceParameterDetails::CreateSingleGroupWidget(FEditorParameter
 			{
 				if (VectorParam && VectorParam->bIsUsedAsChannelMask)
 				{
-					CreateVectorChannelMaskParameterValueWidget(Parameter, ParameterProperty, DetailGroup);
+					CreateVectorChannelMaskParameterValueWidget(Parameter, ParameterValueProperty, DetailGroup);
 				}
 				if (ScalarParam && ScalarParam->AtlasData.bIsUsedAsAtlasPosition)
 				{
-					CreateScalarAtlasPositionParameterValueWidget(Parameter, ParameterProperty, DetailGroup);
+					CreateScalarAtlasPositionParameterValueWidget(Parameter, ParameterValueProperty, DetailGroup);
 				}
 				if (TextureParam && 
 					( !TextureParam->ChannelNames.R.IsEmpty()
@@ -584,20 +587,19 @@ void FMaterialInstanceParameterDetails::CreateSingleGroupWidget(FEditorParameter
 					|| !TextureParam->ChannelNames.B.IsEmpty()
 					|| !TextureParam->ChannelNames.A.IsEmpty()))
 				{
-					CreateLabeledTextureParameterValueWidget(Parameter, ParameterProperty, DetailGroup);
+					CreateLabeledTextureParameterValueWidget(Parameter, ParameterValueProperty, DetailGroup);
 				}
 				else if (LayersParam)
 				{
 				}
 				else if (CompMaskParam)
 				{
-					CreateMaskParameterValueWidget(Parameter, ParameterProperty, DetailGroup);
+					CreateMaskParameterValueWidget(Parameter, ParameterValueProperty, DetailGroup);
 				}
 				else
 				{
 					if (ScalarParam && ScalarParam->SliderMax > ScalarParam->SliderMin)
 					{
-						TSharedPtr<IPropertyHandle> ParameterValueProperty = ParameterProperty->GetChildHandle("ParameterValue");
 						ParameterValueProperty->SetInstanceMetaData("UIMin", FString::Printf(TEXT("%f"), ScalarParam->SliderMin));
 						ParameterValueProperty->SetInstanceMetaData("UIMax", FString::Printf(TEXT("%f"), ScalarParam->SliderMax));
 					}
@@ -626,7 +628,7 @@ void FMaterialInstanceParameterDetails::CreateSingleGroupWidget(FEditorParameter
 						}
 					}
 
-					CreateParameterValueWidget(Parameter, ParameterProperty, DetailGroup);
+					CreateParameterValueWidget(Parameter, ParameterValueProperty, DetailGroup);
 				}
 			}
 		}
@@ -642,14 +644,8 @@ void FMaterialInstanceParameterDetails::CreateSingleGroupWidget(FEditorParameter
 			UDEditorParameterValue* FirstParameter = ParameterGroup.Parameters[ParameterIndices[0]];
 			TSharedPtr<IPropertyHandle> FirstParameterProperty = ParametersArrayProperty->GetChildHandle(ParameterIndices[0]);
 
-			// Create a parameter value widget for the first parameter.
-			CreateParameterValueWidget(FirstParameter, FirstParameterProperty, DetailGroup);
-
 			if (TSharedPtr<IPropertyHandle> FirstParameterValueProperty = FirstParameterProperty->GetChildHandle("ParameterValue"))
 			{
-				// Display the parameter collection's name on the widget.
-				FirstParameterValueProperty->SetPropertyDisplayName(FText::FromName(ParameterCollection->GetFName()));
-
 				// Gather the relevant property handles of the collection's referencing parameters.
 				TArray<TSharedRef<IPropertyHandle>> PropertyHandles{ FirstParameterValueProperty.ToSharedRef() };
 				for (int32 ParameterIndex = 1; ParameterIndex < NumParameterIndices; ++ParameterIndex)
@@ -663,526 +659,517 @@ void FMaterialInstanceParameterDetails::CreateSingleGroupWidget(FEditorParameter
 					}
 				}
 
-				// Create a delegate to invoke when the property value changes.
-				TDelegate<void(const FPropertyChangedEvent&)> OnPropertyValueChangedDelegate = TDelegate<void(const FPropertyChangedEvent&)>::CreateWeakLambda(ParameterCollection, [ParameterCollection, PropertyHandles = MoveTemp(PropertyHandles)](const FPropertyChangedEvent& PropertyChangedEvent)
-					{
-						// Determine if this is the final setting of the property value.
-						if (PropertyChangedEvent.ChangeType != EPropertyChangeType::ValueSet)
-						{
-							return;
-						}
-
-						if (PropertyHandles.IsEmpty())
-						{
-							return;
-						}
-
-						// Get the new value from the first property handle.
-						UObject* Value = nullptr;
-						if (PropertyHandles[0]->GetValue(Value) != FPropertyAccess::Success)
-						{
-							return;
-						}
-
-						// Verify that the new value is valid.
-						if (auto* Collection = Cast<UMaterialParameterCollection>(Value))
-						{
-							if (Collection->IsInstanceOf(ParameterCollection->StateId))
-							{
-								// Set the same value on all of the collection's referencing parameters.
-								for (int32 PropertyHandleIndex = 1; PropertyHandleIndex < PropertyHandles.Num(); ++PropertyHandleIndex)
-								{
-									PropertyHandles[PropertyHandleIndex]->SetValue(Value);
-								}
-
-								return;
-							}
-						}
-
-						// Reset the value to the default parameter collection.
-						Value = ParameterCollection;
-						PropertyHandles[0]->SetValue(Value);
-					});
-
-				// Invoke the delegate when the property value or one of its child property values changes.
-				FirstParameterValueProperty->SetOnPropertyValueChangedWithData(OnPropertyValueChangedDelegate);
-				FirstParameterValueProperty->SetOnChildPropertyValueChangedWithData(OnPropertyValueChangedDelegate);
+				// Create a parameter value widget for the first parameter.
+				CreateParameterCollectionParameterValueWidget(FirstParameter, FirstParameterValueProperty, DetailGroup, ParameterCollection, MoveTemp(PropertyHandles));
 			}
 		}
 	}
 }
 
-void FMaterialInstanceParameterDetails::CreateParameterValueWidget(UDEditorParameterValue* Parameter, TSharedPtr<IPropertyHandle> ParameterProperty, IDetailGroup& DetailGroup)
+IDetailPropertyRow& FMaterialInstanceParameterDetails::CreateParameterPropertyRow(UDEditorParameterValue* Parameter, TSharedPtr<IPropertyHandle> ParameterValueProperty, IDetailGroup& DetailGroup, TOptional<FResetToDefaultOverride> ResetToDefaultOverride)
 {
-	TSharedPtr<IPropertyHandle> ParameterValueProperty = ParameterProperty->GetChildHandle("ParameterValue");
+	IDetailPropertyRow& PropertyRow = DetailGroup.AddPropertyRow(ParameterValueProperty.ToSharedRef());
 
-	if (ParameterValueProperty->IsValidHandle())
+	FResetToDefaultOverride ResetOverride;
+	if (ResetToDefaultOverride.IsSet())
+	{
+		ResetOverride = ResetToDefaultOverride.GetValue();
+	}
+	else
 	{
-		TAttribute<bool> IsParamEnabled = TAttribute<bool>::Create(TAttribute<bool>::FGetter::CreateStatic(&FMaterialPropertyHelpers::IsOverriddenExpression, Parameter));
-
-		IDetailPropertyRow& PropertyRow = DetailGroup.AddPropertyRow(ParameterValueProperty.ToSharedRef());
-
 		TAttribute<bool> IsResetVisible = TAttribute<bool>::Create(TAttribute<bool>::FGetter::CreateStatic(&FMaterialPropertyHelpers::ShouldShowResetToDefault, Parameter, MaterialEditorInstance));
 		FSimpleDelegate ResetHandler = FSimpleDelegate::CreateStatic(&FMaterialPropertyHelpers::ResetToDefault, Parameter, MaterialEditorInstance);
-		FResetToDefaultOverride ResetOverride = FResetToDefaultOverride::Create(IsResetVisible, ResetHandler);
-
-		PropertyRow
-			.DisplayName(FText::FromName(Parameter->ParameterInfo.Name))
-			.ToolTip(FMaterialPropertyHelpers::GetParameterTooltip(Parameter, MaterialEditorInstance))
-			.EditCondition(IsParamEnabled, FOnBooleanValueChanged::CreateStatic(&FMaterialPropertyHelpers::OnOverrideParameter, Parameter, MaterialEditorInstance))
-			.Visibility(TAttribute<EVisibility>::Create(TAttribute<EVisibility>::FGetter::CreateStatic(&FMaterialPropertyHelpers::ShouldShowExpression, Parameter, MaterialEditorInstance, ShowHiddenDelegate)))
-			.OverrideResetToDefault(ResetOverride);
-
-		// Textures need a special widget that filters based on VT or not
-		UDEditorTextureParameterValue* TextureParam = Cast<UDEditorTextureParameterValue>(Parameter);
-		if (TextureParam != nullptr)
-		{
-			UMaterial *Material = MaterialEditorInstance->SourceInstance->GetMaterial();
-			if (Material != nullptr)
-			{
-				UMaterialExpressionTextureSampleParameter* Expression = Material->FindExpressionByGUID<UMaterialExpressionTextureSampleParameter>(TextureParam->ExpressionId);
-				if (Expression != nullptr)
-				{
-					TWeakObjectPtr<UMaterialExpressionTextureSampleParameter> SamplerExpression = Expression;
+		ResetOverride = FResetToDefaultOverride::Create(IsResetVisible, ResetHandler);
+	}
 
-					PropertyRow.CustomWidget()
-					.NameContent()
-					[
-						ParameterValueProperty->CreatePropertyNameWidget()
-					]
-					.ValueContent()
-					.MaxDesiredWidth(TOptional<float>())
-					[
-						SNew(SObjectPropertyEntryBox)
-						.PropertyHandle(ParameterValueProperty)
-						.AllowedClass(UTexture::StaticClass())
-						.ThumbnailPool(PropertyUtilities.Pin()->GetThumbnailPool())
-						.OnShouldFilterAsset_Lambda([SamplerExpression](const FAssetData& AssetData)
-						{
-							if (SamplerExpression.Get())
-							{
-								bool VirtualTextured = false;
-								AssetData.GetTagValue<bool>("VirtualTextureStreaming", VirtualTextured);
+	TAttribute<bool> IsParamEnabled = TAttribute<bool>::Create(TAttribute<bool>::FGetter::CreateStatic(&FMaterialPropertyHelpers::IsOverriddenExpression, Parameter));
 
-								bool ExpressionIsVirtualTextured = IsVirtualSamplerType(SamplerExpression->SamplerType);
+	return PropertyRow
+		.DisplayName(FText::FromName(Parameter->ParameterInfo.Name))
+		.ToolTip(FMaterialPropertyHelpers::GetParameterTooltip(Parameter, MaterialEditorInstance))
+		.EditCondition(IsParamEnabled, FOnBooleanValueChanged::CreateStatic(&FMaterialPropertyHelpers::OnOverrideParameter, Parameter, MaterialEditorInstance))
+		.Visibility(TAttribute<EVisibility>::Create(TAttribute<EVisibility>::FGetter::CreateStatic(&FMaterialPropertyHelpers::ShouldShowExpression, Parameter, MaterialEditorInstance, ShowHiddenDelegate)))
+		.OverrideResetToDefault(ResetOverride);
+}
 
-								return VirtualTextured != ExpressionIsVirtualTextured;
-							}
-							else
-							{
-								return false;
-							}
-						})
-					];
-				}
-			}
-		}
+void FMaterialInstanceParameterDetails::CreateParameterValueWidget(UDEditorParameterValue* Parameter, TSharedPtr<IPropertyHandle> ParameterValueProperty, IDetailGroup& DetailGroup)
+{
+	IDetailPropertyRow& PropertyRow = CreateParameterPropertyRow(Parameter, ParameterValueProperty, DetailGroup);
 
-		// Scalar parameters can present an enumeration combo box.
-		if (UDEditorScalarParameterValue* ScalarParam = Cast<UDEditorScalarParameterValue>(Parameter))
+	// Textures need a special widget that filters based on VT or not
+	UDEditorTextureParameterValue* TextureParam = Cast<UDEditorTextureParameterValue>(Parameter);
+	if (TextureParam != nullptr)
+	{
+		UMaterial *Material = MaterialEditorInstance->SourceInstance->GetMaterial();
+		if (Material != nullptr)
 		{
-			TSoftObjectPtr<UObject> Enumeration;
-			if (!ScalarParam->Enumeration.IsNull())
+			UMaterialExpressionTextureSampleParameter* Expression = Material->FindExpressionByGUID<UMaterialExpressionTextureSampleParameter>(TextureParam->ExpressionId);
+			if (Expression != nullptr)
 			{
-				Enumeration = ScalarParam->Enumeration;
-			}
-			else if (ScalarParam->EnumerationIndex >= 0)
-			{
-				Enumeration = MaterialEditorInstance->GetEnumerationObject(ScalarParam->EnumerationIndex);
-			}
+				TWeakObjectPtr<UMaterialExpressionTextureSampleParameter> SamplerExpression = Expression;
 
-			if (!Enumeration.IsNull())
-			{
 				PropertyRow.CustomWidget()
 				.NameContent()
 				[
 					ParameterValueProperty->CreatePropertyNameWidget()
 				]
 				.ValueContent()
+				.MaxDesiredWidth(TOptional<float>())
 				[
-					FMaterialPropertyHelpers::CreateScalarEnumerationParameterValueWidget(Enumeration, ParameterValueProperty)
+					SNew(SObjectPropertyEntryBox)
+					.PropertyHandle(ParameterValueProperty)
+					.AllowedClass(UTexture::StaticClass())
+					.ThumbnailPool(PropertyUtilities.Pin()->GetThumbnailPool())
+					.OnShouldFilterAsset_Lambda([SamplerExpression](const FAssetData& AssetData)
+					{
+						if (SamplerExpression.Get())
+						{
+							bool VirtualTextured = false;
+							AssetData.GetTagValue<bool>("VirtualTextureStreaming", VirtualTextured);
+
+							bool ExpressionIsVirtualTextured = IsVirtualSamplerType(SamplerExpression->SamplerType);
+
+							return VirtualTextured != ExpressionIsVirtualTextured;
+						}
+						else
+						{
+							return false;
+						}
+					})
 				];
 			}
 		}
 	}
+
+	// Scalar parameters can present an enumeration combo box.
+	if (UDEditorScalarParameterValue* ScalarParam = Cast<UDEditorScalarParameterValue>(Parameter))
+	{
+		TSoftObjectPtr<UObject> Enumeration;
+		if (!ScalarParam->Enumeration.IsNull())
+		{
+			Enumeration = ScalarParam->Enumeration;
+		}
+		else if (ScalarParam->EnumerationIndex >= 0)
+		{
+			Enumeration = MaterialEditorInstance->GetEnumerationObject(ScalarParam->EnumerationIndex);
+		}
+
+		if (!Enumeration.IsNull())
+		{
+			PropertyRow.CustomWidget()
+			.NameContent()
+			[
+				ParameterValueProperty->CreatePropertyNameWidget()
+			]
+			.ValueContent()
+			[
+				FMaterialPropertyHelpers::CreateScalarEnumerationParameterValueWidget(Enumeration, ParameterValueProperty)
+			];
+		}
+	}
 }
 
-void FMaterialInstanceParameterDetails::CreateMaskParameterValueWidget(UDEditorParameterValue* Parameter, TSharedPtr<IPropertyHandle> ParameterProperty, IDetailGroup& DetailGroup )
+void FMaterialInstanceParameterDetails::CreateMaskParameterValueWidget(UDEditorParameterValue* Parameter, TSharedPtr<IPropertyHandle> ParameterValueProperty, IDetailGroup& DetailGroup )
 {
-	TSharedPtr<IPropertyHandle> ParameterValueProperty = ParameterProperty->GetChildHandle("ParameterValue");
 	TSharedPtr<IPropertyHandle> RMaskProperty = ParameterValueProperty->GetChildHandle("R");
 	TSharedPtr<IPropertyHandle> GMaskProperty = ParameterValueProperty->GetChildHandle("G");
 	TSharedPtr<IPropertyHandle> BMaskProperty = ParameterValueProperty->GetChildHandle("B");
 	TSharedPtr<IPropertyHandle> AMaskProperty = ParameterValueProperty->GetChildHandle("A");
 
-	if (ParameterValueProperty->IsValidHandle())
-	{
-		TAttribute<bool> IsParamEnabled = TAttribute<bool>::Create(TAttribute<bool>::FGetter::CreateStatic(&FMaterialPropertyHelpers::IsOverriddenExpression, Parameter));
-
-		IDetailPropertyRow& PropertyRow = DetailGroup.AddPropertyRow(ParameterValueProperty.ToSharedRef());
-		PropertyRow.EditCondition(IsParamEnabled, FOnBooleanValueChanged::CreateStatic(&FMaterialPropertyHelpers::OnOverrideParameter, Parameter, MaterialEditorInstance));
-		// Handle reset to default manually
-		PropertyRow.OverrideResetToDefault(FResetToDefaultOverride::Create(FSimpleDelegate::CreateStatic(&FMaterialPropertyHelpers::ResetToDefault, Parameter, MaterialEditorInstance)));
-		PropertyRow.Visibility(TAttribute<EVisibility>::Create(TAttribute<EVisibility>::FGetter::CreateStatic(&FMaterialPropertyHelpers::ShouldShowExpression, Parameter, MaterialEditorInstance, ShowHiddenDelegate)));
+	IDetailPropertyRow& PropertyRow = CreateParameterPropertyRow(Parameter, ParameterValueProperty, DetailGroup);
 
-		const FText ParameterName = FText::FromName(Parameter->ParameterInfo.Name);
+	PropertyRow.CustomWidget()
+	.NameContent()
+	[
+		ParameterValueProperty->CreatePropertyNameWidget()
+	]
+	.ValueContent()
+		.MaxDesiredWidth(200.0f)
+		[
+			SNew(SHorizontalBox)
+			+ SHorizontalBox::Slot()
+		.FillWidth(1.0f)
+		[
+			SNew(SHorizontalBox)
+			+ SHorizontalBox::Slot()
+		.HAlign(HAlign_Left)
+		.AutoWidth()
+		[
+			RMaskProperty->CreatePropertyNameWidget()
+		]
+	+ SHorizontalBox::Slot()
+		.HAlign(HAlign_Left)
+		.AutoWidth()
+		[
+			RMaskProperty->CreatePropertyValueWidget()
+		]
+	+ SHorizontalBox::Slot()
+		.HAlign(HAlign_Left)
+		.Padding(FMargin(10.0f, 0.0f, 0.0f, 0.0f))
+		.AutoWidth()
+		[
+			GMaskProperty->CreatePropertyNameWidget()
+		]
+	+ SHorizontalBox::Slot()
+		.HAlign(HAlign_Left)
+		.AutoWidth()
+		[
+			GMaskProperty->CreatePropertyValueWidget()
+		]
+	+ SHorizontalBox::Slot()
+		.HAlign(HAlign_Left)
+		.Padding(FMargin(10.0f, 0.0f, 0.0f, 0.0f))
+		.AutoWidth()
+		[
+			BMaskProperty->CreatePropertyNameWidget()
+		]
+	+ SHorizontalBox::Slot()
+		.HAlign(HAlign_Left)
+		.AutoWidth()
+		[
+			BMaskProperty->CreatePropertyValueWidget()
+		]
+	+ SHorizontalBox::Slot()
+		.HAlign(HAlign_Left)
+		.Padding(FMargin(10.0f, 0.0f, 0.0f, 0.0f))
+		.AutoWidth()
+		[
+			AMaskProperty->CreatePropertyNameWidget()
+		]
+	+ SHorizontalBox::Slot()
+		.HAlign(HAlign_Left)
+		.AutoWidth()
+		[
+			AMaskProperty->CreatePropertyValueWidget()
+		]
+		]
+		];
+}
 
-		FDetailWidgetRow& CustomWidget = PropertyRow.CustomWidget();
-		CustomWidget
-			.FilterString(ParameterName)
-			.NameContent()
-			[
-				SNew(STextBlock)
-				.Text(ParameterName)
-				.ToolTipText(FMaterialPropertyHelpers::GetParameterTooltip(Parameter, MaterialEditorInstance))
-				.Font(FAppStyle::GetFontStyle(TEXT("PropertyWindow.NormalFont")))
-			]
-		.ValueContent()
-			.MaxDesiredWidth(200.0f)
-			[
-				SNew(SHorizontalBox)
-				+ SHorizontalBox::Slot()
-			.FillWidth(1.0f)
-			[
-				SNew(SHorizontalBox)
-				+ SHorizontalBox::Slot()
-			.HAlign(HAlign_Left)
-			.AutoWidth()
-			[
-				RMaskProperty->CreatePropertyNameWidget()
-			]
-		+ SHorizontalBox::Slot()
-			.HAlign(HAlign_Left)
-			.AutoWidth()
-			[
-				RMaskProperty->CreatePropertyValueWidget()
-			]
-		+ SHorizontalBox::Slot()
-			.HAlign(HAlign_Left)
-			.Padding(FMargin(10.0f, 0.0f, 0.0f, 0.0f))
-			.AutoWidth()
-			[
-				GMaskProperty->CreatePropertyNameWidget()
-			]
-		+ SHorizontalBox::Slot()
-			.HAlign(HAlign_Left)
-			.AutoWidth()
-			[
-				GMaskProperty->CreatePropertyValueWidget()
-			]
-		+ SHorizontalBox::Slot()
-			.HAlign(HAlign_Left)
-			.Padding(FMargin(10.0f, 0.0f, 0.0f, 0.0f))
-			.AutoWidth()
-			[
-				BMaskProperty->CreatePropertyNameWidget()
-			]
-		+ SHorizontalBox::Slot()
-			.HAlign(HAlign_Left)
-			.AutoWidth()
-			[
-				BMaskProperty->CreatePropertyValueWidget()
-			]
-		+ SHorizontalBox::Slot()
-			.HAlign(HAlign_Left)
-			.Padding(FMargin(10.0f, 0.0f, 0.0f, 0.0f))
-			.AutoWidth()
-			[
-				AMaskProperty->CreatePropertyNameWidget()
-			]
+void FMaterialInstanceParameterDetails::CreateVectorChannelMaskParameterValueWidget(UDEditorParameterValue* Parameter, TSharedPtr<IPropertyHandle> ParameterValueProperty, IDetailGroup& DetailGroup)
+{
+	IDetailPropertyRow& PropertyRow = CreateParameterPropertyRow(Parameter, ParameterValueProperty, DetailGroup);
+
+	// Combo box hooks for converting between our "enum" and colors
+	FOnGetPropertyComboBoxStrings GetMaskStrings = FOnGetPropertyComboBoxStrings::CreateStatic(&FMaterialPropertyHelpers::GetVectorChannelMaskComboBoxStrings);
+	FOnGetPropertyComboBoxValue GetMaskValue = FOnGetPropertyComboBoxValue::CreateStatic(&FMaterialPropertyHelpers::GetVectorChannelMaskValue, Parameter);
+	FOnPropertyComboBoxValueSelected SetMaskValue = FOnPropertyComboBoxValueSelected::CreateStatic(&FMaterialPropertyHelpers::SetVectorChannelMaskValue, ParameterValueProperty, Parameter, (UObject*)MaterialEditorInstance);
+
+	// Widget replaces color picker with combo box
+	PropertyRow.CustomWidget()
+	.NameContent()
+	[
+		ParameterValueProperty->CreatePropertyNameWidget()
+	]
+	.ValueContent()
+	.MaxDesiredWidth(200.0f)
+	[
+		SNew(SHorizontalBox)
 		+ SHorizontalBox::Slot()
+		.FillWidth(1.0f)
+		[
+			SNew(SHorizontalBox)
+			+ SHorizontalBox::Slot()
 			.HAlign(HAlign_Left)
 			.AutoWidth()
 			[
-				AMaskProperty->CreatePropertyValueWidget()
-			]
+				PropertyCustomizationHelpers::MakePropertyComboBox(ParameterValueProperty, GetMaskStrings, GetMaskValue, SetMaskValue)
 			]
-			];
-	}
+		]
+	];
 }
 
-void FMaterialInstanceParameterDetails::CreateVectorChannelMaskParameterValueWidget(UDEditorParameterValue* Parameter, TSharedPtr<IPropertyHandle> ParameterProperty, IDetailGroup& DetailGroup)
+void FMaterialInstanceParameterDetails::CreateScalarAtlasPositionParameterValueWidget(class UDEditorParameterValue* Parameter, TSharedPtr<IPropertyHandle> ParameterValueProperty, IDetailGroup& DetailGroup)
 {
-	TSharedPtr<IPropertyHandle> ParameterValueProperty = ParameterProperty->GetChildHandle("ParameterValue");
-
-	if (ParameterValueProperty->IsValidHandle())
-	{
-		TAttribute<bool> IsParamEnabled = TAttribute<bool>::Create(TAttribute<bool>::FGetter::CreateStatic(&FMaterialPropertyHelpers::IsOverriddenExpression, Parameter));
+	// Handle reset to default manually
+	TAttribute<bool> IsResetVisible = TAttribute<bool>::Create(TAttribute<bool>::FGetter::CreateStatic(&FMaterialPropertyHelpers::ShouldShowResetToDefault, Parameter, MaterialEditorInstance));
+	FSimpleDelegate ResetHandler = FSimpleDelegate::CreateStatic(&FMaterialPropertyHelpers::ResetCurveToDefault, Parameter, MaterialEditorInstance);
+	FResetToDefaultOverride ResetOverride = FResetToDefaultOverride::Create(IsResetVisible, ResetHandler);
 
-		IDetailPropertyRow& PropertyRow = DetailGroup.AddPropertyRow(ParameterValueProperty.ToSharedRef());
-		PropertyRow.EditCondition(IsParamEnabled, FOnBooleanValueChanged::CreateStatic(&FMaterialPropertyHelpers::OnOverrideParameter, Parameter, MaterialEditorInstance));
-		// Handle reset to default manually
-		PropertyRow.OverrideResetToDefault(FResetToDefaultOverride::Create(FSimpleDelegate::CreateStatic(&FMaterialPropertyHelpers::ResetToDefault, Parameter, MaterialEditorInstance)));
-		PropertyRow.Visibility(TAttribute<EVisibility>::Create(TAttribute<EVisibility>::FGetter::CreateStatic(&FMaterialPropertyHelpers::ShouldShowExpression, Parameter, MaterialEditorInstance, ShowHiddenDelegate)));
+	IDetailPropertyRow& PropertyRow = CreateParameterPropertyRow(Parameter, ParameterValueProperty, DetailGroup, ResetOverride);
 
-		const FText ParameterName = FText::FromName(Parameter->ParameterInfo.Name);
+	const FText ParameterName = FText::FromName(Parameter->ParameterInfo.Name);
+	UDEditorScalarParameterValue* AtlasParameter = Cast<UDEditorScalarParameterValue>(Parameter);
 
-		// Combo box hooks for converting between our "enum" and colors
-		FOnGetPropertyComboBoxStrings GetMaskStrings = FOnGetPropertyComboBoxStrings::CreateStatic(&FMaterialPropertyHelpers::GetVectorChannelMaskComboBoxStrings);
-		FOnGetPropertyComboBoxValue GetMaskValue = FOnGetPropertyComboBoxValue::CreateStatic(&FMaterialPropertyHelpers::GetVectorChannelMaskValue, Parameter);
-		FOnPropertyComboBoxValueSelected SetMaskValue = FOnPropertyComboBoxValueSelected::CreateStatic(&FMaterialPropertyHelpers::SetVectorChannelMaskValue, ParameterValueProperty, Parameter, (UObject*)MaterialEditorInstance);
-
-		// Widget replaces color picker with combo box
-		FDetailWidgetRow& CustomWidget = PropertyRow.CustomWidget();
-		CustomWidget
+	FDetailWidgetRow& CustomWidget = PropertyRow.CustomWidget();
+	CustomWidget
 		.FilterString(ParameterName)
 		.NameContent()
 		[
-			SNew(STextBlock)
-			.Text(ParameterName)
-			.ToolTipText(FMaterialPropertyHelpers::GetParameterTooltip(Parameter, MaterialEditorInstance))
-			.Font(FAppStyle::GetFontStyle(TEXT("PropertyWindow.NormalFont")))
+			ParameterValueProperty->CreatePropertyNameWidget()
 		]
 		.ValueContent()
-		.MaxDesiredWidth(200.0f)
+		.HAlign(HAlign_Fill)
+		.MaxDesiredWidth(400.0f)
 		[
-			SNew(SHorizontalBox)
-			+ SHorizontalBox::Slot()
-			.FillWidth(1.0f)
-			[
-				SNew(SHorizontalBox)
-				+ SHorizontalBox::Slot()
-				.HAlign(HAlign_Left)
-				.AutoWidth()
-				[
-					PropertyCustomizationHelpers::MakePropertyComboBox(ParameterValueProperty, GetMaskStrings, GetMaskValue, SetMaskValue)
-				]
-			]
+			SNew(SObjectPropertyEntryBox)
+			.ObjectPath(this, &FMaterialInstanceParameterDetails::GetCurvePath, AtlasParameter)
+			.AllowedClass(UCurveLinearColor::StaticClass())
+			.NewAssetFactories(TArray<UFactory*>())
+			.DisplayThumbnail(true)
+			.ThumbnailPool(PropertyUtilities.Pin()->GetThumbnailPool())
+			.OnShouldFilterAsset(FOnShouldFilterAsset::CreateStatic(&FMaterialPropertyHelpers::OnShouldFilterCurveAsset, AtlasParameter->AtlasData.Atlas))
+			.OnShouldSetAsset(FOnShouldSetAsset::CreateStatic(&FMaterialPropertyHelpers::OnShouldSetCurveAsset, AtlasParameter->AtlasData.Atlas))
+			.OnObjectChanged(FOnSetObject::CreateStatic(&FMaterialPropertyHelpers::SetPositionFromCurveAsset, AtlasParameter->AtlasData.Atlas, AtlasParameter, ParameterValueProperty, (UObject*)MaterialEditorInstance))
+			.DisplayCompactSize(true)
 		];
-	}
-}
-
-void FMaterialInstanceParameterDetails::CreateScalarAtlasPositionParameterValueWidget(class UDEditorParameterValue* Parameter, TSharedPtr<IPropertyHandle> ParameterProperty, IDetailGroup& DetailGroup)
-{
-	TSharedPtr<IPropertyHandle> ParameterValueProperty = ParameterProperty->GetChildHandle("ParameterValue");
-
-	if (ParameterValueProperty->IsValidHandle())
-	{
-		TAttribute<bool> IsParamEnabled = TAttribute<bool>::Create(TAttribute<bool>::FGetter::CreateStatic(&FMaterialPropertyHelpers::IsOverriddenExpression, Parameter));
-
-		IDetailPropertyRow& PropertyRow = DetailGroup.AddPropertyRow(ParameterValueProperty.ToSharedRef());
-		PropertyRow.EditCondition(IsParamEnabled, FOnBooleanValueChanged::CreateStatic(&FMaterialPropertyHelpers::OnOverrideParameter, Parameter, MaterialEditorInstance));
-		// Handle reset to default manually
-		PropertyRow.Visibility(TAttribute<EVisibility>::Create(TAttribute<EVisibility>::FGetter::CreateStatic(&FMaterialPropertyHelpers::ShouldShowExpression, Parameter, MaterialEditorInstance, ShowHiddenDelegate)));
-
-		const FText ParameterName = FText::FromName(Parameter->ParameterInfo.Name);
-		UDEditorScalarParameterValue* AtlasParameter = Cast<UDEditorScalarParameterValue>(Parameter);
-
-		TAttribute<bool> IsResetVisible = TAttribute<bool>::Create(TAttribute<bool>::FGetter::CreateStatic(&FMaterialPropertyHelpers::ShouldShowResetToDefault, Parameter, MaterialEditorInstance));
-		FSimpleDelegate ResetHandler = FSimpleDelegate::CreateStatic(&FMaterialPropertyHelpers::ResetCurveToDefault, Parameter, MaterialEditorInstance);
-		FResetToDefaultOverride ResetOverride = FResetToDefaultOverride::Create(IsResetVisible, ResetHandler);
-
-		PropertyRow.OverrideResetToDefault(ResetOverride);
-
-		FDetailWidgetRow& CustomWidget = PropertyRow.CustomWidget();
-		CustomWidget
-			.FilterString(ParameterName)
-			.NameContent()
-			[
-				SNew(STextBlock)
-				.Text(ParameterName)
-				.ToolTipText(FMaterialPropertyHelpers::GetParameterTooltip(Parameter, MaterialEditorInstance))
-				.Font(FAppStyle::GetFontStyle(TEXT("PropertyWindow.NormalFont")))
-			]
-			.ValueContent()
-			.HAlign(HAlign_Fill)
-			.MaxDesiredWidth(400.0f)
-			[
-				SNew(SObjectPropertyEntryBox)
-				.ObjectPath(this, &FMaterialInstanceParameterDetails::GetCurvePath, AtlasParameter)
-				.AllowedClass(UCurveLinearColor::StaticClass())
-				.NewAssetFactories(TArray<UFactory*>())
-				.DisplayThumbnail(true)
-				.ThumbnailPool(PropertyUtilities.Pin()->GetThumbnailPool())
-				.OnShouldFilterAsset(FOnShouldFilterAsset::CreateStatic(&FMaterialPropertyHelpers::OnShouldFilterCurveAsset, AtlasParameter->AtlasData.Atlas))
-				.OnShouldSetAsset(FOnShouldSetAsset::CreateStatic(&FMaterialPropertyHelpers::OnShouldSetCurveAsset, AtlasParameter->AtlasData.Atlas))
-				.OnObjectChanged(FOnSetObject::CreateStatic(&FMaterialPropertyHelpers::SetPositionFromCurveAsset, AtlasParameter->AtlasData.Atlas, AtlasParameter, ParameterProperty, (UObject*)MaterialEditorInstance))
-				.DisplayCompactSize(true)
-			];
-	}
 }
 
-void FMaterialInstanceParameterDetails::CreateLabeledTextureParameterValueWidget(class UDEditorParameterValue* Parameter, TSharedPtr<IPropertyHandle> ParameterProperty, IDetailGroup& DetailGroup)
+void FMaterialInstanceParameterDetails::CreateLabeledTextureParameterValueWidget(class UDEditorParameterValue* Parameter, TSharedPtr<IPropertyHandle> ParameterValueProperty, IDetailGroup& DetailGroup)
 {
-	TSharedPtr<IPropertyHandle> ParameterValueProperty = ParameterProperty->GetChildHandle("ParameterValue");
-
-	if (ParameterValueProperty->IsValidHandle())
+	UDEditorTextureParameterValue* TextureParam = Cast<UDEditorTextureParameterValue>(Parameter);
+	if (TextureParam)
 	{
-		UDEditorTextureParameterValue* TextureParam = Cast<UDEditorTextureParameterValue>(Parameter);
-		if (TextureParam)
+		UMaterial *Material = MaterialEditorInstance->SourceInstance->GetMaterial();
+		if (Material != nullptr)
 		{
-			UMaterial *Material = MaterialEditorInstance->SourceInstance->GetMaterial();
-			if (Material != nullptr)
+			UMaterialExpressionTextureSampleParameter* Expression = Material->FindExpressionByGUID<UMaterialExpressionTextureSampleParameter>(TextureParam->ExpressionId);
+			if (Expression != nullptr)
 			{
-				UMaterialExpressionTextureSampleParameter* Expression = Material->FindExpressionByGUID<UMaterialExpressionTextureSampleParameter>(TextureParam->ExpressionId);
-				if (Expression != nullptr)
-				{
-					TWeakObjectPtr<UMaterialExpressionTextureSampleParameter> SamplerExpression = Expression;
-					TAttribute<bool> IsParamEnabled = TAttribute<bool>::Create(TAttribute<bool>::FGetter::CreateStatic(&FMaterialPropertyHelpers::IsOverriddenExpression, Parameter));
-
-					IDetailPropertyRow& PropertyRow = DetailGroup.AddPropertyRow(ParameterValueProperty.ToSharedRef());
+				TWeakObjectPtr<UMaterialExpressionTextureSampleParameter> SamplerExpression = Expression;
 
-					TAttribute<bool> IsResetVisible = TAttribute<bool>::Create(TAttribute<bool>::FGetter::CreateStatic(&FMaterialPropertyHelpers::ShouldShowResetToDefault, Parameter, MaterialEditorInstance));
-					FSimpleDelegate ResetHandler = FSimpleDelegate::CreateStatic(&FMaterialPropertyHelpers::ResetToDefault, Parameter, MaterialEditorInstance);
-					FResetToDefaultOverride ResetOverride = FResetToDefaultOverride::Create(IsResetVisible, ResetHandler);
+				IDetailPropertyRow& PropertyRow = CreateParameterPropertyRow(Parameter, ParameterValueProperty, DetailGroup);
 
-					PropertyRow
-						.DisplayName(FText::FromName(Parameter->ParameterInfo.Name))
-						.EditCondition(IsParamEnabled, FOnBooleanValueChanged::CreateStatic(&FMaterialPropertyHelpers::OnOverrideParameter, Parameter, MaterialEditorInstance))
-						.ToolTip(FMaterialPropertyHelpers::GetParameterTooltip(Parameter, MaterialEditorInstance))
-						.Visibility(TAttribute<EVisibility>::Create(TAttribute<EVisibility>::FGetter::CreateStatic(&FMaterialPropertyHelpers::ShouldShowExpression, Parameter, MaterialEditorInstance, ShowHiddenDelegate)));
-
-
-					TSharedPtr<SWidget> NameWidget;
-					TSharedPtr<SWidget> ValueWidget;
-					FDetailWidgetRow Row;
-					PropertyRow.GetDefaultWidgets(NameWidget, ValueWidget, Row);
+				TSharedPtr<SWidget> NameWidget;
+				TSharedPtr<SWidget> ValueWidget;
+				FDetailWidgetRow Row;
+				PropertyRow.GetDefaultWidgets(NameWidget, ValueWidget, Row);
 
-					FDetailWidgetRow &DetailWidgetRow = PropertyRow.CustomWidget();
-					TSharedPtr<SVerticalBox> NameVerticalBox;
-					DetailWidgetRow.NameContent()
-						[
-							SAssignNew(NameVerticalBox, SVerticalBox)
-							+ SVerticalBox::Slot()
+				FDetailWidgetRow &DetailWidgetRow = PropertyRow.CustomWidget();
+				TSharedPtr<SVerticalBox> NameVerticalBox;
+				DetailWidgetRow.NameContent()
+					[
+						SAssignNew(NameVerticalBox, SVerticalBox)
+						+ SVerticalBox::Slot()
 						.AutoHeight()
 						[
-							SNew(STextBlock)
-							.Text(FText::FromName(Parameter->ParameterInfo.Name))
-						.ToolTipText(FMaterialPropertyHelpers::GetParameterTooltip(Parameter, MaterialEditorInstance))
-						.Font(FAppStyle::GetFontStyle(TEXT("PropertyWindow.NormalFont")))
+							ParameterValueProperty->CreatePropertyNameWidget()
 						]
-						];
-					DetailWidgetRow.ValueContent()
-						.MinDesiredWidth(Row.ValueWidget.MinWidth)
-						.MaxDesiredWidth(Row.ValueWidget.MaxWidth)
-						[
-							SNew(SObjectPropertyEntryBox)
-							.PropertyHandle(ParameterValueProperty)
-							.AllowedClass(UTexture::StaticClass())
-							.ThumbnailPool(PropertyUtilities.Pin()->GetThumbnailPool())
-							.OnShouldFilterAsset_Lambda([SamplerExpression](const FAssetData& AssetData)
+					];
+				DetailWidgetRow.ValueContent()
+					.MinDesiredWidth(Row.ValueWidget.MinWidth)
+					.MaxDesiredWidth(Row.ValueWidget.MaxWidth)
+					[
+						SNew(SObjectPropertyEntryBox)
+						.PropertyHandle(ParameterValueProperty)
+						.AllowedClass(UTexture::StaticClass())
+						.ThumbnailPool(PropertyUtilities.Pin()->GetThumbnailPool())
+						.OnShouldFilterAsset_Lambda([SamplerExpression](const FAssetData& AssetData)
+						{
+							if (SamplerExpression.Get())
 							{
-								if (SamplerExpression.Get())
-								{
-									bool VirtualTextured = false;
-									AssetData.GetTagValue<bool>("VirtualTextureStreaming", VirtualTextured);
-
-									bool ExpressionIsVirtualTextured = IsVirtualSamplerType(SamplerExpression->SamplerType);
-
-									return VirtualTextured != ExpressionIsVirtualTextured;
-								}
-								else
-								{
-									return false;
-								}
-							})
-						];
+								bool VirtualTextured = false;
+								AssetData.GetTagValue<bool>("VirtualTextureStreaming", VirtualTextured);
 
-					DetailWidgetRow.OverrideResetToDefault(ResetOverride);
+								bool ExpressionIsVirtualTextured = IsVirtualSamplerType(SamplerExpression->SamplerType);
 
-					static const FName Red("R");
-					static const FName Green("G");
-					static const FName Blue("B");
-					static const FName Alpha("A");
+								return VirtualTextured != ExpressionIsVirtualTextured;
+							}
+							else
+							{
+								return false;
+							}
+						})
+					];
 
-					if (!TextureParam->ChannelNames.R.IsEmpty())
-					{
-						NameVerticalBox->AddSlot()
-							[
-								SNew(SHorizontalBox)
-								+ SHorizontalBox::Slot()
-								.AutoWidth()
-								.Padding(20.0, 2.0, 4.0, 2.0)
-								[
-									SNew(STextBlock)
-									.Text(FText::FromName(Red))
-									.Font(FAppStyle::GetFontStyle(TEXT("PropertyWindow.BoldFont")))
-								]
-								+ SHorizontalBox::Slot()
-								.HAlign(HAlign_Left)
-								.Padding(4.0, 2.0)
-								[
-									SNew(STextBlock)
-									.Text(TextureParam->ChannelNames.R)
-									.Font(FAppStyle::GetFontStyle(TEXT("PropertyWindow.NormalFont")))
-								]
-							];
-					}
-					if (!TextureParam->ChannelNames.G.IsEmpty())
-					{
-						NameVerticalBox->AddSlot()
-							[
-								SNew(SHorizontalBox)
-								+ SHorizontalBox::Slot()
-							.Padding(20.0, 2.0, 4.0, 2.0)
-							.AutoWidth()
-							[
-								SNew(STextBlock)
-								.Text(FText::FromName(Green))
-							.Font(FAppStyle::GetFontStyle(TEXT("PropertyWindow.BoldFont")))
-							]
-						+ SHorizontalBox::Slot()
-							.HAlign(HAlign_Left)
-							.Padding(4.0, 2.0)
-							[
-								SNew(STextBlock)
-								.Text(TextureParam->ChannelNames.G)
-							.Font(FAppStyle::GetFontStyle(TEXT("PropertyWindow.NormalFont")))
-							]
-							];
-					}
-					if (!TextureParam->ChannelNames.B.IsEmpty())
-					{
-						NameVerticalBox->AddSlot()
-							[
-								SNew(SHorizontalBox)
-								+ SHorizontalBox::Slot()
-							.Padding(20.0, 2.0, 4.0, 2.0)
+				static const FName Red("R");
+				static const FName Green("G");
+				static const FName Blue("B");
+				static const FName Alpha("A");
+
+				if (!TextureParam->ChannelNames.R.IsEmpty())
+				{
+					NameVerticalBox->AddSlot()
+						[
+							SNew(SHorizontalBox)
+							+ SHorizontalBox::Slot()
 							.AutoWidth()
-							[
-								SNew(STextBlock)
-								.Text(FText::FromName(Blue))
-							.Font(FAppStyle::GetFontStyle(TEXT("PropertyWindow.BoldFont")))
-							]
-						+ SHorizontalBox::Slot()
-							.HAlign(HAlign_Left)
-							.Padding(4.0, 2.0)
-							[
-								SNew(STextBlock)
-								.Text(TextureParam->ChannelNames.B)
-							.Font(FAppStyle::GetFontStyle(TEXT("PropertyWindow.NormalFont")))
-							]
-							];
-					}
-					if (!TextureParam->ChannelNames.A.IsEmpty())
-					{
-						NameVerticalBox->AddSlot()
-							[
-								SNew(SHorizontalBox)
-								+ SHorizontalBox::Slot()
 							.Padding(20.0, 2.0, 4.0, 2.0)
-							.AutoWidth()
 							[
 								SNew(STextBlock)
-								.Text(FText::FromName(Alpha))
-							.Font(FAppStyle::GetFontStyle(TEXT("PropertyWindow.BoldFont")))
+								.Text(FText::FromName(Red))
+								.Font(FAppStyle::GetFontStyle(TEXT("PropertyWindow.BoldFont")))
 							]
-						+ SHorizontalBox::Slot()
+							+ SHorizontalBox::Slot()
 							.HAlign(HAlign_Left)
 							.Padding(4.0, 2.0)
 							[
 								SNew(STextBlock)
-								.Text(TextureParam->ChannelNames.A)
-							.Font(FAppStyle::GetFontStyle(TEXT("PropertyWindow.NormalFont")))
+								.Text(TextureParam->ChannelNames.R)
+								.Font(FAppStyle::GetFontStyle(TEXT("PropertyWindow.NormalFont")))
 							]
-							];
-					}
+						];
+				}
+				if (!TextureParam->ChannelNames.G.IsEmpty())
+				{
+					NameVerticalBox->AddSlot()
+						[
+							SNew(SHorizontalBox)
+							+ SHorizontalBox::Slot()
+						.Padding(20.0, 2.0, 4.0, 2.0)
+						.AutoWidth()
+						[
+							SNew(STextBlock)
+							.Text(FText::FromName(Green))
+						.Font(FAppStyle::GetFontStyle(TEXT("PropertyWindow.BoldFont")))
+						]
+					+ SHorizontalBox::Slot()
+						.HAlign(HAlign_Left)
+						.Padding(4.0, 2.0)
+						[
+							SNew(STextBlock)
+							.Text(TextureParam->ChannelNames.G)
+						.Font(FAppStyle::GetFontStyle(TEXT("PropertyWindow.NormalFont")))
+						]
+						];
+				}
+				if (!TextureParam->ChannelNames.B.IsEmpty())
+				{
+					NameVerticalBox->AddSlot()
+						[
+							SNew(SHorizontalBox)
+							+ SHorizontalBox::Slot()
+						.Padding(20.0, 2.0, 4.0, 2.0)
+						.AutoWidth()
+						[
+							SNew(STextBlock)
+							.Text(FText::FromName(Blue))
+						.Font(FAppStyle::GetFontStyle(TEXT("PropertyWindow.BoldFont")))
+						]
+					+ SHorizontalBox::Slot()
+						.HAlign(HAlign_Left)
+						.Padding(4.0, 2.0)
+						[
+							SNew(STextBlock)
+							.Text(TextureParam->ChannelNames.B)
+						.Font(FAppStyle::GetFontStyle(TEXT("PropertyWindow.NormalFont")))
+						]
+						];
+				}
+				if (!TextureParam->ChannelNames.A.IsEmpty())
+				{
+					NameVerticalBox->AddSlot()
+						[
+							SNew(SHorizontalBox)
+							+ SHorizontalBox::Slot()
+						.Padding(20.0, 2.0, 4.0, 2.0)
+						.AutoWidth()
+						[
+							SNew(STextBlock)
+							.Text(FText::FromName(Alpha))
+						.Font(FAppStyle::GetFontStyle(TEXT("PropertyWindow.BoldFont")))
+						]
+					+ SHorizontalBox::Slot()
+						.HAlign(HAlign_Left)
+						.Padding(4.0, 2.0)
+						[
+							SNew(STextBlock)
+							.Text(TextureParam->ChannelNames.A)
+						.Font(FAppStyle::GetFontStyle(TEXT("PropertyWindow.NormalFont")))
+						]
+						];
 				}
 			}
 		}
 	}
 }
 
+void FMaterialInstanceParameterDetails::CreateParameterCollectionParameterValueWidget(class UDEditorParameterValue* Parameter, TSharedPtr<IPropertyHandle> ParameterValueProperty, IDetailGroup& DetailGroup, class UMaterialParameterCollection* BaseParameterCollection, TArray<TSharedRef<IPropertyHandle>>&& AllPropertyHandles)
+{
+	IDetailPropertyRow& PropertyRow = CreateParameterPropertyRow(Parameter, ParameterValueProperty, DetailGroup);
+
+	// Display the parameter collection's name on the widget.
+	ParameterValueProperty->SetPropertyDisplayName(FText::FromName(BaseParameterCollection->GetFName()));
+
+	// Create a delegate to invoke when the property value changes.
+	TDelegate<void(const FPropertyChangedEvent&)> OnPropertyValueChangedDelegate = TDelegate<void(const FPropertyChangedEvent&)>::CreateWeakLambda(BaseParameterCollection,
+		[BaseParameterCollection, PropertyHandles = MoveTemp(AllPropertyHandles)](const FPropertyChangedEvent& PropertyChangedEvent)
+		{
+			// Determine if this is the final setting of the property value.
+			if (PropertyChangedEvent.ChangeType != EPropertyChangeType::ValueSet)
+			{
+				return;
+			}
+
+			if (PropertyHandles.IsEmpty())
+			{
+				return;
+			}
+
+			// Get the new value from the first property handle.
+			UObject* Value = nullptr;
+			if (PropertyHandles[0]->GetValue(Value) != FPropertyAccess::Success)
+			{
+				return;
+			}
+
+			// Verify that the new value is valid.
+			if (auto* Collection = Cast<UMaterialParameterCollection>(Value))
+			{
+				if (Collection->IsInstanceOf(BaseParameterCollection->StateId))
+				{
+					// Set the same value on all of the collection's referencing parameters.
+					for (int32 PropertyHandleIndex = 1; PropertyHandleIndex < PropertyHandles.Num(); ++PropertyHandleIndex)
+					{
+						PropertyHandles[PropertyHandleIndex]->SetValue(Value);
+					}
+
+					return;
+				}
+			}
+
+			// Reset the value to the default parameter collection.
+			Value = BaseParameterCollection;
+			PropertyHandles[0]->SetValue(Value);
+		});
+
+	// Invoke the delegate when the property value or one of its child property values changes.
+	ParameterValueProperty->SetOnPropertyValueChangedWithData(OnPropertyValueChangedDelegate);
+	ParameterValueProperty->SetOnChildPropertyValueChangedWithData(OnPropertyValueChangedDelegate);
+
+	FSoftObjectPath BaseParameterCollectionPath = BaseParameterCollection;
+
+	PropertyRow.CustomWidget()
+		.NameContent()
+		[
+			ParameterValueProperty->CreatePropertyNameWidget()
+		]
+		.ValueContent()
+		.MaxDesiredWidth(TOptional<float>())
+		[
+			SNew(SObjectPropertyEntryBox)
+				.PropertyHandle(ParameterValueProperty)
+				.AllowedClass(UMaterialParameterCollection::StaticClass())
+				.ThumbnailPool(PropertyUtilities.Pin()->GetThumbnailPool())
+				.OnShouldFilterAsset_Lambda([BaseParameterCollectionPath](const FAssetData& AssetData)
+					{
+						// Allow the base MPC itself
+						if (AssetData.GetSoftObjectPath() == BaseParameterCollectionPath)
+						{
+							return false;
+						}
+
+						IAssetRegistry& AssetRegistry = IAssetRegistry::GetChecked();
+
+						// Traverse up the hierarchy of UMaterialParameterCollection::Base until we find the base material's MPC
+						FSoftObjectPath ParentAssetPath = FSoftObjectPath(AssetData.GetTagValueRef<FString>(UMaterialParameterCollection::GetBaseParameterCollectionMemberName()));
+						while (!ParentAssetPath.IsNull())
+						{
+							if (ParentAssetPath == BaseParameterCollectionPath)
+							{
+								return false;
+							}
+
+							FAssetData ParentAssetData = AssetRegistry.GetAssetByObjectPath(ParentAssetPath);
+							ParentAssetPath = FSoftObjectPath(ParentAssetData.GetTagValueRef<FString>(UMaterialParameterCollection::GetBaseParameterCollectionMemberName()));
+						}
+
+						// No matching MPC in hierarchy, filter this MPC out
+						return true;
+					})
+		];
+}
+
 FString FMaterialInstanceParameterDetails::GetCurvePath(UDEditorScalarParameterValue* Parameter) const
 {
 	FString Path = Parameter->AtlasData.Curve->GetPathName();
diff --git a/Engine/Source/Editor/MaterialEditor/Private/MaterialEditorInstanceDetailCustomization.h b/Engine/Source/Editor/MaterialEditor/Private/MaterialEditorInstanceDetailCustomization.h
--- a/Engine/Source/Editor/MaterialEditor/Private/MaterialEditorInstanceDetailCustomization.h
+++ b/Engine/Source/Editor/MaterialEditor/Private/MaterialEditorInstanceDetailCustomization.h
@@ -13,6 +13,7 @@
 
 struct FAssetData;
 class IDetailGroup;
+class IDetailPropertyRow;
 class IDetailLayoutBuilder;
 class IPropertyHandle;
 class UMaterialEditorInstanceConstant;
@@ -51,12 +52,16 @@ private:
 	/** Enable/Disable all parameter properties in a group */
 	static void EnableGroupParameters(struct FEditorParameterGroup& ParameterGroup, bool ShouldEnable);
 
+	/** This method generates a PropertyRow with all of the standard settings for any parameter. */
+	IDetailPropertyRow& CreateParameterPropertyRow(UDEditorParameterValue* Parameter, TSharedPtr<IPropertyHandle> ParameterValueProperty, IDetailGroup& DetailGroup, TOptional<FResetToDefaultOverride> ResetToDefaultOverride = TOptional<FResetToDefaultOverride>());
+
 	/** These methods generate the custom widgets for the various parameter types */
-	void CreateParameterValueWidget(class UDEditorParameterValue* Parameter, TSharedPtr<IPropertyHandle> ParameterProperty, IDetailGroup& DetailGroup);
-	void CreateMaskParameterValueWidget(class UDEditorParameterValue* Parameter, TSharedPtr<IPropertyHandle> ParameterProperty, IDetailGroup& DetailGroup);
-	void CreateVectorChannelMaskParameterValueWidget(class UDEditorParameterValue* Parameter, TSharedPtr<IPropertyHandle> ParameterProperty, IDetailGroup& DetailGroup);
-	void CreateScalarAtlasPositionParameterValueWidget(class UDEditorParameterValue* Parameter, TSharedPtr<IPropertyHandle> ParameterProperty, IDetailGroup& DetailGroup);
-	void CreateLabeledTextureParameterValueWidget(class UDEditorParameterValue* Parameter, TSharedPtr<IPropertyHandle> ParameterProperty, IDetailGroup& DetailGroup);
+	void CreateParameterValueWidget(class UDEditorParameterValue* Parameter, TSharedPtr<IPropertyHandle> ParameterValueProperty, IDetailGroup& DetailGroup);
+	void CreateMaskParameterValueWidget(class UDEditorParameterValue* Parameter, TSharedPtr<IPropertyHandle> ParameterValueProperty, IDetailGroup& DetailGroup);
+	void CreateVectorChannelMaskParameterValueWidget(class UDEditorParameterValue* Parameter, TSharedPtr<IPropertyHandle> ParameterValueProperty, IDetailGroup& DetailGroup);
+	void CreateScalarAtlasPositionParameterValueWidget(class UDEditorParameterValue* Parameter, TSharedPtr<IPropertyHandle> ParameterValueProperty, IDetailGroup& DetailGroup);
+	void CreateLabeledTextureParameterValueWidget(class UDEditorParameterValue* Parameter, TSharedPtr<IPropertyHandle> ParameterValueProperty, IDetailGroup& DetailGroup);
+	void CreateParameterCollectionParameterValueWidget(class UDEditorParameterValue* Parameter, TSharedPtr<IPropertyHandle> ParameterValueProperty, IDetailGroup& DetailGroup, class UMaterialParameterCollection* BaseParameterCollection, TArray<TSharedRef<IPropertyHandle>>&& AllPropertyHandles);
 
 	FString GetCurvePath(class UDEditorScalarParameterValue * Parameter) const;
 
diff --git a/Engine/Source/Runtime/Engine/Private/Materials/ParameterCollection.cpp b/Engine/Source/Runtime/Engine/Private/Materials/ParameterCollection.cpp
--- a/Engine/Source/Runtime/Engine/Private/Materials/ParameterCollection.cpp
+++ b/Engine/Source/Runtime/Engine/Private/Materials/ParameterCollection.cpp
@@ -368,6 +368,11 @@ void UMaterialParameterCollection::PostEditChangeProperty(FPropertyChangedEvent&
 	Super::PostEditChangeProperty(PropertyChangedEvent);
 }
 
+FName UMaterialParameterCollection::GetBaseParameterCollectionMemberName()
+{
+	return GET_MEMBER_NAME_CHECKED(UMaterialParameterCollection, Base);
+}
+
 #endif // WITH_EDITOR
 
 TTuple<FString, int32> SplitParameterName(FName ParameterName)
@@ -608,6 +613,11 @@ FLinearColor UMaterialParameterCollection::GetVectorParameterDefaultValue(FName
 	return FLinearColor::Black;
 }
 
+UMaterialParameterCollection* UMaterialParameterCollection::GetBaseParameterCollection() const
+{
+	return Base;
+}
+
 FName UMaterialParameterCollection::GetParameterName(const FGuid& Id) const
 {
 	for (const UMaterialParameterCollection* Collection = this; Collection; Collection = Collection->Base)
diff --git a/Engine/Source/Runtime/Engine/Public/Materials/MaterialParameterCollection.h b/Engine/Source/Runtime/Engine/Public/Materials/MaterialParameterCollection.h
--- a/Engine/Source/Runtime/Engine/Public/Materials/MaterialParameterCollection.h
+++ b/Engine/Source/Runtime/Engine/Public/Materials/MaterialParameterCollection.h
@@ -137,6 +137,10 @@ public:
 	UFUNCTION(BlueprintCallable, Category="Rendering|Material", meta=(Keywords="GetFloatParameterDefaultValue"))
 	ENGINE_API FLinearColor GetVectorParameterDefaultValue(FName ParameterName, bool& bParameterFound) const;
 	
+	/** Gets the base Material Parameter Collection for this one. Returns nullptr if not set. */
+	UFUNCTION(BlueprintCallable, Category="Rendering|Material")
+	ENGINE_API UMaterialParameterCollection* GetBaseParameterCollection() const;
+
 	//~ Begin UObject Interface.
 #if WITH_EDITOR
 	using Super::PreEditChange;
@@ -190,6 +194,10 @@ public:
 
 	ENGINE_API bool IsInstanceOf(const FGuid& CollectionStateId) const;
 
+#if WITH_EDITOR
+	static ENGINE_API FName GetBaseParameterCollectionMemberName();
+#endif
+
 private:
 	virtual ENGINE_API void FinishDestroy() override;
 	virtual ENGINE_API bool IsReadyForFinishDestroy() override;
@@ -199,7 +207,7 @@ private:
 
 	bool UpdateOverrides(UMaterialParameterCollection* BaseCollection);
 
-	UPROPERTY(EditAnywhere, Category=Material)
+	UPROPERTY(EditAnywhere, Category=Material, AssetRegistrySearchable)
 	TObjectPtr<UMaterialParameterCollection> Base;
 
 	UPROPERTY()
