Refactor material property widget creation code into static functions in FMaterialPropertyHelpers, bring material layers to parity with standard material properties.

#tests Material Instance containing one of every kind of parameter widget as well as a Material Layer containing the same parameters.
#rnx

diff --git a/Engine/Source/Editor/MaterialEditor/Private/MaterialEditorInstanceDetailCustomization.cpp b/Engine/Source/Editor/MaterialEditor/Private/MaterialEditorInstanceDetailCustomization.cpp
--- a/Engine/Source/Editor/MaterialEditor/Private/MaterialEditorInstanceDetailCustomization.cpp
+++ b/Engine/Source/Editor/MaterialEditor/Private/MaterialEditorInstanceDetailCustomization.cpp
@@ -12,24 +12,10 @@
 #include "Styling/AppStyle.h"
 #include "MaterialCachedData.h"
 #include "Materials/MaterialInterface.h"
-#include "MaterialEditor/DEditorFontParameterValue.h"
-#include "MaterialEditor/DEditorMaterialLayersParameterValue.h"
 #include "MaterialEditor/DEditorParameterCollectionParameterValue.h"
-#include "MaterialEditor/DEditorRuntimeVirtualTextureParameterValue.h"
-#include "MaterialEditor/DEditorSparseVolumeTextureParameterValue.h"
-#include "MaterialEditor/DEditorScalarParameterValue.h"
-#include "MaterialEditor/DEditorStaticComponentMaskParameterValue.h"
-#include "MaterialEditor/DEditorStaticSwitchParameterValue.h"
-#include "MaterialEditor/DEditorTextureParameterValue.h"
-#include "MaterialEditor/DEditorVectorParameterValue.h"
 #include "MaterialEditor/MaterialEditorInstanceConstant.h"
 #include "SMaterialSubstrateTree.h"
-#include "Materials/MaterialInterface.h"
 #include "Materials/MaterialInstance.h"
-#include "Materials/MaterialExpressionParameter.h"
-#include "Materials/MaterialExpressionTextureSampleParameter.h"
-#include "Materials/MaterialExpressionFontSampleParameter.h"
-#include "Materials/MaterialExpressionMaterialAttributeLayers.h"
 #include "Materials/MaterialParameterCollection.h"
 #include "MaterialShared.h"
 #include "EditorSupportDelegates.h"
@@ -74,24 +60,6 @@ FMaterialInstanceParameterDetails::FMaterialInstanceParameterDetails(UMaterialEd
 {
 }
 
-TOptional<float> FMaterialInstanceParameterDetails::OnGetValue(TSharedRef<IPropertyHandle> PropertyHandle)
-{
-	float Value = 0.0f;
-	if (PropertyHandle->GetValue(Value) == FPropertyAccess::Success)
-	{
-		return TOptional<float>(Value);
-	}
-
-	// Value couldn't be accessed. Return an unset value
-	return TOptional<float>();
-}
-
-void FMaterialInstanceParameterDetails::OnValueCommitted(float NewValue, ETextCommit::Type CommitType, TSharedRef<IPropertyHandle> PropertyHandle)
-{	
-	// Try setting as float, if that fails then set as int
-	ensure(PropertyHandle->SetValue(NewValue) == FPropertyAccess::Success);
-}
-
 FString FMaterialInstanceParameterDetails::GetFunctionParentPath() const
 {
 	FString PathString;
@@ -521,6 +489,9 @@ void FMaterialInstanceParameterDetails::CreateSingleGroupWidget(FEditorParameter
 
 	TMap<UMaterialParameterCollection*, TArray<int32>> ParameterCollectionSharedOverrideMap;
 
+	FMaterialParameterPropertyRowOverrides RowOverrides;
+	RowOverrides.ShowHiddenDelegate = ShowHiddenDelegate;
+
 	// Create a custom widget for each parameter in the group
 	for (int32 ParamIdx = 0; ParamIdx < ParameterGroup.Parameters.Num(); ++ParamIdx)
 	{
@@ -531,18 +502,7 @@ void FMaterialInstanceParameterDetails::CreateSingleGroupWidget(FEditorParameter
 		if (ParameterProperty.IsValid() && ParameterValueProperty.IsValid() && ParameterValueProperty->IsValidHandle() &&
 			(GroupIndex == INDEX_NONE || Parameter->ParameterInfo.Index == GroupIndex))
 		{
-			UDEditorFontParameterValue* FontParam = Cast<UDEditorFontParameterValue>(Parameter);
-			UDEditorMaterialLayersParameterValue* LayersParam = Cast<UDEditorMaterialLayersParameterValue>(Parameter);
-			UDEditorScalarParameterValue* ScalarParam = Cast<UDEditorScalarParameterValue>(Parameter);
-			UDEditorStaticComponentMaskParameterValue* CompMaskParam = Cast<UDEditorStaticComponentMaskParameterValue>(Parameter);
-			UDEditorStaticSwitchParameterValue* SwitchParam = Cast<UDEditorStaticSwitchParameterValue>(Parameter);
-			UDEditorTextureParameterValue* TextureParam = Cast<UDEditorTextureParameterValue>(Parameter);
-			UDEditorRuntimeVirtualTextureParameterValue* RuntimeVirtualTextureParam = Cast<UDEditorRuntimeVirtualTextureParameterValue>(Parameter);
-			UDEditorSparseVolumeTextureParameterValue* SparseVolumeTextureParam = Cast<UDEditorSparseVolumeTextureParameterValue>(Parameter);
-			UDEditorVectorParameterValue* VectorParam = Cast<UDEditorVectorParameterValue>(Parameter);
-			UDEditorParameterCollectionParameterValue* ParameterCollectionParam = Cast<UDEditorParameterCollectionParameterValue>(Parameter);
-
-			if (ParameterCollectionParam)
+			if (UDEditorParameterCollectionParameterValue* ParameterCollectionParam = Cast<UDEditorParameterCollectionParameterValue>(Parameter))
 			{
 				if (MaterialEditorInstance && MaterialEditorInstance->SourceInstance)
 				{
@@ -563,73 +523,17 @@ void FMaterialInstanceParameterDetails::CreateSingleGroupWidget(FEditorParameter
 				continue;
 			}
 
-			// Don't display custom primitive data parameters in the details panel.
-			// This data is pulled from the primitive and can't be changed on the material.
-			if ((VectorParam && VectorParam->bUseCustomPrimitiveData) ||
-				(ScalarParam && ScalarParam->bUseCustomPrimitiveData))
-			{
-				continue;
-			}
-
 			if (Parameter->ParameterInfo.Association == EMaterialParameterAssociation::GlobalParameter || bForceShowParam)
 			{
-				if (VectorParam && VectorParam->bIsUsedAsChannelMask)
-				{
-					CreateVectorChannelMaskParameterValueWidget(Parameter, ParameterValueProperty, DetailGroup);
-				}
-				if (ScalarParam && ScalarParam->AtlasData.bIsUsedAsAtlasPosition)
-				{
-					CreateScalarAtlasPositionParameterValueWidget(Parameter, ParameterValueProperty, DetailGroup);
-				}
-				if (TextureParam && 
-					( !TextureParam->ChannelNames.R.IsEmpty()
-					|| !TextureParam->ChannelNames.G.IsEmpty()
-					|| !TextureParam->ChannelNames.B.IsEmpty()
-					|| !TextureParam->ChannelNames.A.IsEmpty()))
-				{
-					CreateLabeledTextureParameterValueWidget(Parameter, ParameterValueProperty, DetailGroup);
-				}
-				else if (LayersParam)
-				{
-				}
-				else if (CompMaskParam)
+				if (!FMaterialPropertyHelpers::ShouldCreatePropertyRowForParameter(Parameter))
 				{
-					CreateMaskParameterValueWidget(Parameter, ParameterValueProperty, DetailGroup);
+					continue;
 				}
-				else
-				{
-					if (ScalarParam && ScalarParam->SliderMax > ScalarParam->SliderMin)
-					{
-						ParameterValueProperty->SetInstanceMetaData("UIMin", FString::Printf(TEXT("%f"), ScalarParam->SliderMin));
-						ParameterValueProperty->SetInstanceMetaData("UIMax", FString::Printf(TEXT("%f"), ScalarParam->SliderMax));
-					}
 
-					if (VectorParam)
-					{
-						static const FName Red("R");
-						static const FName Green("G");
-						static const FName Blue("B");
-						static const FName Alpha("A");
-						if (!VectorParam->ChannelNames.R.IsEmpty())
-						{
-							ParameterProperty->GetChildHandle(Red)->SetPropertyDisplayName(VectorParam->ChannelNames.R);
-						}
-						if (!VectorParam->ChannelNames.G.IsEmpty())
-						{
-							ParameterProperty->GetChildHandle(Green)->SetPropertyDisplayName(VectorParam->ChannelNames.G);
-						}
-						if (!VectorParam->ChannelNames.B.IsEmpty())
-						{
-							ParameterProperty->GetChildHandle(Blue)->SetPropertyDisplayName(VectorParam->ChannelNames.B);
-						}
-						if (!VectorParam->ChannelNames.A.IsEmpty())
-						{
-							ParameterProperty->GetChildHandle(Alpha)->SetPropertyDisplayName(VectorParam->ChannelNames.A);
-						}
-					}
+				IDetailPropertyRow& PropertyRow = DetailGroup.AddPropertyRow(ParameterValueProperty.ToSharedRef());
 
-					CreateParameterValueWidget(Parameter, ParameterValueProperty, DetailGroup);
-				}
+				FMaterialPropertyHelpers::ConfigurePropertyRowForParameter(PropertyRow, Parameter, ParameterValueProperty, MaterialEditorInstance, RowOverrides);
+				FMaterialPropertyHelpers::SetPropertyRowParameterWidget(PropertyRow, Parameter, ParameterValueProperty, MaterialEditorInstance);
 			}
 		}
 	}
@@ -660,530 +564,13 @@ void FMaterialInstanceParameterDetails::CreateSingleGroupWidget(FEditorParameter
 				}
 
 				// Create a parameter value widget for the first parameter.
-				CreateParameterCollectionParameterValueWidget(FirstParameter, FirstParameterValueProperty, DetailGroup, ParameterCollection, MoveTemp(PropertyHandles));
-			}
-		}
-	}
-}
-
-IDetailPropertyRow& FMaterialInstanceParameterDetails::CreateParameterPropertyRow(UDEditorParameterValue* Parameter, TSharedPtr<IPropertyHandle> ParameterValueProperty, IDetailGroup& DetailGroup, TOptional<FResetToDefaultOverride> ResetToDefaultOverride)
-{
-	IDetailPropertyRow& PropertyRow = DetailGroup.AddPropertyRow(ParameterValueProperty.ToSharedRef());
-
-	FResetToDefaultOverride ResetOverride;
-	if (ResetToDefaultOverride.IsSet())
-	{
-		ResetOverride = ResetToDefaultOverride.GetValue();
-	}
-	else
-	{
-		TAttribute<bool> IsResetVisible = TAttribute<bool>::Create(TAttribute<bool>::FGetter::CreateStatic(&FMaterialPropertyHelpers::ShouldShowResetToDefault, Parameter, MaterialEditorInstance));
-		FSimpleDelegate ResetHandler = FSimpleDelegate::CreateStatic(&FMaterialPropertyHelpers::ResetToDefault, Parameter, MaterialEditorInstance);
-		ResetOverride = FResetToDefaultOverride::Create(IsResetVisible, ResetHandler);
-	}
-
-	TAttribute<bool> IsParamEnabled = TAttribute<bool>::Create(TAttribute<bool>::FGetter::CreateStatic(&FMaterialPropertyHelpers::IsOverriddenExpression, Parameter));
-
-	return PropertyRow
-		.DisplayName(FText::FromName(Parameter->ParameterInfo.Name))
-		.ToolTip(FMaterialPropertyHelpers::GetParameterTooltip(Parameter, MaterialEditorInstance))
-		.EditCondition(IsParamEnabled, FOnBooleanValueChanged::CreateStatic(&FMaterialPropertyHelpers::OnOverrideParameter, Parameter, MaterialEditorInstance))
-		.Visibility(TAttribute<EVisibility>::Create(TAttribute<EVisibility>::FGetter::CreateStatic(&FMaterialPropertyHelpers::ShouldShowExpression, Parameter, MaterialEditorInstance, ShowHiddenDelegate)))
-		.OverrideResetToDefault(ResetOverride);
-}
-
-void FMaterialInstanceParameterDetails::CreateParameterValueWidget(UDEditorParameterValue* Parameter, TSharedPtr<IPropertyHandle> ParameterValueProperty, IDetailGroup& DetailGroup)
-{
-	IDetailPropertyRow& PropertyRow = CreateParameterPropertyRow(Parameter, ParameterValueProperty, DetailGroup);
-
-	// Textures need a special widget that filters based on VT or not
-	UDEditorTextureParameterValue* TextureParam = Cast<UDEditorTextureParameterValue>(Parameter);
-	if (TextureParam != nullptr)
-	{
-		UMaterial *Material = MaterialEditorInstance->SourceInstance->GetMaterial();
-		if (Material != nullptr)
-		{
-			UMaterialExpressionTextureSampleParameter* Expression = Material->FindExpressionByGUID<UMaterialExpressionTextureSampleParameter>(TextureParam->ExpressionId);
-			if (Expression != nullptr)
-			{
-				TWeakObjectPtr<UMaterialExpressionTextureSampleParameter> SamplerExpression = Expression;
-
-				PropertyRow.CustomWidget()
-				.NameContent()
-				[
-					ParameterValueProperty->CreatePropertyNameWidget()
-				]
-				.ValueContent()
-				.MaxDesiredWidth(TOptional<float>())
-				[
-					SNew(SObjectPropertyEntryBox)
-					.PropertyHandle(ParameterValueProperty)
-					.AllowedClass(UTexture::StaticClass())
-					.ThumbnailPool(PropertyUtilities.Pin()->GetThumbnailPool())
-					.OnShouldFilterAsset_Lambda([SamplerExpression](const FAssetData& AssetData)
-					{
-						if (SamplerExpression.Get())
-						{
-							bool VirtualTextured = false;
-							AssetData.GetTagValue<bool>("VirtualTextureStreaming", VirtualTextured);
-
-							bool ExpressionIsVirtualTextured = IsVirtualSamplerType(SamplerExpression->SamplerType);
+				IDetailPropertyRow& PropertyRow = DetailGroup.AddPropertyRow(FirstParameterValueProperty.ToSharedRef());
 
-							return VirtualTextured != ExpressionIsVirtualTextured;
-						}
-						else
-						{
-							return false;
-						}
-					})
-				];
+				FMaterialPropertyHelpers::ConfigurePropertyRowForParameter(PropertyRow, FirstParameter, FirstParameterValueProperty, MaterialEditorInstance, RowOverrides);
+				FMaterialPropertyHelpers::SetPropertyRowParameterCollectionParameterWidget(PropertyRow, FirstParameter, FirstParameterValueProperty, ParameterCollection, MoveTemp(PropertyHandles));
 			}
 		}
 	}
-
-	// Scalar parameters can present an enumeration combo box.
-	if (UDEditorScalarParameterValue* ScalarParam = Cast<UDEditorScalarParameterValue>(Parameter))
-	{
-		TSoftObjectPtr<UObject> Enumeration;
-		if (!ScalarParam->Enumeration.IsNull())
-		{
-			Enumeration = ScalarParam->Enumeration;
-		}
-		else if (ScalarParam->EnumerationIndex >= 0)
-		{
-			Enumeration = MaterialEditorInstance->GetEnumerationObject(ScalarParam->EnumerationIndex);
-		}
-
-		if (!Enumeration.IsNull())
-		{
-			PropertyRow.CustomWidget()
-			.NameContent()
-			[
-				ParameterValueProperty->CreatePropertyNameWidget()
-			]
-			.ValueContent()
-			[
-				FMaterialPropertyHelpers::CreateScalarEnumerationParameterValueWidget(Enumeration, ParameterValueProperty)
-			];
-		}
-	}
-}
-
-void FMaterialInstanceParameterDetails::CreateMaskParameterValueWidget(UDEditorParameterValue* Parameter, TSharedPtr<IPropertyHandle> ParameterValueProperty, IDetailGroup& DetailGroup )
-{
-	TSharedPtr<IPropertyHandle> RMaskProperty = ParameterValueProperty->GetChildHandle("R");
-	TSharedPtr<IPropertyHandle> GMaskProperty = ParameterValueProperty->GetChildHandle("G");
-	TSharedPtr<IPropertyHandle> BMaskProperty = ParameterValueProperty->GetChildHandle("B");
-	TSharedPtr<IPropertyHandle> AMaskProperty = ParameterValueProperty->GetChildHandle("A");
-
-	IDetailPropertyRow& PropertyRow = CreateParameterPropertyRow(Parameter, ParameterValueProperty, DetailGroup);
-
-	PropertyRow.CustomWidget()
-	.NameContent()
-	[
-		ParameterValueProperty->CreatePropertyNameWidget()
-	]
-	.ValueContent()
-		.MaxDesiredWidth(200.0f)
-		[
-			SNew(SHorizontalBox)
-			+ SHorizontalBox::Slot()
-		.FillWidth(1.0f)
-		[
-			SNew(SHorizontalBox)
-			+ SHorizontalBox::Slot()
-		.HAlign(HAlign_Left)
-		.AutoWidth()
-		[
-			RMaskProperty->CreatePropertyNameWidget()
-		]
-	+ SHorizontalBox::Slot()
-		.HAlign(HAlign_Left)
-		.AutoWidth()
-		[
-			RMaskProperty->CreatePropertyValueWidget()
-		]
-	+ SHorizontalBox::Slot()
-		.HAlign(HAlign_Left)
-		.Padding(FMargin(10.0f, 0.0f, 0.0f, 0.0f))
-		.AutoWidth()
-		[
-			GMaskProperty->CreatePropertyNameWidget()
-		]
-	+ SHorizontalBox::Slot()
-		.HAlign(HAlign_Left)
-		.AutoWidth()
-		[
-			GMaskProperty->CreatePropertyValueWidget()
-		]
-	+ SHorizontalBox::Slot()
-		.HAlign(HAlign_Left)
-		.Padding(FMargin(10.0f, 0.0f, 0.0f, 0.0f))
-		.AutoWidth()
-		[
-			BMaskProperty->CreatePropertyNameWidget()
-		]
-	+ SHorizontalBox::Slot()
-		.HAlign(HAlign_Left)
-		.AutoWidth()
-		[
-			BMaskProperty->CreatePropertyValueWidget()
-		]
-	+ SHorizontalBox::Slot()
-		.HAlign(HAlign_Left)
-		.Padding(FMargin(10.0f, 0.0f, 0.0f, 0.0f))
-		.AutoWidth()
-		[
-			AMaskProperty->CreatePropertyNameWidget()
-		]
-	+ SHorizontalBox::Slot()
-		.HAlign(HAlign_Left)
-		.AutoWidth()
-		[
-			AMaskProperty->CreatePropertyValueWidget()
-		]
-		]
-		];
-}
-
-void FMaterialInstanceParameterDetails::CreateVectorChannelMaskParameterValueWidget(UDEditorParameterValue* Parameter, TSharedPtr<IPropertyHandle> ParameterValueProperty, IDetailGroup& DetailGroup)
-{
-	IDetailPropertyRow& PropertyRow = CreateParameterPropertyRow(Parameter, ParameterValueProperty, DetailGroup);
-
-	// Combo box hooks for converting between our "enum" and colors
-	FOnGetPropertyComboBoxStrings GetMaskStrings = FOnGetPropertyComboBoxStrings::CreateStatic(&FMaterialPropertyHelpers::GetVectorChannelMaskComboBoxStrings);
-	FOnGetPropertyComboBoxValue GetMaskValue = FOnGetPropertyComboBoxValue::CreateStatic(&FMaterialPropertyHelpers::GetVectorChannelMaskValue, Parameter);
-	FOnPropertyComboBoxValueSelected SetMaskValue = FOnPropertyComboBoxValueSelected::CreateStatic(&FMaterialPropertyHelpers::SetVectorChannelMaskValue, ParameterValueProperty, Parameter, (UObject*)MaterialEditorInstance);
-
-	// Widget replaces color picker with combo box
-	PropertyRow.CustomWidget()
-	.NameContent()
-	[
-		ParameterValueProperty->CreatePropertyNameWidget()
-	]
-	.ValueContent()
-	.MaxDesiredWidth(200.0f)
-	[
-		SNew(SHorizontalBox)
-		+ SHorizontalBox::Slot()
-		.FillWidth(1.0f)
-		[
-			SNew(SHorizontalBox)
-			+ SHorizontalBox::Slot()
-			.HAlign(HAlign_Left)
-			.AutoWidth()
-			[
-				PropertyCustomizationHelpers::MakePropertyComboBox(ParameterValueProperty, GetMaskStrings, GetMaskValue, SetMaskValue)
-			]
-		]
-	];
-}
-
-void FMaterialInstanceParameterDetails::CreateScalarAtlasPositionParameterValueWidget(class UDEditorParameterValue* Parameter, TSharedPtr<IPropertyHandle> ParameterValueProperty, IDetailGroup& DetailGroup)
-{
-	// Handle reset to default manually
-	TAttribute<bool> IsResetVisible = TAttribute<bool>::Create(TAttribute<bool>::FGetter::CreateStatic(&FMaterialPropertyHelpers::ShouldShowResetToDefault, Parameter, MaterialEditorInstance));
-	FSimpleDelegate ResetHandler = FSimpleDelegate::CreateStatic(&FMaterialPropertyHelpers::ResetCurveToDefault, Parameter, MaterialEditorInstance);
-	FResetToDefaultOverride ResetOverride = FResetToDefaultOverride::Create(IsResetVisible, ResetHandler);
-
-	IDetailPropertyRow& PropertyRow = CreateParameterPropertyRow(Parameter, ParameterValueProperty, DetailGroup, ResetOverride);
-
-	const FText ParameterName = FText::FromName(Parameter->ParameterInfo.Name);
-	UDEditorScalarParameterValue* AtlasParameter = Cast<UDEditorScalarParameterValue>(Parameter);
-
-	FDetailWidgetRow& CustomWidget = PropertyRow.CustomWidget();
-	CustomWidget
-		.FilterString(ParameterName)
-		.NameContent()
-		[
-			ParameterValueProperty->CreatePropertyNameWidget()
-		]
-		.ValueContent()
-		.HAlign(HAlign_Fill)
-		.MaxDesiredWidth(400.0f)
-		[
-			SNew(SObjectPropertyEntryBox)
-			.ObjectPath(this, &FMaterialInstanceParameterDetails::GetCurvePath, AtlasParameter)
-			.AllowedClass(UCurveLinearColor::StaticClass())
-			.NewAssetFactories(TArray<UFactory*>())
-			.DisplayThumbnail(true)
-			.ThumbnailPool(PropertyUtilities.Pin()->GetThumbnailPool())
-			.OnShouldFilterAsset(FOnShouldFilterAsset::CreateStatic(&FMaterialPropertyHelpers::OnShouldFilterCurveAsset, AtlasParameter->AtlasData.Atlas))
-			.OnShouldSetAsset(FOnShouldSetAsset::CreateStatic(&FMaterialPropertyHelpers::OnShouldSetCurveAsset, AtlasParameter->AtlasData.Atlas))
-			.OnObjectChanged(FOnSetObject::CreateStatic(&FMaterialPropertyHelpers::SetPositionFromCurveAsset, AtlasParameter->AtlasData.Atlas, AtlasParameter, ParameterValueProperty, (UObject*)MaterialEditorInstance))
-			.DisplayCompactSize(true)
-		];
-}
-
-void FMaterialInstanceParameterDetails::CreateLabeledTextureParameterValueWidget(class UDEditorParameterValue* Parameter, TSharedPtr<IPropertyHandle> ParameterValueProperty, IDetailGroup& DetailGroup)
-{
-	UDEditorTextureParameterValue* TextureParam = Cast<UDEditorTextureParameterValue>(Parameter);
-	if (TextureParam)
-	{
-		UMaterial *Material = MaterialEditorInstance->SourceInstance->GetMaterial();
-		if (Material != nullptr)
-		{
-			UMaterialExpressionTextureSampleParameter* Expression = Material->FindExpressionByGUID<UMaterialExpressionTextureSampleParameter>(TextureParam->ExpressionId);
-			if (Expression != nullptr)
-			{
-				TWeakObjectPtr<UMaterialExpressionTextureSampleParameter> SamplerExpression = Expression;
-
-				IDetailPropertyRow& PropertyRow = CreateParameterPropertyRow(Parameter, ParameterValueProperty, DetailGroup);
-
-				TSharedPtr<SWidget> NameWidget;
-				TSharedPtr<SWidget> ValueWidget;
-				FDetailWidgetRow Row;
-				PropertyRow.GetDefaultWidgets(NameWidget, ValueWidget, Row);
-
-				FDetailWidgetRow &DetailWidgetRow = PropertyRow.CustomWidget();
-				TSharedPtr<SVerticalBox> NameVerticalBox;
-				DetailWidgetRow.NameContent()
-					[
-						SAssignNew(NameVerticalBox, SVerticalBox)
-						+ SVerticalBox::Slot()
-						.AutoHeight()
-						[
-							ParameterValueProperty->CreatePropertyNameWidget()
-						]
-					];
-				DetailWidgetRow.ValueContent()
-					.MinDesiredWidth(Row.ValueWidget.MinWidth)
-					.MaxDesiredWidth(Row.ValueWidget.MaxWidth)
-					[
-						SNew(SObjectPropertyEntryBox)
-						.PropertyHandle(ParameterValueProperty)
-						.AllowedClass(UTexture::StaticClass())
-						.ThumbnailPool(PropertyUtilities.Pin()->GetThumbnailPool())
-						.OnShouldFilterAsset_Lambda([SamplerExpression](const FAssetData& AssetData)
-						{
-							if (SamplerExpression.Get())
-							{
-								bool VirtualTextured = false;
-								AssetData.GetTagValue<bool>("VirtualTextureStreaming", VirtualTextured);
-
-								bool ExpressionIsVirtualTextured = IsVirtualSamplerType(SamplerExpression->SamplerType);
-
-								return VirtualTextured != ExpressionIsVirtualTextured;
-							}
-							else
-							{
-								return false;
-							}
-						})
-					];
-
-				static const FName Red("R");
-				static const FName Green("G");
-				static const FName Blue("B");
-				static const FName Alpha("A");
-
-				if (!TextureParam->ChannelNames.R.IsEmpty())
-				{
-					NameVerticalBox->AddSlot()
-						[
-							SNew(SHorizontalBox)
-							+ SHorizontalBox::Slot()
-							.AutoWidth()
-							.Padding(20.0, 2.0, 4.0, 2.0)
-							[
-								SNew(STextBlock)
-								.Text(FText::FromName(Red))
-								.Font(FAppStyle::GetFontStyle(TEXT("PropertyWindow.BoldFont")))
-							]
-							+ SHorizontalBox::Slot()
-							.HAlign(HAlign_Left)
-							.Padding(4.0, 2.0)
-							[
-								SNew(STextBlock)
-								.Text(TextureParam->ChannelNames.R)
-								.Font(FAppStyle::GetFontStyle(TEXT("PropertyWindow.NormalFont")))
-							]
-						];
-				}
-				if (!TextureParam->ChannelNames.G.IsEmpty())
-				{
-					NameVerticalBox->AddSlot()
-						[
-							SNew(SHorizontalBox)
-							+ SHorizontalBox::Slot()
-						.Padding(20.0, 2.0, 4.0, 2.0)
-						.AutoWidth()
-						[
-							SNew(STextBlock)
-							.Text(FText::FromName(Green))
-						.Font(FAppStyle::GetFontStyle(TEXT("PropertyWindow.BoldFont")))
-						]
-					+ SHorizontalBox::Slot()
-						.HAlign(HAlign_Left)
-						.Padding(4.0, 2.0)
-						[
-							SNew(STextBlock)
-							.Text(TextureParam->ChannelNames.G)
-						.Font(FAppStyle::GetFontStyle(TEXT("PropertyWindow.NormalFont")))
-						]
-						];
-				}
-				if (!TextureParam->ChannelNames.B.IsEmpty())
-				{
-					NameVerticalBox->AddSlot()
-						[
-							SNew(SHorizontalBox)
-							+ SHorizontalBox::Slot()
-						.Padding(20.0, 2.0, 4.0, 2.0)
-						.AutoWidth()
-						[
-							SNew(STextBlock)
-							.Text(FText::FromName(Blue))
-						.Font(FAppStyle::GetFontStyle(TEXT("PropertyWindow.BoldFont")))
-						]
-					+ SHorizontalBox::Slot()
-						.HAlign(HAlign_Left)
-						.Padding(4.0, 2.0)
-						[
-							SNew(STextBlock)
-							.Text(TextureParam->ChannelNames.B)
-						.Font(FAppStyle::GetFontStyle(TEXT("PropertyWindow.NormalFont")))
-						]
-						];
-				}
-				if (!TextureParam->ChannelNames.A.IsEmpty())
-				{
-					NameVerticalBox->AddSlot()
-						[
-							SNew(SHorizontalBox)
-							+ SHorizontalBox::Slot()
-						.Padding(20.0, 2.0, 4.0, 2.0)
-						.AutoWidth()
-						[
-							SNew(STextBlock)
-							.Text(FText::FromName(Alpha))
-						.Font(FAppStyle::GetFontStyle(TEXT("PropertyWindow.BoldFont")))
-						]
-					+ SHorizontalBox::Slot()
-						.HAlign(HAlign_Left)
-						.Padding(4.0, 2.0)
-						[
-							SNew(STextBlock)
-							.Text(TextureParam->ChannelNames.A)
-						.Font(FAppStyle::GetFontStyle(TEXT("PropertyWindow.NormalFont")))
-						]
-						];
-				}
-			}
-		}
-	}
-}
-
-void FMaterialInstanceParameterDetails::CreateParameterCollectionParameterValueWidget(class UDEditorParameterValue* Parameter, TSharedPtr<IPropertyHandle> ParameterValueProperty, IDetailGroup& DetailGroup, class UMaterialParameterCollection* BaseParameterCollection, TArray<TSharedRef<IPropertyHandle>>&& AllPropertyHandles)
-{
-	IDetailPropertyRow& PropertyRow = CreateParameterPropertyRow(Parameter, ParameterValueProperty, DetailGroup);
-
-	// Display the parameter collection's name on the widget.
-	ParameterValueProperty->SetPropertyDisplayName(FText::FromName(BaseParameterCollection->GetFName()));
-
-	// Create a delegate to invoke when the property value changes.
-	TDelegate<void(const FPropertyChangedEvent&)> OnPropertyValueChangedDelegate = TDelegate<void(const FPropertyChangedEvent&)>::CreateWeakLambda(BaseParameterCollection,
-		[BaseParameterCollection, PropertyHandles = MoveTemp(AllPropertyHandles)](const FPropertyChangedEvent& PropertyChangedEvent)
-		{
-			// Determine if this is the final setting of the property value.
-			if (PropertyChangedEvent.ChangeType != EPropertyChangeType::ValueSet)
-			{
-				return;
-			}
-
-			if (PropertyHandles.IsEmpty())
-			{
-				return;
-			}
-
-			// Get the new value from the first property handle.
-			UObject* Value = nullptr;
-			if (PropertyHandles[0]->GetValue(Value) != FPropertyAccess::Success)
-			{
-				return;
-			}
-
-			// Verify that the new value is valid.
-			if (auto* Collection = Cast<UMaterialParameterCollection>(Value))
-			{
-				if (Collection->IsInstanceOf(BaseParameterCollection->StateId))
-				{
-					// Set the same value on all of the collection's referencing parameters.
-					for (int32 PropertyHandleIndex = 1; PropertyHandleIndex < PropertyHandles.Num(); ++PropertyHandleIndex)
-					{
-						PropertyHandles[PropertyHandleIndex]->SetValue(Value);
-					}
-
-					return;
-				}
-			}
-
-			// Reset the value to the default parameter collection.
-			Value = BaseParameterCollection;
-			PropertyHandles[0]->SetValue(Value);
-		});
-
-	// Invoke the delegate when the property value or one of its child property values changes.
-	ParameterValueProperty->SetOnPropertyValueChangedWithData(OnPropertyValueChangedDelegate);
-	ParameterValueProperty->SetOnChildPropertyValueChangedWithData(OnPropertyValueChangedDelegate);
-
-	FSoftObjectPath BaseParameterCollectionPath = BaseParameterCollection;
-
-	PropertyRow.CustomWidget()
-		.NameContent()
-		[
-			ParameterValueProperty->CreatePropertyNameWidget()
-		]
-		.ValueContent()
-		.MaxDesiredWidth(TOptional<float>())
-		[
-			SNew(SObjectPropertyEntryBox)
-				.PropertyHandle(ParameterValueProperty)
-				.AllowedClass(UMaterialParameterCollection::StaticClass())
-				.ThumbnailPool(PropertyUtilities.Pin()->GetThumbnailPool())
-				.OnShouldFilterAsset_Lambda([BaseParameterCollectionPath](const FAssetData& AssetData)
-					{
-						// Allow the base MPC itself
-						if (AssetData.GetSoftObjectPath() == BaseParameterCollectionPath)
-						{
-							return false;
-						}
-
-						IAssetRegistry& AssetRegistry = IAssetRegistry::GetChecked();
-
-						// Traverse up the hierarchy of UMaterialParameterCollection::Base until we find the base material's MPC
-						FSoftObjectPath ParentAssetPath = FSoftObjectPath(AssetData.GetTagValueRef<FString>(UMaterialParameterCollection::GetBaseParameterCollectionMemberName()));
-						while (!ParentAssetPath.IsNull())
-						{
-							if (ParentAssetPath == BaseParameterCollectionPath)
-							{
-								return false;
-							}
-
-							FAssetData ParentAssetData = AssetRegistry.GetAssetByObjectPath(ParentAssetPath);
-							ParentAssetPath = FSoftObjectPath(ParentAssetData.GetTagValueRef<FString>(UMaterialParameterCollection::GetBaseParameterCollectionMemberName()));
-						}
-
-						// No matching MPC in hierarchy, filter this MPC out
-						return true;
-					})
-		];
-}
-
-FString FMaterialInstanceParameterDetails::GetCurvePath(UDEditorScalarParameterValue* Parameter) const
-{
-	FString Path = Parameter->AtlasData.Curve->GetPathName();
-	return Path;
-}
-
-bool FMaterialInstanceParameterDetails::IsVisibleExpression(UDEditorParameterValue* Parameter)
-{
-	return MaterialEditorInstance->VisibleExpressions.Contains(Parameter->ParameterInfo);
-}
-
-EVisibility FMaterialInstanceParameterDetails::ShouldShowExpression(UDEditorParameterValue* Parameter) const
-{
-	return FMaterialPropertyHelpers::ShouldShowExpression(Parameter, MaterialEditorInstance, ShowHiddenDelegate);
 }
 
 bool FMaterialInstanceParameterDetails::OnShouldSetAsset(const FAssetData& AssetData) const
diff --git a/Engine/Source/Editor/MaterialEditor/Private/MaterialEditorInstanceDetailCustomization.h b/Engine/Source/Editor/MaterialEditor/Private/MaterialEditorInstanceDetailCustomization.h
--- a/Engine/Source/Editor/MaterialEditor/Private/MaterialEditorInstanceDetailCustomization.h
+++ b/Engine/Source/Editor/MaterialEditor/Private/MaterialEditorInstanceDetailCustomization.h
@@ -36,9 +36,6 @@ public:
 	/** IDetailCustomization interface */
 	virtual void CustomizeDetails(IDetailLayoutBuilder& DetailLayout) override;
 
-	static TOptional<float> OnGetValue(TSharedRef<IPropertyHandle> PropertyHandle);
-	static void OnValueCommitted(float NewValue, ETextCommit::Type CommitType, TSharedRef<IPropertyHandle> PropertyHandle);
-
 private:
 	/** Returns the function parent path */
 	FString GetFunctionParentPath() const;
@@ -52,25 +49,6 @@ private:
 	/** Enable/Disable all parameter properties in a group */
 	static void EnableGroupParameters(struct FEditorParameterGroup& ParameterGroup, bool ShouldEnable);
 
-	/** This method generates a PropertyRow with all of the standard settings for any parameter. */
-	IDetailPropertyRow& CreateParameterPropertyRow(UDEditorParameterValue* Parameter, TSharedPtr<IPropertyHandle> ParameterValueProperty, IDetailGroup& DetailGroup, TOptional<FResetToDefaultOverride> ResetToDefaultOverride = TOptional<FResetToDefaultOverride>());
-
-	/** These methods generate the custom widgets for the various parameter types */
-	void CreateParameterValueWidget(class UDEditorParameterValue* Parameter, TSharedPtr<IPropertyHandle> ParameterValueProperty, IDetailGroup& DetailGroup);
-	void CreateMaskParameterValueWidget(class UDEditorParameterValue* Parameter, TSharedPtr<IPropertyHandle> ParameterValueProperty, IDetailGroup& DetailGroup);
-	void CreateVectorChannelMaskParameterValueWidget(class UDEditorParameterValue* Parameter, TSharedPtr<IPropertyHandle> ParameterValueProperty, IDetailGroup& DetailGroup);
-	void CreateScalarAtlasPositionParameterValueWidget(class UDEditorParameterValue* Parameter, TSharedPtr<IPropertyHandle> ParameterValueProperty, IDetailGroup& DetailGroup);
-	void CreateLabeledTextureParameterValueWidget(class UDEditorParameterValue* Parameter, TSharedPtr<IPropertyHandle> ParameterValueProperty, IDetailGroup& DetailGroup);
-	void CreateParameterCollectionParameterValueWidget(class UDEditorParameterValue* Parameter, TSharedPtr<IPropertyHandle> ParameterValueProperty, IDetailGroup& DetailGroup, class UMaterialParameterCollection* BaseParameterCollection, TArray<TSharedRef<IPropertyHandle>>&& AllPropertyHandles);
-
-	FString GetCurvePath(class UDEditorScalarParameterValue * Parameter) const;
-
-	/** Returns true if the parameter is in the visible expressions list */
-	bool IsVisibleExpression(class UDEditorParameterValue* Parameter);
-
-	/** Returns true if the parameter should be displayed */
-	EVisibility ShouldShowExpression(class UDEditorParameterValue* Parameter) const;
-
 	/** Called to check if an asset can be set as a parent */
 	bool OnShouldSetAsset(const FAssetData& InAssetData) const;
 
diff --git a/Engine/Source/Editor/MaterialEditor/Private/MaterialPropertyHelpers.cpp b/Engine/Source/Editor/MaterialEditor/Private/MaterialPropertyHelpers.cpp
--- a/Engine/Source/Editor/MaterialEditor/Private/MaterialPropertyHelpers.cpp
+++ b/Engine/Source/Editor/MaterialEditor/Private/MaterialPropertyHelpers.cpp
@@ -28,6 +28,7 @@
 #include "Materials/MaterialExpressionFontSampleParameter.h"
 #include "Materials/MaterialExpressionMaterialAttributeLayers.h"
 #include "Materials/MaterialExpressionChannelMaskParameter.h"
+#include "Materials/MaterialParameterCollection.h"
 #include "EditorSupportDelegates.h"
 #include "DetailWidgetRow.h"
 #include "PropertyHandle.h"
@@ -42,6 +43,7 @@
 #include "Materials/MaterialFunction.h"
 #include "Materials/MaterialFunctionInterface.h"
 #include "Modules/ModuleManager.h"
+#include "AssetRegistry/IAssetRegistry.h"
 #include "AssetToolsModule.h"
 #include "Factories/MaterialInstanceConstantFactoryNew.h"
 #include "StaticParameterSet.h"
@@ -549,6 +551,565 @@ FReply FMaterialPropertyHelpers::OnClickedSaveNewLayerInstance(class UMaterialFu
 	return FReply::Handled();
 }
 
+bool FMaterialPropertyHelpers::ShouldCreatePropertyRowForParameter(UDEditorParameterValue* Parameter)
+{
+	if (UDEditorMaterialLayersParameterValue* LayersParam = Cast<UDEditorMaterialLayersParameterValue>(Parameter))
+	{
+		// No widget for layers
+		return false;
+	}
+	else if (UDEditorVectorParameterValue* VectorParam = Cast<UDEditorVectorParameterValue>(Parameter))
+	{
+		// No widget for CPD
+		return !VectorParam->bUseCustomPrimitiveData;
+	}
+	else if (UDEditorScalarParameterValue* ScalarParam = Cast<UDEditorScalarParameterValue>(Parameter))
+	{
+		// No widget for CPD
+		return !ScalarParam->bUseCustomPrimitiveData;
+	}
+
+	return true;
+}
+
+bool FMaterialPropertyHelpers::ShouldCreateChildPropertiesForParameter(UDEditorParameterValue* Parameter)
+{
+	if (UDEditorStaticComponentMaskParameterValue* CompMaskParam = Cast<UDEditorStaticComponentMaskParameterValue>(Parameter))
+	{
+		// Static component mask has a checkbox for each channel
+		return false;
+	}
+	else if (UDEditorVectorParameterValue* VectorParam = Cast<UDEditorVectorParameterValue>(Parameter))
+	{
+		// Vector channel mask ahs a dropdown for each channel
+		return !VectorParam->bIsUsedAsChannelMask;
+	}
+
+	return true;
+}
+
+void FMaterialPropertyHelpers::ConfigurePropertyRowForParameter(IDetailPropertyRow& PropertyRow, UDEditorParameterValue* Parameter, TSharedPtr<IPropertyHandle> ParameterValueProperty, UMaterialEditorInstanceConstant* MaterialEditorInstance, const FMaterialParameterPropertyRowOverrides& RowOverrides /*= FMaterialParameterPropertyRowOverrides()*/)
+{
+	// Set name and tooltip
+	PropertyRow
+		.DisplayName(FText::FromName(Parameter->ParameterInfo.Name))
+		.ToolTip(FMaterialPropertyHelpers::GetParameterTooltip(Parameter, MaterialEditorInstance))
+		.OverrideResetToDefault(FMaterialPropertyHelpers::CreateResetToDefaultOverride(Parameter, MaterialEditorInstance));
+
+	// Set Visibility
+	if (RowOverrides.Visibility.IsSet())
+	{
+		PropertyRow.Visibility(RowOverrides.Visibility);
+	}
+	else
+	{
+		PropertyRow.Visibility(TAttribute<EVisibility>::Create(TAttribute<EVisibility>::FGetter::CreateStatic(&FMaterialPropertyHelpers::ShouldShowExpression, Parameter, MaterialEditorInstance, RowOverrides.ShowHiddenDelegate)));
+	}
+
+	// Set EditCondition
+	TAttribute<bool> EditConditionValue = RowOverrides.EditConditionValue;
+	if (!EditConditionValue.IsSet())
+	{
+		EditConditionValue = TAttribute<bool>::Create(TAttribute<bool>::FGetter::CreateStatic(&FMaterialPropertyHelpers::IsOverriddenExpression, Parameter));
+	}
+
+	FOnBooleanValueChanged OnEditConditionValueChanged;
+	if (RowOverrides.OnEditConditionValueChanged.IsSet())
+	{
+		OnEditConditionValueChanged = RowOverrides.OnEditConditionValueChanged.GetValue();
+	}
+	else
+	{
+		OnEditConditionValueChanged = FOnBooleanValueChanged::CreateStatic(&FMaterialPropertyHelpers::OnOverrideParameter, Parameter, MaterialEditorInstance);
+	}
+
+	PropertyRow.EditCondition(EditConditionValue, OnEditConditionValueChanged);
+}
+
+void FMaterialPropertyHelpers::SetPropertyRowParameterWidget(IDetailPropertyRow& PropertyRow, UDEditorParameterValue* Parameter, TSharedPtr<IPropertyHandle> ParameterValueProperty, UMaterialEditorParameters* MaterialEditorParameters)
+{
+	bool bUseDefaultWidget = false;
+
+	// Set custom widgets if this parameter needs one. Otherwise, the default widget for the property will be used
+	if (UDEditorStaticComponentMaskParameterValue* CompMaskParam = Cast<UDEditorStaticComponentMaskParameterValue>(Parameter))
+	{
+		// Custom widget for static component masks
+		SetPropertyRowMaskParameterWidget(PropertyRow, Parameter, ParameterValueProperty);
+	}
+	else if (UDEditorTextureParameterValue* TextureParam = Cast<UDEditorTextureParameterValue>(Parameter))
+	{
+		// Custom widget for textures
+		SetPropertyRowTextureParameterWidget(PropertyRow, Parameter, ParameterValueProperty, MaterialEditorParameters);
+	}
+	else if (UDEditorVectorParameterValue* VectorParam = Cast<UDEditorVectorParameterValue>(Parameter))
+	{
+		// Don't display custom primitive data parameters in the details panel.
+		// This data is pulled from the primitive and can't be changed on the material.
+		if (VectorParam->bUseCustomPrimitiveData)
+		{
+			return;
+		}
+
+		// Set channel names if specified
+		static const FName Red("R");
+		static const FName Green("G");
+		static const FName Blue("B");
+		static const FName Alpha("A");
+		if (!VectorParam->ChannelNames.R.IsEmpty())
+		{
+			ParameterValueProperty->GetChildHandle(Red)->SetPropertyDisplayName(VectorParam->ChannelNames.R);
+		}
+		if (!VectorParam->ChannelNames.G.IsEmpty())
+		{
+			ParameterValueProperty->GetChildHandle(Green)->SetPropertyDisplayName(VectorParam->ChannelNames.G);
+		}
+		if (!VectorParam->ChannelNames.B.IsEmpty())
+		{
+			ParameterValueProperty->GetChildHandle(Blue)->SetPropertyDisplayName(VectorParam->ChannelNames.B);
+		}
+		if (!VectorParam->ChannelNames.A.IsEmpty())
+		{
+			ParameterValueProperty->GetChildHandle(Alpha)->SetPropertyDisplayName(VectorParam->ChannelNames.A);
+		}
+
+		// Custom widget for channel masks
+		if (VectorParam->bIsUsedAsChannelMask)
+		{
+			SetPropertyRowVectorChannelMaskParameterWidget(PropertyRow, Parameter, ParameterValueProperty, MaterialEditorParameters);
+		}
+		else
+		{
+			bUseDefaultWidget = true;
+		}
+	}
+	else if (UDEditorScalarParameterValue* ScalarParam = Cast<UDEditorScalarParameterValue>(Parameter))
+	{
+		// Don't display custom primitive data parameters in the details panel.
+		// This data is pulled from the primitive and can't be changed on the material.
+		if (ScalarParam->bUseCustomPrimitiveData)
+		{
+			return;
+		}
+
+		// Atlas position params handled separately
+		if (ScalarParam->AtlasData.bIsUsedAsAtlasPosition)
+		{
+			SetPropertyRowScalarAtlasPositionParameterWidget(PropertyRow, Parameter, ParameterValueProperty, MaterialEditorParameters);
+		}
+		else
+		{
+			// Configure slider if specified
+			if (ScalarParam->SliderMax > ScalarParam->SliderMin)
+			{
+				ParameterValueProperty->SetInstanceMetaData("UIMin", FString::Printf(TEXT("%f"), ScalarParam->SliderMin));
+				ParameterValueProperty->SetInstanceMetaData("UIMax", FString::Printf(TEXT("%f"), ScalarParam->SliderMax));
+			}
+
+			// Custom widget for enums
+			TSoftObjectPtr<UObject> Enumeration;
+			if (!ScalarParam->Enumeration.IsNull())
+			{
+				Enumeration = ScalarParam->Enumeration;
+			}
+			else if (ScalarParam->EnumerationIndex >= 0)
+			{
+				if (UMaterialEditorInstanceConstant* MaterialEditorInstance = Cast<UMaterialEditorInstanceConstant>(MaterialEditorParameters))
+				{
+					Enumeration = MaterialEditorInstance->GetEnumerationObject(ScalarParam->EnumerationIndex);
+				}
+			}
+
+			if (!Enumeration.IsNull())
+			{
+				SetPropertyRowScalarEnumerationParameterWidget(PropertyRow, Parameter, ParameterValueProperty, Enumeration);
+			}
+			else
+			{
+				bUseDefaultWidget = true;
+			}
+		}
+	}
+	else
+	{
+		bUseDefaultWidget = true;
+	}
+
+	// Fix up property name if not using a custom widget. Otherwise, the name will display as "Parameter Value"
+	if (bUseDefaultWidget)
+	{
+		if (FDetailWidgetDecl* NameWidgetPtr = PropertyRow.CustomNameWidget())
+		{
+			// Already customized, just set the name widget
+			(*NameWidgetPtr)
+			[
+				ParameterValueProperty->CreatePropertyNameWidget()
+			];
+		}
+		else
+		{
+			// Not customized, grab default widgets and customize
+			TSharedPtr<SWidget> NameWidget;
+			TSharedPtr<SWidget> ValueWidget;
+			PropertyRow.GetDefaultWidgets(NameWidget, ValueWidget);
+
+			PropertyRow.CustomWidget(true)
+				.NameContent()
+				[
+					ParameterValueProperty->CreatePropertyNameWidget()
+				]
+				.ValueContent()
+				[
+					ValueWidget.ToSharedRef()
+				];
+		}
+	}
+}
+
+void FMaterialPropertyHelpers::SetPropertyRowMaskParameterWidget(IDetailPropertyRow& PropertyRow, UDEditorParameterValue* Parameter, TSharedPtr<IPropertyHandle> ParameterValueProperty)
+{
+	TSharedPtr<SHorizontalBox> ValueHorizontalBox;
+
+	// Set up the custom widget
+	PropertyRow.CustomWidget()
+		.NameContent()
+		[
+			ParameterValueProperty->CreatePropertyNameWidget()
+		]
+		.ValueContent()
+		.MaxDesiredWidth(200.0f)
+		[
+			SNew(SHorizontalBox)
+			+ SHorizontalBox::Slot()
+			.FillWidth(1.0f)
+			[
+				SAssignNew(ValueHorizontalBox, SHorizontalBox)
+			]
+		];
+
+	// Add checkboxes for each mask channel
+	auto AddChannelWidgets =
+		[&ValueHorizontalBox](TSharedPtr<IPropertyHandle> MaskProperty)
+		{
+			FMargin NameWidgetPadding = FMargin(10.0f, 0.0f, 2.0f, 0.0f);
+			if (ValueHorizontalBox->NumSlots() == 0)
+			{
+				// No left padding on the first widget
+				NameWidgetPadding.Left = 0.0f;
+			}
+
+			ValueHorizontalBox->AddSlot()
+				.HAlign(HAlign_Left)
+				.Padding(NameWidgetPadding)
+				.AutoWidth()
+				[
+					MaskProperty->CreatePropertyNameWidget()
+				];
+
+			ValueHorizontalBox->AddSlot()
+				.HAlign(HAlign_Left)
+				.AutoWidth()
+				[
+					MaskProperty->CreatePropertyValueWidget()
+				];
+		};
+
+	AddChannelWidgets(ParameterValueProperty->GetChildHandle("R"));
+	AddChannelWidgets(ParameterValueProperty->GetChildHandle("G"));
+	AddChannelWidgets(ParameterValueProperty->GetChildHandle("B"));
+	AddChannelWidgets(ParameterValueProperty->GetChildHandle("A"));
+}
+
+void FMaterialPropertyHelpers::SetPropertyRowVectorChannelMaskParameterWidget(IDetailPropertyRow& PropertyRow, UDEditorParameterValue* Parameter, TSharedPtr<IPropertyHandle> ParameterValueProperty, UMaterialEditorParameters* MaterialEditorParameters)
+{
+	// Combo box hooks for converting between our "enum" and colors
+	FOnGetPropertyComboBoxStrings GetMaskStrings = FOnGetPropertyComboBoxStrings::CreateStatic(&FMaterialPropertyHelpers::GetVectorChannelMaskComboBoxStrings);
+	FOnGetPropertyComboBoxValue GetMaskValue = FOnGetPropertyComboBoxValue::CreateStatic(&FMaterialPropertyHelpers::GetVectorChannelMaskValue, Parameter);
+	FOnPropertyComboBoxValueSelected SetMaskValue = FOnPropertyComboBoxValueSelected::CreateStatic(&FMaterialPropertyHelpers::SetVectorChannelMaskValue, ParameterValueProperty, Parameter, (UObject*)MaterialEditorParameters);
+
+	// Widget replaces color picker with combo box
+	PropertyRow.CustomWidget()
+		.NameContent()
+		[
+			ParameterValueProperty->CreatePropertyNameWidget()
+		]
+		.ValueContent()
+		.MaxDesiredWidth(200.0f)
+		[
+			SNew(SHorizontalBox)
+			+ SHorizontalBox::Slot()
+			.FillWidth(1.0f)
+			[
+				SNew(SHorizontalBox)
+				+ SHorizontalBox::Slot()
+				.HAlign(HAlign_Left)
+				.AutoWidth()
+				[
+					PropertyCustomizationHelpers::MakePropertyComboBox(ParameterValueProperty, GetMaskStrings, GetMaskValue, SetMaskValue)
+				]
+			]
+		];
+}
+
+void FMaterialPropertyHelpers::SetPropertyRowScalarAtlasPositionParameterWidget(IDetailPropertyRow& PropertyRow, UDEditorParameterValue* Parameter, TSharedPtr<IPropertyHandle> ParameterValueProperty, UMaterialEditorParameters* MaterialEditorParameters)
+{
+	const FText ParameterName = FText::FromName(Parameter->ParameterInfo.Name);
+	UDEditorScalarParameterValue* AtlasParameter = Cast<UDEditorScalarParameterValue>(Parameter);
+
+	// Create the value widget first for GetDesiredWidth
+	TSharedRef<SObjectPropertyEntryBox> ObjectPropertyEntryBox = SNew(SObjectPropertyEntryBox)
+		.ObjectPath_UObject(AtlasParameter, &UDEditorScalarParameterValue::GetAtlasCurvePath)
+		.AllowedClass(UCurveLinearColor::StaticClass())
+		.NewAssetFactories(TArray<UFactory*>())
+		.DisplayThumbnail(true)
+		.ThumbnailPool(UThumbnailManager::Get().GetSharedThumbnailPool())
+		.OnShouldFilterAsset(FOnShouldFilterAsset::CreateStatic(&FMaterialPropertyHelpers::OnShouldFilterCurveAsset, AtlasParameter->AtlasData.Atlas))
+		.OnShouldSetAsset(FOnShouldSetAsset::CreateStatic(&FMaterialPropertyHelpers::OnShouldSetCurveAsset, AtlasParameter->AtlasData.Atlas))
+		.OnObjectChanged(FOnSetObject::CreateStatic(&FMaterialPropertyHelpers::SetPositionFromCurveAsset, AtlasParameter->AtlasData.Atlas, AtlasParameter, ParameterValueProperty, (UObject*)MaterialEditorParameters))
+		.DisplayCompactSize(true);
+
+	float MinDesiredWidth, MaxDesiredWidth;
+	ObjectPropertyEntryBox->GetDesiredWidth(MinDesiredWidth, MaxDesiredWidth);
+
+	// Set up the custom widget
+	PropertyRow.CustomWidget()
+		.NameContent()
+		[
+			ParameterValueProperty->CreatePropertyNameWidget()
+		]
+		.ValueContent()
+		.MinDesiredWidth(MinDesiredWidth)
+		.MaxDesiredWidth(MaxDesiredWidth)
+		[
+			ObjectPropertyEntryBox
+		];
+}
+
+void FMaterialPropertyHelpers::SetPropertyRowScalarEnumerationParameterWidget(IDetailPropertyRow& PropertyRow, UDEditorParameterValue* Parameter, TSharedPtr<IPropertyHandle> ParameterValueProperty, const TSoftObjectPtr<UObject>& SoftEnumeration)
+{
+	if (SoftEnumeration.IsNull())
+	{
+		return;
+	}
+
+	// Scalar parameters can present an enumeration combo box.
+	PropertyRow.CustomWidget()
+		.NameContent()
+		[
+			ParameterValueProperty->CreatePropertyNameWidget()
+		]
+		.ValueContent()
+		[
+			FMaterialPropertyHelpers::CreateScalarEnumerationParameterValueWidget(SoftEnumeration, ParameterValueProperty)
+		];
+}
+
+void FMaterialPropertyHelpers::SetPropertyRowTextureParameterWidget(IDetailPropertyRow& PropertyRow, UDEditorParameterValue* Parameter, TSharedPtr<IPropertyHandle> ParameterValueProperty, UMaterialEditorParameters* MaterialEditorParameters)
+{
+	TWeakObjectPtr<UMaterialExpressionTextureSampleParameter> SamplerExpression;
+
+	// Cache the SamplerExpression for asset filtering below
+	UDEditorTextureParameterValue* TextureParam = Cast<UDEditorTextureParameterValue>(Parameter);
+	if (TextureParam)
+	{
+		UMaterial* Material = Cast<UMaterial>(MaterialEditorParameters->GetMaterialInterface());
+		if (Material != nullptr)
+		{
+			UMaterialExpressionTextureSampleParameter* Expression = Material->FindExpressionByGUID<UMaterialExpressionTextureSampleParameter>(TextureParam->ExpressionId);
+			if (Expression != nullptr)
+			{
+				SamplerExpression = Expression;
+			}
+		}
+	}
+
+	// Create the value widget first for GetDesiredWidth
+	TSharedRef<SObjectPropertyEntryBox> ObjectPropertyEntryBox = SNew(SObjectPropertyEntryBox)
+		.PropertyHandle(ParameterValueProperty)
+		.AllowedClass(UTexture::StaticClass())
+		.ThumbnailPool(UThumbnailManager::Get().GetSharedThumbnailPool())
+		.OnShouldFilterAsset_Lambda([SamplerExpression](const FAssetData& AssetData)
+			{
+				if (SamplerExpression.Get())
+				{
+					bool VirtualTextured = false;
+					AssetData.GetTagValue<bool>("VirtualTextureStreaming", VirtualTextured);
+
+					bool ExpressionIsVirtualTextured = IsVirtualSamplerType(SamplerExpression->SamplerType);
+
+					return VirtualTextured != ExpressionIsVirtualTextured;
+				}
+				else
+				{
+					return false;
+				}
+			});
+
+	float MinDesiredWidth, MaxDesiredWidth;
+	ObjectPropertyEntryBox->GetDesiredWidth(MinDesiredWidth, MaxDesiredWidth);
+
+	TSharedPtr<SVerticalBox> NameVerticalBox;
+
+	// Set up the custom widget
+	PropertyRow.CustomWidget()
+		.NameContent()
+		[
+			SAssignNew(NameVerticalBox, SVerticalBox)
+				+ SVerticalBox::Slot()
+				.AutoHeight()
+				[
+					ParameterValueProperty->CreatePropertyNameWidget()
+				]
+		]
+		.ValueContent()
+		.MinDesiredWidth(MinDesiredWidth)
+		.MaxDesiredWidth(MaxDesiredWidth)
+		[
+			ObjectPropertyEntryBox
+		];
+
+	// Add extra widgets for specified channel names, if any
+	auto AddChannelNameWidgets =
+		[&NameVerticalBox](const FText& ChannelName, FName ChannelKey)
+		{
+			if (ChannelName.IsEmpty())
+			{
+				return;
+			}
+
+			NameVerticalBox->AddSlot()
+			[
+				SNew(SHorizontalBox)
+					+ SHorizontalBox::Slot()
+					.AutoWidth()
+					.Padding(20.0, 2.0, 4.0, 2.0)
+					[
+						SNew(STextBlock)
+							.Text(FText::FromName(ChannelKey))
+							.Font(FAppStyle::GetFontStyle(TEXT("PropertyWindow.BoldFont")))
+					]
+					+ SHorizontalBox::Slot()
+					.HAlign(HAlign_Left)
+					.Padding(4.0, 2.0)
+					[
+						SNew(STextBlock)
+							.Text(ChannelName)
+							.Font(FAppStyle::GetFontStyle(TEXT("PropertyWindow.NormalFont")))
+					]
+			];
+		};
+
+	static const FName Red("R");
+	static const FName Green("G");
+	static const FName Blue("B");
+	static const FName Alpha("A");
+	AddChannelNameWidgets(TextureParam->ChannelNames.R, Red);
+	AddChannelNameWidgets(TextureParam->ChannelNames.G, Green);
+	AddChannelNameWidgets(TextureParam->ChannelNames.B, Blue);
+	AddChannelNameWidgets(TextureParam->ChannelNames.A, Alpha);
+}
+
+void FMaterialPropertyHelpers::SetPropertyRowParameterCollectionParameterWidget(IDetailPropertyRow& PropertyRow, UDEditorParameterValue* Parameter, TSharedPtr<IPropertyHandle> ParameterValueProperty, class UMaterialParameterCollection* BaseParameterCollection, TArray<TSharedRef<IPropertyHandle>>&& AllPropertyHandles)
+{
+	// Display the parameter collection's name on the widget.
+	ParameterValueProperty->SetPropertyDisplayName(FText::FromName(BaseParameterCollection->GetFName()));
+
+	// Create a delegate to invoke when the property value changes.
+	TDelegate<void(const FPropertyChangedEvent&)> OnPropertyValueChangedDelegate = TDelegate<void(const FPropertyChangedEvent&)>::CreateWeakLambda(BaseParameterCollection,
+		[BaseParameterCollection, PropertyHandles = MoveTemp(AllPropertyHandles)](const FPropertyChangedEvent& PropertyChangedEvent)
+		{
+			// Determine if this is the final setting of the property value.
+			if (PropertyChangedEvent.ChangeType != EPropertyChangeType::ValueSet)
+			{
+				return;
+			}
+
+			if (PropertyHandles.IsEmpty())
+			{
+				return;
+			}
+
+			// Get the new value from the first property handle.
+			UObject* Value = nullptr;
+			if (PropertyHandles[0]->GetValue(Value) != FPropertyAccess::Success)
+			{
+				return;
+			}
+
+			// Verify that the new value is valid.
+			if (auto* Collection = Cast<UMaterialParameterCollection>(Value))
+			{
+				if (Collection->IsInstanceOf(BaseParameterCollection->StateId))
+				{
+					// Set the same value on all of the collection's referencing parameters.
+					for (int32 PropertyHandleIndex = 1; PropertyHandleIndex < PropertyHandles.Num(); ++PropertyHandleIndex)
+					{
+						PropertyHandles[PropertyHandleIndex]->SetValue(Value);
+					}
+
+					return;
+				}
+			}
+
+			// Reset the value to the default parameter collection.
+			Value = BaseParameterCollection;
+			PropertyHandles[0]->SetValue(Value);
+		});
+
+	// Invoke the delegate when the property value or one of its child property values changes.
+	ParameterValueProperty->SetOnPropertyValueChangedWithData(OnPropertyValueChangedDelegate);
+	ParameterValueProperty->SetOnChildPropertyValueChangedWithData(OnPropertyValueChangedDelegate);
+
+	FSoftObjectPath BaseParameterCollectionPath = BaseParameterCollection;
+
+	// Create the value widget first for GetDesiredWidth
+	TSharedRef<SObjectPropertyEntryBox> ObjectPropertyEntryBox = SNew(SObjectPropertyEntryBox)
+		.PropertyHandle(ParameterValueProperty)
+		.AllowedClass(UMaterialParameterCollection::StaticClass())
+		.ThumbnailPool(UThumbnailManager::Get().GetSharedThumbnailPool())
+		.OnShouldFilterAsset_Lambda([BaseParameterCollectionPath](const FAssetData& AssetData)
+			{
+				// Allow the base MPC itself
+				if (AssetData.GetSoftObjectPath() == BaseParameterCollectionPath)
+				{
+					return false;
+				}
+
+				IAssetRegistry& AssetRegistry = IAssetRegistry::GetChecked();
+
+				// Traverse up the hierarchy of UMaterialParameterCollection::Base until we find the base material's MPC
+				FSoftObjectPath ParentAssetPath = FSoftObjectPath(AssetData.GetTagValueRef<FString>(UMaterialParameterCollection::GetBaseParameterCollectionMemberName()));
+				while (!ParentAssetPath.IsNull())
+				{
+					if (ParentAssetPath == BaseParameterCollectionPath)
+					{
+						return false;
+					}
+
+					FAssetData ParentAssetData = AssetRegistry.GetAssetByObjectPath(ParentAssetPath);
+					ParentAssetPath = FSoftObjectPath(ParentAssetData.GetTagValueRef<FString>(UMaterialParameterCollection::GetBaseParameterCollectionMemberName()));
+				}
+
+				// No matching MPC in hierarchy, filter this MPC out
+				return true;
+			});
+
+	float MinDesiredWidth, MaxDesiredWidth;
+	ObjectPropertyEntryBox->GetDesiredWidth(MinDesiredWidth, MaxDesiredWidth);
+
+	// Set up the custom widget
+	PropertyRow.CustomWidget()
+		.NameContent()
+		[
+			ParameterValueProperty->CreatePropertyNameWidget()
+		]
+		.ValueContent()
+		.MinDesiredWidth(MinDesiredWidth)
+		.MaxDesiredWidth(MaxDesiredWidth)
+		[
+			ObjectPropertyEntryBox
+		];
+}
+
 bool FMaterialPropertyHelpers::IsOverriddenExpression(UDEditorParameterValue* Parameter)
 {
 	return Parameter && Parameter->bOverride != 0;
@@ -618,6 +1179,20 @@ FText FMaterialPropertyHelpers::GetParameterTooltip(UDEditorParameterValue* Para
 	return TooltipText;
 }
 
+FResetToDefaultOverride FMaterialPropertyHelpers::CreateResetToDefaultOverride(UDEditorParameterValue* Parameter, UMaterialEditorInstanceConstant* MaterialEditorInstance)
+{
+	UDEditorScalarParameterValue* ScalarParam = Cast<UDEditorScalarParameterValue>(Parameter);
+	if (ScalarParam && ScalarParam->AtlasData.bIsUsedAsAtlasPosition)
+	{
+		TAttribute<bool> IsResetVisible = TAttribute<bool>::Create(TAttribute<bool>::FGetter::CreateStatic(&FMaterialPropertyHelpers::ShouldShowResetToDefault, Parameter, MaterialEditorInstance));
+		FSimpleDelegate ResetHandler = FSimpleDelegate::CreateStatic(&FMaterialPropertyHelpers::ResetCurveToDefault, Parameter, MaterialEditorInstance);
+		return FResetToDefaultOverride::Create(IsResetVisible, ResetHandler);
+	}
+
+	TAttribute<bool> IsResetVisible = TAttribute<bool>::Create(TAttribute<bool>::FGetter::CreateStatic(&FMaterialPropertyHelpers::ShouldShowResetToDefault, Parameter, MaterialEditorInstance));
+	FSimpleDelegate ResetHandler = FSimpleDelegate::CreateStatic(&FMaterialPropertyHelpers::ResetToDefault, Parameter, MaterialEditorInstance);
+	return FResetToDefaultOverride::Create(IsResetVisible, ResetHandler);
+}
 
 void FMaterialPropertyHelpers::ResetToDefault(class UDEditorParameterValue* Parameter, UMaterialEditorInstanceConstant* MaterialEditorInstance)
 {
diff --git a/Engine/Source/Editor/MaterialEditor/Private/SMaterialLayersFunctionsTree.cpp b/Engine/Source/Editor/MaterialEditor/Private/SMaterialLayersFunctionsTree.cpp
--- a/Engine/Source/Editor/MaterialEditor/Private/SMaterialLayersFunctionsTree.cpp
+++ b/Engine/Source/Editor/MaterialEditor/Private/SMaterialLayersFunctionsTree.cpp
@@ -2,14 +2,7 @@
 
 #include "SMaterialLayersFunctionsTree.h"
 #include "SMaterialSubstrateTree.h"
-#include "MaterialEditor/DEditorFontParameterValue.h"
 #include "MaterialEditor/DEditorMaterialLayersParameterValue.h"
-#include "MaterialEditor/DEditorRuntimeVirtualTextureParameterValue.h"
-#include "MaterialEditor/DEditorScalarParameterValue.h"
-#include "MaterialEditor/DEditorStaticComponentMaskParameterValue.h"
-#include "MaterialEditor/DEditorStaticSwitchParameterValue.h"
-#include "MaterialEditor/DEditorTextureParameterValue.h"
-#include "MaterialEditor/DEditorVectorParameterValue.h"
 #include "MaterialEditor/MaterialEditorInstanceConstant.h"
 #include "Materials/Material.h"
 #include "PropertyHandle.h"
@@ -32,32 +25,21 @@
 #include "Widgets/Images/SImage.h"
 #include "MaterialEditor/MaterialEditorPreviewParameters.h"
 #include "Widgets/Input/SButton.h"
-#include "Widgets/Text/SInlineEditableTextBlock.h"
 #include "Materials/MaterialFunctionInstance.h"
-#include "Materials/MaterialExpressionTextureSampleParameter.h"
 #include "Framework/Application/SlateApplication.h"
 #include "Framework/MultiBox/MultiBoxBuilder.h"
 #include "Widgets/Input/SEditableTextBox.h"
 #include "Widgets/Input/SComboButton.h"
-#include "Curves/CurveLinearColor.h"
-#include "Curves/CurveLinearColorAtlas.h"
 #include "Subsystems/AssetEditorSubsystem.h"
 #include "Editor.h"
 #include "EditorFontGlyphs.h"
 #include "SResetToDefaultPropertyEditor.h"
 #include "ThumbnailRendering/ThumbnailManager.h"
-#include "Styling/AppStyle.h"
 #include "Styling/StyleColors.h"
 #include "RenderUtils.h"
 
 #define LOCTEXT_NAMESPACE "MaterialLayerCustomization"
 
-FString SMaterialLayersFunctionsInstanceTreeItem::GetCurvePath(UDEditorScalarParameterValue* Parameter) const
-{
-	FString Path = Parameter->AtlasData.Curve->GetPathName();
-	return Path;
-}
-
 const FSlateBrush* SMaterialLayersFunctionsInstanceTreeItem::GetBorderImage() const
 {
 	return FAppStyle::Get().GetBrush("DetailsView.CategoryMiddle");
@@ -251,6 +233,8 @@ void SMaterialLayersFunctionsInstanceTreeItem::Construct(const FArguments& InArg
 	FText NameOverride;
 	TSharedRef<SVerticalBox> WrapperWidget = SNew(SVerticalBox);
 	EHorizontalAlignment ValueAlignment = HAlign_Left;
+	FMargin RightSidePadding;
+
 // STACK --------------------------------------------------
 	if (StackParameterData->StackDataType == EStackDataType::Stack)
 	{
@@ -553,404 +537,28 @@ void SMaterialLayersFunctionsInstanceTreeItem::Construct(const FArguments& InArg
 // END ASSET
 
 // PROPERTY ----------------------------------------------
-	bool bisPaddedProperty = false;
 	if (StackParameterData->StackDataType == EStackDataType::Property)
 	{
-
-		UDEditorStaticComponentMaskParameterValue* CompMaskParam = Cast<UDEditorStaticComponentMaskParameterValue>(StackParameterData->Parameter);
-		UDEditorVectorParameterValue* VectorParam = Cast<UDEditorVectorParameterValue>(StackParameterData->Parameter);
-		UDEditorScalarParameterValue* ScalarParam = Cast<UDEditorScalarParameterValue>(StackParameterData->Parameter);
-		UDEditorTextureParameterValue* TextureParam = Cast<UDEditorTextureParameterValue>(StackParameterData->Parameter);
-
 		const int32 LayerStateIndex = StackParameterData->ParameterInfo.Association == EMaterialParameterAssociation::BlendParameter ? StackParameterData->ParameterInfo.Index + 1 : StackParameterData->ParameterInfo.Index;
+		WrapperWidget->SetVisibility(TAttribute<EVisibility>::Create(TAttribute<EVisibility>::FGetter::CreateStatic(&FMaterialPropertyHelpers::ShouldShowExpression, StackParameterData->Parameter, MaterialEditorInstance, Tree->GetShowHiddenDelegate())));
 
-		TAttribute<bool> IsParamEnabled = TAttribute<bool>::Create(TAttribute<bool>::FGetter::CreateSP(Tree, &SMaterialLayersFunctionsInstanceTree::IsOverriddenExpression, StackParameterData->Parameter, LayerStateIndex));
-		NameOverride = FText::FromName(StackParameterData->Parameter->ParameterInfo.Name);
-		TAttribute<bool> IsResetVisible = TAttribute<bool>::Create(TAttribute<bool>::FGetter::CreateStatic(&FMaterialPropertyHelpers::ShouldShowResetToDefault, StackParameterData->Parameter, MaterialEditorInstance));
-		FSimpleDelegate ResetHandler = FSimpleDelegate::CreateStatic(&FMaterialPropertyHelpers::ResetToDefault, StackParameterData->Parameter, MaterialEditorInstance);
-		FResetToDefaultOverride ResetOverride = FResetToDefaultOverride::Create(IsResetVisible, ResetHandler);
-		
-		if (ScalarParam && ScalarParam->AtlasData.bIsUsedAsAtlasPosition)
-		{
-			IsResetVisible = TAttribute<bool>::Create(TAttribute<bool>::FGetter::CreateStatic(&FMaterialPropertyHelpers::ShouldShowResetToDefault, StackParameterData->Parameter, MaterialEditorInstance));
-			ResetHandler = FSimpleDelegate::CreateStatic(&FMaterialPropertyHelpers::ResetCurveToDefault, StackParameterData->Parameter, MaterialEditorInstance);
-			ResetOverride = FResetToDefaultOverride::Create(IsResetVisible, ResetHandler);
-		}
+		// Handle our own EditCondition, force Visibility to "Visible" since WrapperWidget controls our actual visibility
+		FMaterialParameterPropertyRowOverrides PropertyRowOverrides;
+		PropertyRowOverrides.EditConditionValue = TAttribute<bool>::Create(TAttribute<bool>::FGetter::CreateSP(Tree, &SMaterialLayersFunctionsInstanceTree::IsOverriddenExpression, StackParameterData->Parameter, LayerStateIndex));
+		PropertyRowOverrides.OnEditConditionValueChanged = FOnBooleanValueChanged::CreateSP(this, &SMaterialLayersFunctionsInstanceTreeItem::OnOverrideParameter, StackParameterData->Parameter);
+		PropertyRowOverrides.Visibility = TAttribute<EVisibility>(EVisibility::Visible);
 
 		IDetailTreeNode& Node = *StackParameterData->ParameterNode;
+		TSharedPtr<IPropertyHandle> PropertyHandle = Node.CreatePropertyHandle();
 		TSharedPtr<IDetailPropertyRow> GeneratedRow = StaticCastSharedPtr<IDetailPropertyRow>(Node.GetRow());
 		IDetailPropertyRow& Row = *GeneratedRow.Get();
-		Row
-			.DisplayName(NameOverride)
-			.OverrideResetToDefault(ResetOverride)
-			.EditCondition(IsParamEnabled, FOnBooleanValueChanged::CreateSP(this, &SMaterialLayersFunctionsInstanceTreeItem::OnOverrideParameter, StackParameterData->Parameter));
-
-		WrapperWidget->SetVisibility(TAttribute<EVisibility>::Create(TAttribute<EVisibility>::FGetter::CreateStatic(&FMaterialPropertyHelpers::ShouldShowExpression, StackParameterData->Parameter, MaterialEditorInstance, Tree->GetShowHiddenDelegate())));
-
-		if (VectorParam && VectorParam->bIsUsedAsChannelMask)
-		{
-			FOnGetPropertyComboBoxStrings GetMaskStrings = FOnGetPropertyComboBoxStrings::CreateStatic(&FMaterialPropertyHelpers::GetVectorChannelMaskComboBoxStrings);
-			FOnGetPropertyComboBoxValue GetMaskValue = FOnGetPropertyComboBoxValue::CreateStatic(&FMaterialPropertyHelpers::GetVectorChannelMaskValue, StackParameterData->Parameter);
-			FOnPropertyComboBoxValueSelected SetMaskValue = FOnPropertyComboBoxValueSelected::CreateStatic(&FMaterialPropertyHelpers::SetVectorChannelMaskValue, StackParameterData->ParameterNode->CreatePropertyHandle(), StackParameterData->Parameter, (UObject*)MaterialEditorInstance);
-
-			FDetailWidgetRow& CustomWidget = Row.CustomWidget();
-			CustomWidget
-			.FilterString(NameOverride)
-			.NameContent()
-			[
-				SNew(STextBlock)
-				.Text(NameOverride)
-				.ToolTipText(FMaterialPropertyHelpers::GetParameterExpressionDescription(StackParameterData->Parameter, MaterialEditorInstance))
-				.Font(FAppStyle::GetFontStyle(TEXT("PropertyWindow.NormalFont")))
-			]
-			.ValueContent()
-			.MaxDesiredWidth(200.0f)
-			[
-				SNew(SHorizontalBox)
-				+ SHorizontalBox::Slot()
-				.FillWidth(1.0f)
-				[
-					SNew(SHorizontalBox)
-					+ SHorizontalBox::Slot()
-					.AutoWidth()
-					[
-						PropertyCustomizationHelpers::MakePropertyComboBox(StackParameterData->ParameterNode->CreatePropertyHandle(), GetMaskStrings, GetMaskValue, SetMaskValue)
-					]
-				]
-			];
-		}
-		else if (ScalarParam)
-		{
-			if (ScalarParam->AtlasData.bIsUsedAsAtlasPosition)
-			{
-				const FText ParameterName = FText::FromName(StackParameterData->Parameter->ParameterInfo.Name);
 
-				FDetailWidgetRow& CustomWidget = Row.CustomWidget();
-				CustomWidget
-					.FilterString(ParameterName)
-					.NameContent()
-					[
-						SNew(STextBlock)
-						.Text(ParameterName)
-						.ToolTipText(FMaterialPropertyHelpers::GetParameterExpressionDescription(StackParameterData->Parameter, MaterialEditorInstance))
-						.Font(FAppStyle::GetFontStyle(TEXT("PropertyWindow.NormalFont")))
-					]
-					.ValueContent()
-					.HAlign(HAlign_Fill)
-					.MaxDesiredWidth(400.0f)
-					[
-						SNew(SObjectPropertyEntryBox)
-						.ObjectPath(this, &SMaterialLayersFunctionsInstanceTreeItem::GetCurvePath, ScalarParam)
-						.AllowedClass(UCurveLinearColor::StaticClass())
-						.NewAssetFactories(TArray<UFactory*>())
-						.DisplayThumbnail(true)
-						.ThumbnailPool(InArgs._InTree->GetTreeThumbnailPool())
-						.OnShouldSetAsset(FOnShouldSetAsset::CreateStatic(&FMaterialPropertyHelpers::OnShouldSetCurveAsset, ScalarParam->AtlasData.Atlas))
-						.OnShouldFilterAsset(FOnShouldFilterAsset::CreateStatic(&FMaterialPropertyHelpers::OnShouldFilterCurveAsset, ScalarParam->AtlasData.Atlas))
-						.OnObjectChanged(FOnSetObject::CreateStatic(&FMaterialPropertyHelpers::SetPositionFromCurveAsset, ScalarParam->AtlasData.Atlas, ScalarParam, StackParameterData->ParameterHandle, (UObject*)MaterialEditorInstance))
-						.DisplayCompactSize(true)
-					];
-			}
-			else if (!ScalarParam->bUseCustomPrimitiveData)
-			{
-				// Try to find a valid enumeration
-				TSoftObjectPtr<UObject> Enumeration;
-				if (!ScalarParam->Enumeration.IsNull())
-				{
-					Enumeration = ScalarParam->Enumeration;
-				}
-				else if (ScalarParam->EnumerationIndex >= 0)
-				{
-					Enumeration = MaterialEditorInstance->SourceInstance->GetEnumerationObject(ScalarParam->EnumerationIndex);
-				}
-				// If found a valid Enumeration, customize the widget, else fallback todefault scalar numeric edit box
-				if (!Enumeration.IsNull())
-				{
-					const FText ParameterName = FText::FromName(StackParameterData->Parameter->ParameterInfo.Name);
-					TSharedPtr<IPropertyHandle> ParameterValueProperty = StackParameterData->ParameterNode->CreatePropertyHandle();
-
-					Row.CustomWidget()
-						.FilterString(ParameterName)
-						.NameContent()
-						[
-							ParameterValueProperty->CreatePropertyNameWidget()
-						]
-						.ValueContent()
-						.HAlign(HAlign_Fill)
-						.MaxDesiredWidth(400.0f)
-						[
-							FMaterialPropertyHelpers::CreateScalarEnumerationParameterValueWidget(Enumeration, ParameterValueProperty)
-						];
-				}
-			}
-		}
-		else if (TextureParam)
-		{
-			UMaterial *Material = MaterialEditorInstance->SourceInstance->GetMaterial();
-			if (Material != nullptr)
-			{
-				UMaterialExpressionTextureSampleParameter* Expression = Material->FindExpressionByGUID<UMaterialExpressionTextureSampleParameter>(TextureParam->ExpressionId);
-				if (Expression != nullptr)
-				{
-					TWeakObjectPtr<UMaterialExpressionTextureSampleParameter> SamplerExpression = Expression;
-					TSharedPtr<SWidget> NameWidget;
-					TSharedPtr<SWidget> ValueWidget;
-					TSharedPtr<SVerticalBox> NameVerticalBox;
-					const FText ParameterName = FText::FromName(StackParameterData->Parameter->ParameterInfo.Name);
-
-					Row.OverrideResetToDefault(ResetOverride);
-
-					FDetailWidgetRow& CustomWidget = Row.CustomWidget();
-					CustomWidget
-						.FilterString(ParameterName)
-						.NameContent()
-						[
-							SAssignNew(NameVerticalBox, SVerticalBox)
-							+ SVerticalBox::Slot()
-							.AutoHeight()
-							[
-							SNew(STextBlock)
-							.Text(ParameterName)
-							.ToolTipText(FMaterialPropertyHelpers::GetParameterExpressionDescription(StackParameterData->Parameter, MaterialEditorInstance))
-							.Font(FAppStyle::GetFontStyle(TEXT("PropertyWindow.NormalFont")))
-							]
-						];
-					CustomWidget.ValueContent()
-						[
-							SNew(SObjectPropertyEntryBox)
-							.PropertyHandle(StackParameterData->ParameterNode->CreatePropertyHandle())
-							.AllowedClass(UTexture::StaticClass())
-							.ThumbnailPool(Tree->GetTreeThumbnailPool())
-							.OnShouldFilterAsset_Lambda([SamplerExpression](const FAssetData& AssetData)
-							{
-								if (SamplerExpression.Get())
-								{
-									bool VirtualTextured = false;
-									AssetData.GetTagValue<bool>("VirtualTextureStreaming", VirtualTextured);
-
-									bool ExpressionIsVirtualTextured = IsVirtualSamplerType(SamplerExpression->SamplerType);
-
-									return VirtualTextured != ExpressionIsVirtualTextured;
-								}
-								else
-								{
-									return false;
-								}
-							})
-						];
-
-					static const FName Red("R");
-					static const FName Green("G");
-					static const FName Blue("B");
-					static const FName Alpha("A");
-
-					if (!TextureParam->ChannelNames.R.IsEmpty())
-					{
-						NameVerticalBox->AddSlot()
-							[
-								SNew(SHorizontalBox)
-								+ SHorizontalBox::Slot()
-								.AutoWidth()
-								.Padding(20.0, 2.0, 4.0, 2.0)
-								[
-									SNew(STextBlock)
-									.Text(FText::FromName(Red))
-									.Font(FAppStyle::GetFontStyle(TEXT("PropertyWindow.BoldFont")))
-								]
-								+ SHorizontalBox::Slot()
-								.HAlign(HAlign_Left)
-								.Padding(4.0, 2.0)
-								[
-									SNew(STextBlock)
-									.Text(TextureParam->ChannelNames.R)
-									.Font(FAppStyle::GetFontStyle(TEXT("PropertyWindow.NormalFont")))
-								]
-							];
-					}
-					if (!TextureParam->ChannelNames.G.IsEmpty())
-					{
-						NameVerticalBox->AddSlot()
-							[
-								SNew(SHorizontalBox)
-								+ SHorizontalBox::Slot()
-								.Padding(20.0, 2.0, 4.0, 2.0)
-								.AutoWidth()
-								[
-									SNew(STextBlock)
-									.Text(FText::FromName(Green))
-									.Font(FAppStyle::GetFontStyle(TEXT("PropertyWindow.BoldFont")))
-								]
-								+ SHorizontalBox::Slot()
-								.HAlign(HAlign_Left)
-								.Padding(4.0, 2.0)
-								[
-									SNew(STextBlock)
-									.Text(TextureParam->ChannelNames.G)
-									.Font(FAppStyle::GetFontStyle(TEXT("PropertyWindow.NormalFont")))
-								]
-							];
-					}
-					if (!TextureParam->ChannelNames.B.IsEmpty())
-					{
-						NameVerticalBox->AddSlot()
-							[
-								SNew(SHorizontalBox)
-								+ SHorizontalBox::Slot()
-								.Padding(20.0, 2.0, 4.0, 2.0)
-								.AutoWidth()
-								[
-									SNew(STextBlock)
-									.Text(FText::FromName(Blue))
-									.Font(FAppStyle::GetFontStyle(TEXT("PropertyWindow.BoldFont")))
-								]
-								+ SHorizontalBox::Slot()
-								.HAlign(HAlign_Left)
-								.Padding(4.0, 2.0)
-								[
-									SNew(STextBlock)
-									.Text(TextureParam->ChannelNames.B)
-									.Font(FAppStyle::GetFontStyle(TEXT("PropertyWindow.NormalFont")))
-								]
-							];
-					}
-					if (!TextureParam->ChannelNames.A.IsEmpty())
-					{
-						NameVerticalBox->AddSlot()
-							[
-								SNew(SHorizontalBox)
-								+ SHorizontalBox::Slot()
-								.Padding(20.0, 2.0, 4.0, 2.0)
-								.AutoWidth()
-								[
-									SNew(STextBlock)
-									.Text(FText::FromName(Alpha))
-									.Font(FAppStyle::GetFontStyle(TEXT("PropertyWindow.BoldFont")))
-								]
-								+ SHorizontalBox::Slot()
-								.HAlign(HAlign_Left)
-								.Padding(4.0, 2.0)
-								[
-									SNew(STextBlock)
-									.Text(TextureParam->ChannelNames.A)
-									.Font(FAppStyle::GetFontStyle(TEXT("PropertyWindow.NormalFont")))
-								]
-							];
-						}
-					}
-			}
-		}
-		else if (CompMaskParam)
-		{
-			TSharedPtr<IPropertyHandle> RMaskProperty = StackParameterData->ParameterNode->CreatePropertyHandle()->GetChildHandle("R");
-			TSharedPtr<IPropertyHandle> GMaskProperty = StackParameterData->ParameterNode->CreatePropertyHandle()->GetChildHandle("G");
-			TSharedPtr<IPropertyHandle> BMaskProperty = StackParameterData->ParameterNode->CreatePropertyHandle()->GetChildHandle("B");
-			TSharedPtr<IPropertyHandle> AMaskProperty = StackParameterData->ParameterNode->CreatePropertyHandle()->GetChildHandle("A");
-			FDetailWidgetRow& CustomWidget = Row.CustomWidget();
-			CustomWidget
-			.FilterString(NameOverride)
-			.NameContent()
-			[
-				SNew(SHorizontalBox)
-				+ SHorizontalBox::Slot()
-				.VAlign(VAlign_Center)
-				[
-					SNew(STextBlock)
-					.Text(NameOverride)
-					.ToolTipText(FMaterialPropertyHelpers::GetParameterExpressionDescription(StackParameterData->Parameter, MaterialEditorInstance))
-					.Font(FAppStyle::GetFontStyle(TEXT("PropertyWindow.NormalFont")))
-				]
-			]
-			.ValueContent()
-			.MaxDesiredWidth(200.0f)
-			[
-				SNew(SHorizontalBox)
-				+ SHorizontalBox::Slot()
-				.FillWidth(1.0f)
-				[
-					SNew(SHorizontalBox)
-					+ SHorizontalBox::Slot()
-					.HAlign(HAlign_Left)
-					.AutoWidth()
-					[
-						RMaskProperty->CreatePropertyNameWidget()
-					]
-					+ SHorizontalBox::Slot()
-					.HAlign(HAlign_Left)
-					.AutoWidth()
-					[
-						RMaskProperty->CreatePropertyValueWidget()
-					]
-					+ SHorizontalBox::Slot()
-					.HAlign(HAlign_Left)
-					.Padding(FMargin(10.0f, 0.0f, 0.0f, 0.0f))
-					.AutoWidth()
-					[
-						GMaskProperty->CreatePropertyNameWidget()
-					]
-					+ SHorizontalBox::Slot()
-					.HAlign(HAlign_Left)
-					.AutoWidth()
-					[
-						GMaskProperty->CreatePropertyValueWidget()
-					]
-					+ SHorizontalBox::Slot()
-					.HAlign(HAlign_Left)
-					.Padding(FMargin(10.0f, 0.0f, 0.0f, 0.0f))
-					.AutoWidth()
-					[
-						BMaskProperty->CreatePropertyNameWidget()
-					]
-					+ SHorizontalBox::Slot()
-					.HAlign(HAlign_Left)
-					.AutoWidth()
-					[
-						BMaskProperty->CreatePropertyValueWidget()
-					]
-					+ SHorizontalBox::Slot()
-					.HAlign(HAlign_Left)
-					.Padding(FMargin(10.0f, 0.0f, 0.0f, 0.0f))
-					.AutoWidth()
-					[
-						AMaskProperty->CreatePropertyNameWidget()
-					]
-					+ SHorizontalBox::Slot()
-					.HAlign(HAlign_Left)
-					.AutoWidth()
-					[
-						AMaskProperty->CreatePropertyValueWidget()
-					]
-				]	
-			];
-		}
-		else
-		{
-			FDetailWidgetDecl* CustomNameWidget = Row.CustomNameWidget();
-			if (CustomNameWidget)
-			{
-				(*CustomNameWidget)
-				[
-					SNew(SHorizontalBox)
-					+SHorizontalBox::Slot()
-					.VAlign(VAlign_Center)
-					[
-						SNew(STextBlock)
-						.Text(NameOverride)
-						.ToolTipText(FMaterialPropertyHelpers::GetParameterExpressionDescription(StackParameterData->Parameter, MaterialEditorInstance))
-						.Font(FAppStyle::GetFontStyle(TEXT("PropertyWindow.NormalFont")))
-					]
-				];
-			}
-			else
-			{
-				Row.ToolTip(FMaterialPropertyHelpers::GetParameterExpressionDescription(StackParameterData->Parameter, MaterialEditorInstance));
-			}
-
-			bisPaddedProperty = true;
-		}
+		FMaterialPropertyHelpers::ConfigurePropertyRowForParameter(Row, StackParameterData->Parameter, PropertyHandle, MaterialEditorInstance, PropertyRowOverrides);
+		FMaterialPropertyHelpers::SetPropertyRowParameterWidget(Row, StackParameterData->Parameter, PropertyHandle, MaterialEditorInstance);
 
 		FNodeWidgets NodeWidgets = Node.CreateNodeWidgets();
 		LeftSideWidget = SNew(SHorizontalBox)
-			+SHorizontalBox::Slot()
+			+ SHorizontalBox::Slot()
 			.Padding(6.0f, 0.0f, 6.0f, 0.0f)
 			.HAlign(HAlign_Center)
 			.VAlign(VAlign_Center)
@@ -958,20 +566,22 @@ void SMaterialLayersFunctionsInstanceTreeItem::Construct(const FArguments& InArg
 			[
 				NodeWidgets.EditConditionWidget.ToSharedRef()
 			]
-			+SHorizontalBox::Slot()
+			+ SHorizontalBox::Slot()
 			[
 				NodeWidgets.NameWidget.ToSharedRef()
 			];
 		RightSideWidget = NodeWidgets.ValueWidget.ToSharedRef();
 
-		ResetWidget = SNew(SResetToDefaultPropertyEditor, Node.CreatePropertyHandle())
-			.IsEnabled(IsParamEnabled)
-			.CustomResetToDefault(ResetOverride);
+		ResetWidget = SNew(SResetToDefaultPropertyEditor, PropertyHandle)
+			.IsEnabled(PropertyRowOverrides.EditConditionValue)
+			.CustomResetToDefault(FMaterialPropertyHelpers::CreateResetToDefaultOverride(StackParameterData->Parameter, MaterialEditorInstance));
 
-		StackParameterData->ParameterNode->CreatePropertyHandle()->SetOnPropertyValueChanged(FSimpleDelegate::CreateSP(Tree, &SMaterialLayersFunctionsInstanceTree::UpdateThumbnailMaterial, StackParameterData->ParameterInfo.Association, StackParameterData->ParameterInfo.Index, false));
-		StackParameterData->ParameterNode->CreatePropertyHandle()->SetOnChildPropertyValueChanged(FSimpleDelegate::CreateSP(Tree, &SMaterialLayersFunctionsInstanceTree::UpdateThumbnailMaterial, StackParameterData->ParameterInfo.Association, StackParameterData->ParameterInfo.Index, false));
+		PropertyHandle->SetOnPropertyValueChanged(FSimpleDelegate::CreateSP(Tree, &SMaterialLayersFunctionsInstanceTree::UpdateThumbnailMaterial, StackParameterData->ParameterInfo.Association, StackParameterData->ParameterInfo.Index, false));
+		PropertyHandle->SetOnChildPropertyValueChanged(FSimpleDelegate::CreateSP(Tree, &SMaterialLayersFunctionsInstanceTree::UpdateThumbnailMaterial, StackParameterData->ParameterInfo.Association, StackParameterData->ParameterInfo.Index, false));
 
 		LeftSideWidget->SetEnabled(InArgs._InTree->FunctionInstance->EditorOnly.LayerStates[LayerStateIndex]);
+
+		RightSidePadding = FMargin(12.0f, 0.0f, 2.0f, 0.0f);	// Match DetailWidgetConstants::RightRowPadding
 	}
 // END PROPERTY
 
@@ -984,11 +594,14 @@ void SMaterialLayersFunctionsInstanceTreeItem::Construct(const FArguments& InArg
 
 		const int32 LayerStateIndex = StackParameterData->ParameterInfo.Association == EMaterialParameterAssociation::BlendParameter ? StackParameterData->ParameterInfo.Index + 1 : StackParameterData->ParameterInfo.Index;
 		LeftSideWidget->SetEnabled(InArgs._InTree->FunctionInstance->EditorOnly.LayerStates[LayerStateIndex]);
+
 		TAttribute<bool> EnabledAttribute = TAttribute<bool>::Create([this, LayerStateIndex]() -> bool
 			{
 				return FMaterialPropertyHelpers::IsOverriddenExpression(StackParameterData->Parameter) && Tree->FunctionInstance->EditorOnly.LayerStates[LayerStateIndex];
 			});
 		RightSideWidget->SetEnabled(EnabledAttribute);
+
+		RightSidePadding = FMargin(12.0f, 0.0f, 2.0f, 0.0f);	// Match DetailWidgetConstants::RightRowPadding
 	}
 // END PROPERTY CHILD
 
@@ -1090,54 +703,59 @@ void SMaterialLayersFunctionsInstanceTreeItem::Construct(const FArguments& InArg
 					.BorderImage(this, &SMaterialLayersFunctionsInstanceTreeItem::GetBorderImage)
 					.BorderBackgroundColor(this, &SMaterialLayersFunctionsInstanceTreeItem::GetOuterBackgroundColor, StackParameterData)
 					[
-						SNew(SSplitter)
-						.Style(FAppStyle::Get(), "DetailsView.Splitter")
-						.PhysicalSplitterHandleSize(1.0f)
-						.HitDetectionSplitterHandleSize(5.0f)
-						.HighlightedHandleIndex(Tree->ColumnSizeData.GetHoveredSplitterIndex())
-						+ SSplitter::Slot()
-						.Value(Tree->ColumnSizeData.GetNameColumnWidth())
-						.OnSlotResized(Tree->ColumnSizeData.GetOnNameColumnResized())
-						.Value(0.25f)
+						SNew(SBox)
+						.MinDesiredHeight(26.0f)	// Match PropertyEditorConstants::PropertyRowHeight
 						[
-							SNew(SHorizontalBox)
-							+ SHorizontalBox::Slot()
-							.AutoWidth()
-							.VAlign(VAlign_Center)
-							.Padding(FMargin(3.0f))
-							[
-								SNew(SExpanderArrow, SharedThis(this))
-							]
-							+ SHorizontalBox::Slot()
-							.Padding(FMargin(2.0f))
-							.VAlign(VAlign_Center)
+							SNew(SSplitter)
+							.Style(FAppStyle::Get(), "DetailsView.Splitter")
+							.PhysicalSplitterHandleSize(1.0f)
+							.HitDetectionSplitterHandleSize(5.0f)
+							.HighlightedHandleIndex(Tree->ColumnSizeData.GetHoveredSplitterIndex())
+							+ SSplitter::Slot()
+							.Value(Tree->ColumnSizeData.GetNameColumnWidth())
+							.OnSlotResized(Tree->ColumnSizeData.GetOnNameColumnResized())
+							.Value(0.25f)
 							[
-								LeftSideWidget
+								SNew(SHorizontalBox)
+								+ SHorizontalBox::Slot()
+								.AutoWidth()
+								.VAlign(VAlign_Center)
+								.Padding(FMargin(3.0f))
+								[
+									SNew(SExpanderArrow, SharedThis(this))
+								]
+								+ SHorizontalBox::Slot()
+								.Padding(FMargin(2.0f))
+								.VAlign(VAlign_Center)
+								[
+									LeftSideWidget
+								]
 							]
-						]
-						+ SSplitter::Slot()
-						.Value(Tree->ColumnSizeData.GetValueColumnWidth())
-						.OnSlotResized(Tree->ColumnSizeData.GetOnValueColumnResized())
-						[
-							SNew(SHorizontalBox)
-							.Clipping(EWidgetClipping::OnDemand)
-							+ SHorizontalBox::Slot()
-							.HAlign(ValueAlignment)
-							.VAlign(VAlign_Center)
+							+ SSplitter::Slot()
+							.Value(Tree->ColumnSizeData.GetValueColumnWidth())
+							.OnSlotResized(Tree->ColumnSizeData.GetOnValueColumnResized())
 							[
-								RightSideWidget
+								SNew(SHorizontalBox)
+								.Clipping(EWidgetClipping::OnDemand)
+								+ SHorizontalBox::Slot()
+								.HAlign(ValueAlignment)
+								.VAlign(VAlign_Center)
+								.Padding(RightSidePadding)
+								[
+									RightSideWidget
+								]
 							]
-						]
-						+ SSplitter::Slot()
-						.SizeRule(SSplitter::ESizeRule::SizeToContent)
-						[
-							SNew(SHorizontalBox)
-							+ SHorizontalBox::Slot()
-							.VAlign(VAlign_Center)
-							.HAlign(HAlign_Center)
-							.Padding(5.0f)
+							+ SSplitter::Slot()
+							.SizeRule(SSplitter::ESizeRule::SizeToContent)
 							[
-								ResetWidget
+								SNew(SHorizontalBox)
+								+ SHorizontalBox::Slot()
+								.VAlign(VAlign_Center)
+								.HAlign(HAlign_Center)
+								.Padding(5.0f)
+								[
+									ResetWidget
+								]
 							]
 						]
 					]
@@ -1604,15 +1222,14 @@ void SMaterialLayersFunctionsInstanceTree::CreateGroupsWidget()
 				}
 				else
 				{
-					FUnsortedParamData NonLayerProperty;
-					UDEditorScalarParameterValue* ScalarParam = Cast<UDEditorScalarParameterValue>(Parameter);
-
-					if (ScalarParam && ScalarParam->SliderMax > ScalarParam->SliderMin)
+					// Skip parameters that shouldn't have widgets here
+					if (!FMaterialPropertyHelpers::ShouldCreatePropertyRowForParameter(Parameter))
 					{
-						ParameterValueProperty->SetInstanceMetaData("UIMin", FString::Printf(TEXT("%f"), ScalarParam->SliderMin));
-						ParameterValueProperty->SetInstanceMetaData("UIMax", FString::Printf(TEXT("%f"), ScalarParam->SliderMax));
+						continue;
 					}
 
+					FUnsortedParamData NonLayerProperty;
+
 					NonLayerProperty.Parameter = Parameter;
 					NonLayerProperty.ParameterGroup = ParameterGroup;
 
@@ -1753,23 +1370,25 @@ void SMaterialLayersFunctionsInstanceTree::ShowSubParameters(TSharedPtr<FSortedP
 		if (Parameter->ParameterInfo.Index == ParentParameter->ParameterInfo.Index
 			&& Parameter->ParameterInfo.Association == ParentParameter->ParameterInfo.Association)
 		{
-			TSharedPtr<FSortedParamData> GroupProperty(new FSortedParamData());
-			GroupProperty->StackDataType = EStackDataType::Group;
-			GroupProperty->ParameterInfo.Index = Parameter->ParameterInfo.Index;
-			GroupProperty->ParameterInfo.Association = Parameter->ParameterInfo.Association;
-			GroupProperty->Group = Property.ParameterGroup;
-			GroupProperty->NodeKey = FString::FromInt(GroupProperty->ParameterInfo.Index) + FString::FromInt(GroupProperty->ParameterInfo.Association) + Property.ParameterGroup.GroupName.ToString();
+			FString GroupNodeKey = FString::FromInt(Parameter->ParameterInfo.Index) + FString::FromInt(Parameter->ParameterInfo.Association) + Property.ParameterGroup.GroupName.ToString();
 
 			bool bAddNewGroup = true;
 			for (TSharedPtr<struct FSortedParamData> GroupChild : ParentParameter->Children)
 			{
-				if (GroupChild->NodeKey == GroupProperty->NodeKey)
+				if (GroupChild->NodeKey == GroupNodeKey)
 				{
 					bAddNewGroup = false;
 				}
 			}
 			if (bAddNewGroup)
 			{
+				TSharedPtr<FSortedParamData> GroupProperty(new FSortedParamData());
+				GroupProperty->StackDataType = EStackDataType::Group;
+				GroupProperty->ParameterInfo.Index = Parameter->ParameterInfo.Index;
+				GroupProperty->ParameterInfo.Association = Parameter->ParameterInfo.Association;
+				GroupProperty->Group = Property.ParameterGroup;
+				GroupProperty->NodeKey = GroupNodeKey;
+
 				ParentParameter->Children.Add(GroupProperty);
 			}
 
@@ -1783,8 +1402,7 @@ void SMaterialLayersFunctionsInstanceTree::ShowSubParameters(TSharedPtr<FSortedP
 			ChildProperty->NodeKey = FString::FromInt(ChildProperty->ParameterInfo.Index) + FString::FromInt(ChildProperty->ParameterInfo.Association) +  Property.ParameterGroup.GroupName.ToString() + Property.UnsortedName.ToString();
 
 
-			UDEditorStaticComponentMaskParameterValue* CompMaskParam = Cast<UDEditorStaticComponentMaskParameterValue>(Parameter);
-			if (!CompMaskParam)
+			if (FMaterialPropertyHelpers::ShouldCreateChildPropertiesForParameter(Parameter))
 			{
 				TArray<TSharedRef<IDetailTreeNode>> ParamChildren;
 				Property.ParameterNode->GetChildren(ParamChildren);
@@ -2099,12 +1717,6 @@ TSharedPtr<IPropertyRowGenerator> SMaterialLayersFunctionsInstanceWrapper::GetGe
 /////////////// MATERIAL VERSION
 #define LOCTEXT_NAMESPACE "MaterialLayerDisplay"
 
-FString SMaterialLayersFunctionsMaterialTreeItem::GetCurvePath(UDEditorScalarParameterValue* Parameter) const
-{
-	FString Path = Parameter->AtlasData.Curve->GetPathName();
-	return Path;
-}
-
 const FSlateBrush* SMaterialLayersFunctionsMaterialTreeItem::GetBorderImage() const
 {
 	return FAppStyle::Get().GetBrush("DetailsView.CategoryMiddle");
@@ -2340,354 +1952,16 @@ void SMaterialLayersFunctionsMaterialTreeItem::Construct(const FArguments& InArg
 	// PROPERTY ----------------------------------------------
 	if (StackParameterData->StackDataType == EStackDataType::Property)
 	{
-
-		UDEditorStaticComponentMaskParameterValue* CompMaskParam = Cast<UDEditorStaticComponentMaskParameterValue>(StackParameterData->Parameter);
-		UDEditorVectorParameterValue* VectorParam = Cast<UDEditorVectorParameterValue>(StackParameterData->Parameter);
-		UDEditorScalarParameterValue* ScalarParam = Cast<UDEditorScalarParameterValue>(StackParameterData->Parameter);
-		UDEditorTextureParameterValue* TextureParam = Cast<UDEditorTextureParameterValue>(StackParameterData->Parameter);
 		NameOverride = FText::FromName(StackParameterData->Parameter->ParameterInfo.Name);
 
 		IDetailTreeNode& Node = *StackParameterData->ParameterNode;
+		TSharedPtr<IPropertyHandle> PropertyHandle = Node.CreatePropertyHandle();
 		TSharedPtr<IDetailPropertyRow> GeneratedRow = StaticCastSharedPtr<IDetailPropertyRow>(Node.GetRow());
 		IDetailPropertyRow& Row = *GeneratedRow.Get();
 		Row
 			.DisplayName(NameOverride);
 
-		if (VectorParam && VectorParam->bIsUsedAsChannelMask)
-		{
-			FOnGetPropertyComboBoxStrings GetMaskStrings = FOnGetPropertyComboBoxStrings::CreateStatic(&FMaterialPropertyHelpers::GetVectorChannelMaskComboBoxStrings);
-			FOnGetPropertyComboBoxValue GetMaskValue = FOnGetPropertyComboBoxValue::CreateStatic(&FMaterialPropertyHelpers::GetVectorChannelMaskValue, StackParameterData->Parameter);
-			FOnPropertyComboBoxValueSelected SetMaskValue = FOnPropertyComboBoxValueSelected::CreateStatic(&FMaterialPropertyHelpers::SetVectorChannelMaskValue, StackParameterData->ParameterNode->CreatePropertyHandle(), StackParameterData->Parameter, (UObject*)MaterialEditorInstance);
-
-			FDetailWidgetRow& CustomWidget = Row.CustomWidget();
-			CustomWidget
-				.FilterString(NameOverride)
-				.NameContent()
-				[
-					SNew(STextBlock)
-					.Text(NameOverride)
-					.ToolTipText(FMaterialPropertyHelpers::GetParameterExpressionDescription(StackParameterData->Parameter, MaterialEditorInstance))
-					.Font(FAppStyle::GetFontStyle(TEXT("PropertyWindow.NormalFont")))
-				]
-				.ValueContent()
-				.MaxDesiredWidth(200.0f)
-				[
-					SNew(SHorizontalBox)
-					+ SHorizontalBox::Slot()
-					.FillWidth(1.0f)
-					[
-						SNew(SHorizontalBox)
-						+ SHorizontalBox::Slot()
-						.HAlign(HAlign_Left)
-						.AutoWidth()
-						[
-							PropertyCustomizationHelpers::MakePropertyComboBox(StackParameterData->ParameterNode->CreatePropertyHandle(), GetMaskStrings, GetMaskValue, SetMaskValue)
-						]
-					]
-				];
-		}
-		else if (ScalarParam)
-		{
-			if (ScalarParam->AtlasData.bIsUsedAsAtlasPosition)
-			{
-				const FText ParameterName = FText::FromName(StackParameterData->Parameter->ParameterInfo.Name);
-
-				FDetailWidgetRow& CustomWidget = Row.CustomWidget();
-				CustomWidget
-					.FilterString(ParameterName)
-					.NameContent()
-					[
-						SNew(STextBlock)
-						.Text(ParameterName)
-						.ToolTipText(FMaterialPropertyHelpers::GetParameterExpressionDescription(StackParameterData->Parameter, MaterialEditorInstance))
-						.Font(FAppStyle::GetFontStyle(TEXT("PropertyWindow.NormalFont")))
-					]
-				.ValueContent()
-					.HAlign(HAlign_Fill)
-					.MaxDesiredWidth(400.0f)
-					[
-						SNew(SObjectPropertyEntryBox)
-						.ObjectPath(this, &SMaterialLayersFunctionsMaterialTreeItem::GetCurvePath, ScalarParam)
-						.AllowedClass(UCurveLinearColor::StaticClass())
-						.NewAssetFactories(TArray<UFactory*>())
-						.DisplayThumbnail(true)
-						.ThumbnailPool(InArgs._InTree->GetTreeThumbnailPool())
-						.DisplayCompactSize(true)
-					];
-			}
-			else if (!ScalarParam->bUseCustomPrimitiveData)
-			{
-				// Try to find a valid enumeration
-				TSoftObjectPtr<UObject> Enumeration;
-				if (!ScalarParam->Enumeration.IsNull())
-				{
-					Enumeration = ScalarParam->Enumeration;
-				}
-				// EnumerationIndex is not supported in a Material
-
-				// If found a valid Enumeration, customize the widget, else fallback todefault scalar numeric edit box
-				if (!Enumeration.IsNull())
-				{
-					const FText ParameterName = FText::FromName(StackParameterData->Parameter->ParameterInfo.Name);
-					TSharedPtr<IPropertyHandle> ParameterValueProperty = StackParameterData->ParameterNode->CreatePropertyHandle();
-
-					Row.CustomWidget()
-						.FilterString(ParameterName)
-						.NameContent()
-						[
-							ParameterValueProperty->CreatePropertyNameWidget()
-						]
-						.ValueContent()
-						.HAlign(HAlign_Fill)
-						.MaxDesiredWidth(400.0f)
-						[
-							FMaterialPropertyHelpers::CreateScalarEnumerationParameterValueWidget(Enumeration, ParameterValueProperty)
-						];
-				}
-			}
-		}
-		else if (TextureParam)
-		{
-			UMaterial* Material = MaterialEditorInstance->PreviewMaterial;
-			if (Material != nullptr)
-			{
-				UMaterialExpressionTextureSampleParameter* Expression = Material->FindExpressionByGUID<UMaterialExpressionTextureSampleParameter>(TextureParam->ExpressionId);
-				if (Expression != nullptr)
-				{
-					TWeakObjectPtr<UMaterialExpressionTextureSampleParameter> SamplerExpression = Expression;
-					TSharedPtr<SWidget> NameWidget;
-					TSharedPtr<SWidget> ValueWidget;
-					TSharedPtr<SVerticalBox> NameVerticalBox;
-					const FText ParameterName = FText::FromName(StackParameterData->Parameter->ParameterInfo.Name);
-					FDetailWidgetRow& CustomWidget = Row.CustomWidget();
-					CustomWidget
-						.FilterString(ParameterName)
-						.NameContent()
-						[
-							SAssignNew(NameVerticalBox, SVerticalBox)
-							+ SVerticalBox::Slot()
-							.AutoHeight()
-							[
-								SNew(STextBlock)
-								.Text(ParameterName)
-								.ToolTipText(FMaterialPropertyHelpers::GetParameterExpressionDescription(StackParameterData->Parameter, MaterialEditorInstance))
-								.Font(FAppStyle::GetFontStyle(TEXT("PropertyWindow.NormalFont")))
-							]
-						];
-					CustomWidget.ValueContent()
-						[
-							SNew(SObjectPropertyEntryBox)
-							.PropertyHandle(StackParameterData->ParameterNode->CreatePropertyHandle())
-							.AllowedClass(UTexture::StaticClass())
-							.ThumbnailPool(Tree->GetTreeThumbnailPool())
-						];
-
-					static const FName Red("R");
-					static const FName Green("G");
-					static const FName Blue("B");
-					static const FName Alpha("A");
-
-					if (!TextureParam->ChannelNames.R.IsEmpty())
-					{
-						NameVerticalBox->AddSlot()
-							[
-								SNew(SHorizontalBox)
-								+ SHorizontalBox::Slot()
-								.AutoWidth()
-								.Padding(20.0, 2.0, 4.0, 2.0)
-								[
-									SNew(STextBlock)
-									.Text(FText::FromName(Red))
-									.Font(FAppStyle::GetFontStyle(TEXT("PropertyWindow.BoldFont")))
-								]
-								+ SHorizontalBox::Slot()
-								.HAlign(HAlign_Left)
-								.Padding(4.0, 2.0)
-								[
-									SNew(STextBlock)
-									.Text(TextureParam->ChannelNames.R)
-									.Font(FAppStyle::GetFontStyle(TEXT("PropertyWindow.NormalFont")))
-								]
-							];
-					}
-					if (!TextureParam->ChannelNames.G.IsEmpty())
-					{
-						NameVerticalBox->AddSlot()
-							[
-								SNew(SHorizontalBox)
-								+ SHorizontalBox::Slot()
-								.Padding(20.0, 2.0, 4.0, 2.0)
-								.AutoWidth()
-								[
-									SNew(STextBlock)
-									.Text(FText::FromName(Green))
-									.Font(FAppStyle::GetFontStyle(TEXT("PropertyWindow.BoldFont")))
-								]
-								+ SHorizontalBox::Slot()
-								.HAlign(HAlign_Left)
-								.Padding(4.0, 2.0)
-								[
-									SNew(STextBlock)
-									.Text(TextureParam->ChannelNames.G)
-									.Font(FAppStyle::GetFontStyle(TEXT("PropertyWindow.NormalFont")))
-								]
-							];
-					}
-					if (!TextureParam->ChannelNames.B.IsEmpty())
-					{
-						NameVerticalBox->AddSlot()
-							[
-								SNew(SHorizontalBox)
-								+ SHorizontalBox::Slot()
-								.Padding(20.0, 2.0, 4.0, 2.0)
-								.AutoWidth()
-								[
-									SNew(STextBlock)
-									.Text(FText::FromName(Blue))
-									.Font(FAppStyle::GetFontStyle(TEXT("PropertyWindow.BoldFont")))
-								]
-								+ SHorizontalBox::Slot()
-								.HAlign(HAlign_Left)
-								.Padding(4.0, 2.0)
-								[
-									SNew(STextBlock)
-									.Text(TextureParam->ChannelNames.B)
-									.Font(FAppStyle::GetFontStyle(TEXT("PropertyWindow.NormalFont")))
-								]
-							];
-					}
-					if (!TextureParam->ChannelNames.A.IsEmpty())
-					{
-						NameVerticalBox->AddSlot()
-							[
-								SNew(SHorizontalBox)
-								+ SHorizontalBox::Slot()
-								.Padding(20.0, 2.0, 4.0, 2.0)
-								.AutoWidth()
-								[
-									SNew(STextBlock)
-									.Text(FText::FromName(Alpha))
-									.Font(FAppStyle::GetFontStyle(TEXT("PropertyWindow.BoldFont")))
-								]
-								+ SHorizontalBox::Slot()
-								.HAlign(HAlign_Left)
-								.Padding(4.0, 2.0)
-								[
-									SNew(STextBlock)
-									.Text(TextureParam->ChannelNames.A)
-									.Font(FAppStyle::GetFontStyle(TEXT("PropertyWindow.NormalFont")))
-								]
-							];
-					}
-				}
-			}
-		}
-		else if (CompMaskParam)
-		{
-			TSharedPtr<IPropertyHandle> RMaskProperty = StackParameterData->ParameterNode->CreatePropertyHandle()->GetChildHandle("R");
-			TSharedPtr<IPropertyHandle> GMaskProperty = StackParameterData->ParameterNode->CreatePropertyHandle()->GetChildHandle("G");
-			TSharedPtr<IPropertyHandle> BMaskProperty = StackParameterData->ParameterNode->CreatePropertyHandle()->GetChildHandle("B");
-			TSharedPtr<IPropertyHandle> AMaskProperty = StackParameterData->ParameterNode->CreatePropertyHandle()->GetChildHandle("A");
-			FDetailWidgetRow& CustomWidget = Row.CustomWidget();
-			CustomWidget
-				.FilterString(NameOverride)
-				.NameContent()
-				[
-					SNew(SHorizontalBox)
-					+ SHorizontalBox::Slot()
-					.VAlign(VAlign_Center)
-					[
-						SNew(STextBlock)
-						.Text(NameOverride)
-						.ToolTipText(FMaterialPropertyHelpers::GetParameterExpressionDescription(StackParameterData->Parameter, MaterialEditorInstance))
-						.Font(FAppStyle::GetFontStyle(TEXT("PropertyWindow.NormalFont")))
-					]
-				]
-				.ValueContent()
-				.MaxDesiredWidth(200.0f)
-				[
-					SNew(SHorizontalBox)
-					+ SHorizontalBox::Slot()
-					.FillWidth(1.0f)
-				[
-					SNew(SHorizontalBox)
-					+ SHorizontalBox::Slot()
-					.HAlign(HAlign_Left)
-					.AutoWidth()
-					[
-						RMaskProperty->CreatePropertyNameWidget()
-					]
-			+ SHorizontalBox::Slot()
-				.HAlign(HAlign_Left)
-				.AutoWidth()
-				[
-					RMaskProperty->CreatePropertyValueWidget()
-				]
-			+ SHorizontalBox::Slot()
-				.HAlign(HAlign_Left)
-				.Padding(FMargin(10.0f, 0.0f, 0.0f, 0.0f))
-				.AutoWidth()
-				[
-					GMaskProperty->CreatePropertyNameWidget()
-				]
-			+ SHorizontalBox::Slot()
-				.HAlign(HAlign_Left)
-				.AutoWidth()
-				[
-					GMaskProperty->CreatePropertyValueWidget()
-				]
-			+ SHorizontalBox::Slot()
-				.HAlign(HAlign_Left)
-				.Padding(FMargin(10.0f, 0.0f, 0.0f, 0.0f))
-				.AutoWidth()
-				[
-					BMaskProperty->CreatePropertyNameWidget()
-				]
-			+ SHorizontalBox::Slot()
-				.HAlign(HAlign_Left)
-				.AutoWidth()
-				[
-					BMaskProperty->CreatePropertyValueWidget()
-				]
-			+ SHorizontalBox::Slot()
-				.HAlign(HAlign_Left)
-				.Padding(FMargin(10.0f, 0.0f, 0.0f, 0.0f))
-				.AutoWidth()
-				[
-					AMaskProperty->CreatePropertyNameWidget()
-				]
-			+ SHorizontalBox::Slot()
-				.HAlign(HAlign_Left)
-				.AutoWidth()
-				[
-					AMaskProperty->CreatePropertyValueWidget()
-				]
-				]
-				];
-		}
-		else
-		{
-			FDetailWidgetDecl* CustomNameWidget = Row.CustomNameWidget();
-			if (CustomNameWidget)
-			{
-				(*CustomNameWidget)
-				[
-					SNew(SHorizontalBox)
-					+SHorizontalBox::Slot()
-					.VAlign(VAlign_Center)
-					[
-						SNew(STextBlock)
-						.Text(NameOverride)
-						.ToolTipText(FMaterialPropertyHelpers::GetParameterExpressionDescription(StackParameterData->Parameter, MaterialEditorInstance))
-						.Font(FAppStyle::GetFontStyle(TEXT("PropertyWindow.NormalFont")))
-					]
-				];
-			}
-			else
-			{
-				Row.ToolTip(FMaterialPropertyHelpers::GetParameterExpressionDescription(StackParameterData->Parameter, MaterialEditorInstance));
-			}
-		}
+		FMaterialPropertyHelpers::SetPropertyRowParameterWidget(Row, StackParameterData->Parameter, PropertyHandle, MaterialEditorInstance);
 
 		FNodeWidgets NodeWidgets = Node.CreateNodeWidgets();
 		LeftSideWidget = NodeWidgets.NameWidget.ToSharedRef();
@@ -2777,46 +2051,51 @@ void SMaterialLayersFunctionsMaterialTreeItem::Construct(const FArguments& InArg
 				.BorderImage(FAppStyle::Get().GetBrush("DetailsView.GridLine"))
 				.Padding(FMargin(0, 0, 0, 1))
 				[
-					SNew(SBorder)
-					.Padding(0.0f)
-					.BorderImage(this, &SMaterialLayersFunctionsMaterialTreeItem::GetBorderImage)
-					.BorderBackgroundColor(this, &SMaterialLayersFunctionsMaterialTreeItem::GetOuterBackgroundColor, StackParameterData)
+					SNew(SBox)
+					.MinDesiredHeight(26.0f)	// Match PropertyEditorConstants::PropertyRowHeight
 					[
-						SNew(SSplitter)
-						.Style(FAppStyle::Get(), "DetailsView.Splitter")
-						.PhysicalSplitterHandleSize(1.0f)
-						.HitDetectionSplitterHandleSize(5.0f)
-						+ SSplitter::Slot()
-						.Value(Tree->ColumnSizeData.GetNameColumnWidth())
-						.OnSlotResized(Tree->ColumnSizeData.GetOnNameColumnResized())
-						.Value(0.25f)
+						SNew(SBorder)
+						.Padding(0.0f)
+						.BorderImage(this, &SMaterialLayersFunctionsMaterialTreeItem::GetBorderImage)
+						.BorderBackgroundColor(this, &SMaterialLayersFunctionsMaterialTreeItem::GetOuterBackgroundColor, StackParameterData)
 						[
-							SNew(SHorizontalBox)
-							+ SHorizontalBox::Slot()
-							.AutoWidth()
-							.VAlign(VAlign_Center)
-							.Padding(FMargin(3.0f))
+							SNew(SSplitter)
+							.Style(FAppStyle::Get(), "DetailsView.Splitter")
+							.PhysicalSplitterHandleSize(1.0f)
+							.HitDetectionSplitterHandleSize(5.0f)
+							+ SSplitter::Slot()
+							.Value(Tree->ColumnSizeData.GetNameColumnWidth())
+							.OnSlotResized(Tree->ColumnSizeData.GetOnNameColumnResized())
+							.Value(0.25f)
 							[
-								SNew(SExpanderArrow, SharedThis(this))
-							]
-							+ SHorizontalBox::Slot()
-							.Padding(FMargin(2.0f))
-							.VAlign(VAlign_Center)
-							[
-								LeftSideWidget
+								SNew(SHorizontalBox)
+								+ SHorizontalBox::Slot()
+								.AutoWidth()
+								.VAlign(VAlign_Center)
+								.Padding(FMargin(3.0f))
+								[
+									SNew(SExpanderArrow, SharedThis(this))
+								]
+								+ SHorizontalBox::Slot()
+								.Padding(FMargin(2.0f))
+								.VAlign(VAlign_Center)
+								[
+									LeftSideWidget
+								]
 							]
-						]
-						+ SSplitter::Slot()
-						.Value(Tree->ColumnSizeData.GetValueColumnWidth())
-						.OnSlotResized(Tree->ColumnSizeData.GetOnValueColumnResized())
-						[
-							SNew(SHorizontalBox)
-							.Clipping(EWidgetClipping::OnDemand)
-							+ SHorizontalBox::Slot()
-							.HAlign(ValueAlignment)
-							.VAlign(VAlign_Center)
+							+ SSplitter::Slot()
+							.Value(Tree->ColumnSizeData.GetValueColumnWidth())
+							.OnSlotResized(Tree->ColumnSizeData.GetOnValueColumnResized())
 							[
-								RightSideWidget
+								SNew(SHorizontalBox)
+								.Clipping(EWidgetClipping::OnDemand)
+								+ SHorizontalBox::Slot()
+								.HAlign(ValueAlignment)
+								.VAlign(VAlign_Center)
+								.Padding(FMargin(12.0f, 0.0f, 2.0f, 0.0f))	// Match DetailWidgetConstants::RightRowPadding
+								[
+									RightSideWidget
+								]
 							]
 						]
 					]
@@ -3139,15 +2418,14 @@ void SMaterialLayersFunctionsMaterialTree::CreateGroupsWidget()
 			}
 			else
 			{
-				FUnsortedParamData NonLayerProperty;
-				UDEditorScalarParameterValue* ScalarParam = Cast<UDEditorScalarParameterValue>(Parameter);
-
-				if (ScalarParam && ScalarParam->SliderMax > ScalarParam->SliderMin)
+				// Skip parameters that shouldn't have widgets here
+				if (!FMaterialPropertyHelpers::ShouldCreatePropertyRowForParameter(Parameter))
 				{
-					ParameterValueProperty->SetInstanceMetaData("UIMin", FString::Printf(TEXT("%f"), ScalarParam->SliderMin));
-					ParameterValueProperty->SetInstanceMetaData("UIMax", FString::Printf(TEXT("%f"), ScalarParam->SliderMax));
+					continue;
 				}
 
+				FUnsortedParamData NonLayerProperty;
+
 				NonLayerProperty.Parameter = Parameter;
 				NonLayerProperty.ParameterGroup = ParameterGroup;
 
@@ -3259,23 +2537,25 @@ void SMaterialLayersFunctionsMaterialTree::ShowSubParameters(TSharedPtr<FSortedP
 		if (Parameter->ParameterInfo.Index == ParentParameter->ParameterInfo.Index
 			&& Parameter->ParameterInfo.Association == ParentParameter->ParameterInfo.Association)
 		{
-			TSharedPtr<FSortedParamData> GroupProperty(new FSortedParamData());
-			GroupProperty->StackDataType = EStackDataType::Group;
-			GroupProperty->ParameterInfo.Index = Parameter->ParameterInfo.Index;
-			GroupProperty->ParameterInfo.Association = Parameter->ParameterInfo.Association;
-			GroupProperty->Group = Property.ParameterGroup;
-			GroupProperty->NodeKey = FString::FromInt(GroupProperty->ParameterInfo.Index) + FString::FromInt(GroupProperty->ParameterInfo.Association) + Property.ParameterGroup.GroupName.ToString();
+			FString GroupNodeKey = FString::FromInt(Parameter->ParameterInfo.Index) + FString::FromInt(Parameter->ParameterInfo.Association) + Property.ParameterGroup.GroupName.ToString();
 
 			bool bAddNewGroup = true;
 			for (TSharedPtr<struct FSortedParamData> GroupChild : ParentParameter->Children)
 			{
-				if (GroupChild->NodeKey == GroupProperty->NodeKey)
+				if (GroupChild->NodeKey == GroupNodeKey)
 				{
 					bAddNewGroup = false;
 				}
 			}
 			if (bAddNewGroup)
 			{
+				TSharedPtr<FSortedParamData> GroupProperty(new FSortedParamData());
+				GroupProperty->StackDataType = EStackDataType::Group;
+				GroupProperty->ParameterInfo.Index = Parameter->ParameterInfo.Index;
+				GroupProperty->ParameterInfo.Association = Parameter->ParameterInfo.Association;
+				GroupProperty->Group = Property.ParameterGroup;
+				GroupProperty->NodeKey = GroupNodeKey;
+
 				ParentParameter->Children.Add(GroupProperty);
 			}
 
@@ -3289,8 +2569,7 @@ void SMaterialLayersFunctionsMaterialTree::ShowSubParameters(TSharedPtr<FSortedP
 			ChildProperty->NodeKey = FString::FromInt(ChildProperty->ParameterInfo.Index) + FString::FromInt(ChildProperty->ParameterInfo.Association) + Property.ParameterGroup.GroupName.ToString() + Property.UnsortedName.ToString();
 
 
-			UDEditorStaticComponentMaskParameterValue* CompMaskParam = Cast<UDEditorStaticComponentMaskParameterValue>(Parameter);
-			if (!CompMaskParam)
+			if (FMaterialPropertyHelpers::ShouldCreateChildPropertiesForParameter(Parameter))
 			{
 				TArray<TSharedRef<IDetailTreeNode>> ParamChildren;
 				Property.ParameterNode->GetChildren(ParamChildren);
diff --git a/Engine/Source/Editor/MaterialEditor/Private/SMaterialLayersFunctionsTree.h b/Engine/Source/Editor/MaterialEditor/Private/SMaterialLayersFunctionsTree.h
--- a/Engine/Source/Editor/MaterialEditor/Private/SMaterialLayersFunctionsTree.h
+++ b/Engine/Source/Editor/MaterialEditor/Private/SMaterialLayersFunctionsTree.h
@@ -45,7 +45,6 @@ private:
 	bool bIsHoveredDragTarget;
 
 
-	FString GetCurvePath(class UDEditorScalarParameterValue* Parameter) const;
 	const FSlateBrush* GetBorderImage() const;
 
 	FSlateColor GetOuterBackgroundColor(TSharedPtr<FSortedParamData> InParamData) const;
@@ -252,7 +251,6 @@ public:
 
 private:
 	bool bIsHoveredDragTarget;
-	FString GetCurvePath(class UDEditorScalarParameterValue* Parameter) const;
 	const FSlateBrush* GetBorderImage() const;
 
 	FSlateColor GetOuterBackgroundColor(TSharedPtr<FSortedParamData> InParamData) const;
diff --git a/Engine/Source/Editor/MaterialEditor/Private/SMaterialParametersOverviewWidget.cpp b/Engine/Source/Editor/MaterialEditor/Private/SMaterialParametersOverviewWidget.cpp
--- a/Engine/Source/Editor/MaterialEditor/Private/SMaterialParametersOverviewWidget.cpp
+++ b/Engine/Source/Editor/MaterialEditor/Private/SMaterialParametersOverviewWidget.cpp
@@ -1,47 +1,28 @@
 // Copyright Epic Games, Inc. All Rights Reserved.
 
 #include "SMaterialParametersOverviewWidget.h"
-#include "MaterialEditor/DEditorFontParameterValue.h"
-#include "MaterialEditor/DEditorMaterialLayersParameterValue.h"
-#include "MaterialEditor/DEditorScalarParameterValue.h"
-#include "MaterialEditor/DEditorStaticComponentMaskParameterValue.h"
-#include "MaterialEditor/DEditorStaticSwitchParameterValue.h"
-#include "MaterialEditor/DEditorTextureParameterValue.h"
-#include "MaterialEditor/DEditorVectorParameterValue.h"
 #include "MaterialEditor/MaterialEditorInstanceConstant.h"
 #include "Materials/Material.h"
 #include "PropertyHandle.h"
 #include "Modules/ModuleManager.h"
 #include "PropertyEditorModule.h"
-#include "ISinglePropertyView.h"
 #include "Materials/MaterialInstanceConstant.h"
 #include "Styling/AppStyle.h"
-#include "PropertyCustomizationHelpers.h"
-#include "ScopedTransaction.h"
 #include "IPropertyRowGenerator.h"
 #include "Widgets/Views/STreeView.h"
 #include "IDetailTreeNode.h"
 #include "AssetThumbnail.h"
-#include "MaterialEditorInstanceDetailCustomization.h"
 #include "MaterialPropertyHelpers.h"
 #include "DetailWidgetRow.h"
-#include "Widgets/Input/SCheckBox.h"
 #include "EditorSupportDelegates.h"
 #include "Widgets/Images/SImage.h"
 #include "MaterialEditor/MaterialEditorPreviewParameters.h"
 #include "Widgets/Input/SButton.h"
-#include "Widgets/Text/SInlineEditableTextBlock.h"
 #include "Materials/MaterialFunction.h"
-#include "Materials/MaterialFunctionInstance.h"
 #include "Materials/MaterialFunctionInterface.h"
-#include "Materials/MaterialExpressionTextureSampleParameter.h"
 #include "Framework/Application/SlateApplication.h"
 
-#include "Widgets/Input/SEditableTextBox.h"
-#include "Curves/CurveLinearColor.h"
-#include "Curves/CurveLinearColorAtlas.h"
 #include "Widgets/Views/SExpanderArrow.h"
-#include "Widgets/Layout/SWidgetSwitcher.h"
 #include "MaterialEditor/DEditorRuntimeVirtualTextureParameterValue.h"
 #include "MaterialEditor/DEditorSparseVolumeTextureParameterValue.h"
 #include "ThumbnailRendering/ThumbnailManager.h"
@@ -51,12 +32,6 @@
 
 #define LOCTEXT_NAMESPACE "MaterialLayerCustomization"
 
-FString SMaterialParametersOverviewTreeItem::GetCurvePath(UDEditorScalarParameterValue* Parameter) const
-{
-	FString Path = Parameter->AtlasData.Curve->GetPathName();
-	return Path;
-}
-
 const FSlateBrush* SMaterialParametersOverviewTreeItem::GetBorderImage() const
 {
 	return FAppStyle::Get().GetBrush("DetailsView.CategoryMiddle");
@@ -90,6 +65,7 @@ void SMaterialParametersOverviewTreeItem::Construct(const FArguments& InArgs, co
 	TSharedRef<SWidget> RightSideWidget = SNullWidget::NullWidget;
 	FText NameOverride;
 	TSharedRef<SVerticalBox> WrapperWidget = SNew(SVerticalBox);
+	FMargin RightSidePadding;
 
 // GROUP --------------------------------------------------
 	if (StackParameterData->StackDataType == EStackDataType::Group)
@@ -104,358 +80,23 @@ void SMaterialParametersOverviewTreeItem::Construct(const FArguments& InArgs, co
 // END GROUP
 
 // PROPERTY ----------------------------------------------
-	bool bisPaddedProperty = false;
 	if (StackParameterData->StackDataType == EStackDataType::Property)
 	{
-		UDEditorStaticComponentMaskParameterValue* CompMaskParam = Cast<UDEditorStaticComponentMaskParameterValue>(StackParameterData->Parameter);
-		UDEditorVectorParameterValue* VectorParam = Cast<UDEditorVectorParameterValue>(StackParameterData->Parameter);
-		UDEditorScalarParameterValue* ScalarParam = Cast<UDEditorScalarParameterValue>(StackParameterData->Parameter);
-		UDEditorTextureParameterValue* TextureParam = Cast<UDEditorTextureParameterValue>(StackParameterData->Parameter);
-
-		TAttribute<bool> IsParamEnabled = TAttribute<bool>::Create(TAttribute<bool>::FGetter::CreateStatic(&FMaterialPropertyHelpers::IsOverriddenExpression, StackParameterData->Parameter));
 		NameOverride = FText::FromName(StackParameterData->Parameter->ParameterInfo.Name);
 
 		IDetailTreeNode& Node = *StackParameterData->ParameterNode;
+		TSharedPtr<IPropertyHandle> PropertyHandle = Node.CreatePropertyHandle();
 		TSharedPtr<IDetailPropertyRow> GeneratedRow = StaticCastSharedPtr<IDetailPropertyRow>(Node.GetRow());
 		IDetailPropertyRow& Row = *GeneratedRow.Get();
 		Row.DisplayName(NameOverride);
 
-		if (VectorParam && VectorParam->bIsUsedAsChannelMask)
-		{
-			FOnGetPropertyComboBoxStrings GetMaskStrings = FOnGetPropertyComboBoxStrings::CreateStatic(&FMaterialPropertyHelpers::GetVectorChannelMaskComboBoxStrings);
-			FOnGetPropertyComboBoxValue GetMaskValue = FOnGetPropertyComboBoxValue::CreateStatic(&FMaterialPropertyHelpers::GetVectorChannelMaskValue, StackParameterData->Parameter);
-			FOnPropertyComboBoxValueSelected SetMaskValue = FOnPropertyComboBoxValueSelected::CreateStatic(&FMaterialPropertyHelpers::SetVectorChannelMaskValue, StackParameterData->ParameterNode->CreatePropertyHandle(), StackParameterData->Parameter, (UObject*)MaterialEditorInstance);
-
-			FDetailWidgetRow& CustomWidget = Row.CustomWidget();
-			CustomWidget
-			.FilterString(NameOverride)
-			.NameContent()
-			[
-				SNew(STextBlock)
-				.Text(NameOverride)
-				.ToolTipText(FMaterialPropertyHelpers::GetParameterExpressionDescription(StackParameterData->Parameter, MaterialEditorInstance))
-				.Font(FAppStyle::GetFontStyle(TEXT("PropertyWindow.NormalFont")))
-			]
-			.ValueContent()
-			.MaxDesiredWidth(200.0f)
-			[
-				SNew(SHorizontalBox)
-				+ SHorizontalBox::Slot()
-				.FillWidth(1.0f)
-				[
-					SNew(SHorizontalBox)
-					+ SHorizontalBox::Slot()
-					.HAlign(HAlign_Left)
-					.AutoWidth()
-					[
-						PropertyCustomizationHelpers::MakePropertyComboBox(StackParameterData->ParameterNode->CreatePropertyHandle(), GetMaskStrings, GetMaskValue, SetMaskValue)
-					]
-				]
-			];
-		}
-		else if (ScalarParam && ScalarParam->AtlasData.bIsUsedAsAtlasPosition)
-		{
-			const FText ParameterName = FText::FromName(StackParameterData->Parameter->ParameterInfo.Name);
-
-			FDetailWidgetRow& CustomWidget = Row.CustomWidget();
-			CustomWidget
-				.FilterString(ParameterName)
-				.NameContent()
-				[
-					SNew(STextBlock)
-					.Text(ParameterName)
-					.ToolTipText(FMaterialPropertyHelpers::GetParameterExpressionDescription(StackParameterData->Parameter, MaterialEditorInstance))
-					.Font(FAppStyle::GetFontStyle(TEXT("PropertyWindow.NormalFont")))
-				]
-				.ValueContent()
-				.HAlign(HAlign_Fill)
-				.MaxDesiredWidth(400.0f)
-				[
-					SNew(SObjectPropertyEntryBox)
-					.ObjectPath(this, &SMaterialParametersOverviewTreeItem::GetCurvePath, ScalarParam)
-					.AllowedClass(UCurveLinearColor::StaticClass())
-					.NewAssetFactories(TArray<UFactory*>())
-					.DisplayThumbnail(true)
-					.ThumbnailPool(Tree.Pin()->GetTreeThumbnailPool())
-					.OnShouldFilterAsset(FOnShouldFilterAsset::CreateStatic(&FMaterialPropertyHelpers::OnShouldFilterCurveAsset, ScalarParam->AtlasData.Atlas))
-					.OnShouldSetAsset(FOnShouldSetAsset::CreateStatic(&FMaterialPropertyHelpers::OnShouldSetCurveAsset, ScalarParam->AtlasData.Atlas))
-					.OnObjectChanged(FOnSetObject::CreateStatic(&FMaterialPropertyHelpers::SetPositionFromCurveAsset, ScalarParam->AtlasData.Atlas, ScalarParam, StackParameterData->ParameterHandle, (UObject*)MaterialEditorInstance))
-					.DisplayCompactSize(true)
-				];
-			
-		}
-		else if (TextureParam)
-		{
-			UMaterial *Material = MaterialEditorInstance->PreviewMaterial;
-			if (Material != nullptr)
-			{
-				UMaterialExpressionTextureSampleParameter* Expression = Material->FindExpressionByGUID<UMaterialExpressionTextureSampleParameter>(TextureParam->ExpressionId);
-				if (Expression != nullptr)
-				{
-					TWeakObjectPtr<UMaterialExpressionTextureSampleParameter> SamplerExpression = Expression;
-					TSharedPtr<SWidget> NameWidget;
-					TSharedPtr<SWidget> ValueWidget;
-					FDetailWidgetRow DefaultRow;
-					Row.GetDefaultWidgets(NameWidget, ValueWidget, DefaultRow);
-
-					FDetailWidgetRow &DetailWidgetRow = Row.CustomWidget();
-					TSharedPtr<SVerticalBox> NameVerticalBox;
-					DetailWidgetRow.NameContent()
-						[
-							SAssignNew(NameVerticalBox, SVerticalBox)
-							+ SVerticalBox::Slot()
-							.AutoHeight()
-							[
-								SNew(STextBlock)
-								.Text(FText::FromName(StackParameterData->Parameter->ParameterInfo.Name))
-								.Font(FAppStyle::GetFontStyle(TEXT("PropertyWindow.NormalFont")))
-							]
-						];
-
-					DetailWidgetRow.ValueContent()
-						.MinDesiredWidth(DefaultRow.ValueWidget.MinWidth)
-						.MaxDesiredWidth(DefaultRow.ValueWidget.MaxWidth)
-						[
-							SNew(SObjectPropertyEntryBox)
-							.PropertyHandle(StackParameterData->ParameterNode->CreatePropertyHandle())
-							.AllowedClass(UTexture::StaticClass())
-							.ThumbnailPool(Tree.Pin()->GetTreeThumbnailPool())
-							.OnShouldFilterAsset_Lambda([SamplerExpression](const FAssetData& AssetData)
-							{
-								if (SamplerExpression.Get())
-								{
-									bool VirtualTextured = false;
-									AssetData.GetTagValue<bool>("VirtualTextureStreaming", VirtualTextured);
-
-									bool ExpressionIsVirtualTextured = IsVirtualSamplerType(SamplerExpression->SamplerType);
-
-									return VirtualTextured != ExpressionIsVirtualTextured;
-								}
-								else
-								{
-									return false;
-								}
-							})
-						];
-
-					static const FName Red("R");
-					static const FName Green("G");
-					static const FName Blue("B");
-					static const FName Alpha("A");
-
-					if (!TextureParam->ChannelNames.R.IsEmpty())
-					{
-						NameVerticalBox->AddSlot()
-							[
-								SNew(SHorizontalBox)
-								+ SHorizontalBox::Slot()
-								.AutoWidth()
-								.Padding(20.0, 2.0, 4.0, 2.0)
-								[
-									SNew(STextBlock)
-									.Text(FText::FromName(Red))
-									.Font(FAppStyle::GetFontStyle(TEXT("PropertyWindow.BoldFont")))
-								]
-								+ SHorizontalBox::Slot()
-								.HAlign(HAlign_Left)
-								.Padding(4.0, 2.0)
-								[
-									SNew(STextBlock)
-									.Text(TextureParam->ChannelNames.R)
-									.Font(FAppStyle::GetFontStyle(TEXT("PropertyWindow.NormalFont")))
-								]
-							];
-					}
-					if (!TextureParam->ChannelNames.G.IsEmpty())
-					{
-						NameVerticalBox->AddSlot()
-							[
-								SNew(SHorizontalBox)
-								+ SHorizontalBox::Slot()
-								.Padding(20.0, 2.0, 4.0, 2.0)
-								.AutoWidth()
-								[
-									SNew(STextBlock)
-									.Text(FText::FromName(Green))
-									.Font(FAppStyle::GetFontStyle(TEXT("PropertyWindow.BoldFont")))
-								]
-								+ SHorizontalBox::Slot()
-								.HAlign(HAlign_Left)
-								.Padding(4.0, 2.0)
-								[
-									SNew(STextBlock)
-									.Text(TextureParam->ChannelNames.G)
-									.Font(FAppStyle::GetFontStyle(TEXT("PropertyWindow.NormalFont")))
-								]
-							];
-					}
-					if (!TextureParam->ChannelNames.B.IsEmpty())
-					{
-						NameVerticalBox->AddSlot()
-							[
-								SNew(SHorizontalBox)
-								+ SHorizontalBox::Slot()
-								.Padding(20.0, 2.0, 4.0, 2.0)
-								.AutoWidth()
-								[
-									SNew(STextBlock)
-									.Text(FText::FromName(Blue))
-									.Font(FAppStyle::GetFontStyle(TEXT("PropertyWindow.BoldFont")))
-								]
-								+ SHorizontalBox::Slot()
-								.HAlign(HAlign_Left)
-								.Padding(4.0, 2.0)
-								[
-									SNew(STextBlock)
-									.Text(TextureParam->ChannelNames.B)
-									.Font(FAppStyle::GetFontStyle(TEXT("PropertyWindow.NormalFont")))
-								]
-							];
-					}
-					if (!TextureParam->ChannelNames.A.IsEmpty())
-					{
-						NameVerticalBox->AddSlot()
-							[
-								SNew(SHorizontalBox)
-								+ SHorizontalBox::Slot()
-								.Padding(20.0, 2.0, 4.0, 2.0)
-								.AutoWidth()
-								[
-									SNew(STextBlock)
-									.Text(FText::FromName(Alpha))
-									.Font(FAppStyle::GetFontStyle(TEXT("PropertyWindow.BoldFont")))
-								]
-								+ SHorizontalBox::Slot()
-								.HAlign(HAlign_Left)
-								.Padding(4.0, 2.0)
-								[
-									SNew(STextBlock)
-									.Text(TextureParam->ChannelNames.A)
-									.Font(FAppStyle::GetFontStyle(TEXT("PropertyWindow.NormalFont")))
-								]
-							];
-					}
-
-				}
-			}
-		}
-		else if (CompMaskParam)
-		{
-			TSharedPtr<IPropertyHandle> RMaskProperty = StackParameterData->ParameterNode->CreatePropertyHandle()->GetChildHandle("R");
-			TSharedPtr<IPropertyHandle> GMaskProperty = StackParameterData->ParameterNode->CreatePropertyHandle()->GetChildHandle("G");
-			TSharedPtr<IPropertyHandle> BMaskProperty = StackParameterData->ParameterNode->CreatePropertyHandle()->GetChildHandle("B");
-			TSharedPtr<IPropertyHandle> AMaskProperty = StackParameterData->ParameterNode->CreatePropertyHandle()->GetChildHandle("A");
-			FDetailWidgetRow& CustomWidget = Row.CustomWidget();
-			CustomWidget
-			.FilterString(NameOverride)
-			.NameContent()
-			[
-				SNew(SHorizontalBox)
-				+ SHorizontalBox::Slot()
-				.VAlign(VAlign_Center)
-				[
-					SNew(STextBlock)
-					.Text(NameOverride)
-					.ToolTipText(FMaterialPropertyHelpers::GetParameterExpressionDescription(StackParameterData->Parameter, MaterialEditorInstance))
-					.Font(FAppStyle::GetFontStyle(TEXT("PropertyWindow.NormalFont")))
-				]
-			]
-			.ValueContent()
-			.MaxDesiredWidth(200.0f)
-			[
-				SNew(SHorizontalBox)
-				+ SHorizontalBox::Slot()
-				.FillWidth(1.0f)
-				[
-					SNew(SHorizontalBox)
-					+ SHorizontalBox::Slot()
-					.HAlign(HAlign_Left)
-					.AutoWidth()
-					[
-						RMaskProperty->CreatePropertyNameWidget()
-					]
-					+ SHorizontalBox::Slot()
-					.HAlign(HAlign_Left)
-					.AutoWidth()
-					[
-						RMaskProperty->CreatePropertyValueWidget()
-					]
-					+ SHorizontalBox::Slot()
-					.HAlign(HAlign_Left)
-					.Padding(FMargin(10.0f, 0.0f, 0.0f, 0.0f))
-					.AutoWidth()
-					[
-						GMaskProperty->CreatePropertyNameWidget()
-					]
-					+ SHorizontalBox::Slot()
-					.HAlign(HAlign_Left)
-					.AutoWidth()
-					[
-						GMaskProperty->CreatePropertyValueWidget()
-					]
-					+ SHorizontalBox::Slot()
-					.HAlign(HAlign_Left)
-					.Padding(FMargin(10.0f, 0.0f, 0.0f, 0.0f))
-					.AutoWidth()
-					[
-						BMaskProperty->CreatePropertyNameWidget()
-					]
-					+ SHorizontalBox::Slot()
-					.HAlign(HAlign_Left)
-					.AutoWidth()
-					[
-						BMaskProperty->CreatePropertyValueWidget()
-					]
-					+ SHorizontalBox::Slot()
-					.HAlign(HAlign_Left)
-					.Padding(FMargin(10.0f, 0.0f, 0.0f, 0.0f))
-					.AutoWidth()
-					[
-						AMaskProperty->CreatePropertyNameWidget()
-					]
-					+ SHorizontalBox::Slot()
-					.HAlign(HAlign_Left)
-					.AutoWidth()
-					[
-						AMaskProperty->CreatePropertyValueWidget()
-					]
-				]	
-			];
-		}
-		else
-		{	
-			if (TSharedPtr<IPropertyHandle> PropertyHandle = StackParameterData->ParameterNode->CreatePropertyHandle())
-			{
-				PropertyHandle->MarkResetToDefaultCustomized(true);
-			}
-
-			FDetailWidgetDecl* CustomNameWidget = Row.CustomNameWidget();
-			if (CustomNameWidget)
-			{
-				(*CustomNameWidget)
-				[
-					SNew(SHorizontalBox)
-					+SHorizontalBox::Slot()
-					.VAlign(VAlign_Center)
-					[
-						SNew(STextBlock)
-						.Text(NameOverride)
-						.ToolTipText(FMaterialPropertyHelpers::GetParameterExpressionDescription(StackParameterData->Parameter, MaterialEditorInstance))
-						.Font(FAppStyle::GetFontStyle(TEXT("PropertyWindow.NormalFont")))
-					]
-				];
-			}
-			else
-			{
-				Row.ToolTip(FMaterialPropertyHelpers::GetParameterExpressionDescription(StackParameterData->Parameter, MaterialEditorInstance));
-			}
-
-			bisPaddedProperty = true;
-		}
+		FMaterialPropertyHelpers::SetPropertyRowParameterWidget(Row, StackParameterData->Parameter, PropertyHandle, MaterialEditorInstance);
 
 		FNodeWidgets NodeWidgets = Node.CreateNodeWidgets();
 		LeftSideWidget = NodeWidgets.NameWidget.ToSharedRef();
 		RightSideWidget = NodeWidgets.ValueWidget.ToSharedRef();
+
+		RightSidePadding = FMargin(12.0f, 0.0f, 2.0f, 0.0f);	// Match DetailWidgetConstants::RightRowPadding
 	}
 // END PROPERTY
 
@@ -465,6 +106,8 @@ void SMaterialParametersOverviewTreeItem::Construct(const FArguments& InArgs, co
 		FNodeWidgets NodeWidgets = StackParameterData->ParameterNode->CreateNodeWidgets();
 		LeftSideWidget = NodeWidgets.NameWidget.ToSharedRef();
 		RightSideWidget = NodeWidgets.ValueWidget.ToSharedRef();
+
+		RightSidePadding = FMargin(12.0f, 0.0f, 2.0f, 0.0f);	// Match DetailWidgetConstants::RightRowPadding
 	}
 // END PROPERTY CHILD
 
@@ -473,7 +116,6 @@ void SMaterialParametersOverviewTreeItem::Construct(const FArguments& InArgs, co
 	{
 		FDetailColumnSizeData& ColumnSizeData = InArgs._InTree->GetColumnSizeData();
 
-		float ValuePadding = bisPaddedProperty ? 20.0f : 0.0f;
 		WrapperWidget->AddSlot()
 			.AutoHeight()
 			[
@@ -482,45 +124,49 @@ void SMaterialParametersOverviewTreeItem::Construct(const FArguments& InArgs, co
 				.Padding(FMargin(0, 0, 0, 1))
 				[
 					SNew(SBorder)
-					.Padding(3.0f)
+					.Padding(0.0f)
 					.BorderImage(this, &SMaterialParametersOverviewTreeItem::GetBorderImage)
 					.BorderBackgroundColor(this, &SMaterialParametersOverviewTreeItem::GetOuterBackgroundColor, StackParameterData)
 					[
-						SNew(SSplitter)
-						.Style(FAppStyle::Get(), "DetailsView.Splitter")
-						.PhysicalSplitterHandleSize(1.0f)
-						.HitDetectionSplitterHandleSize(5.0f)
-						+ SSplitter::Slot()
-						.Value(ColumnSizeData.GetNameColumnWidth())
-						.OnSlotResized(ColumnSizeData.GetOnNameColumnResized())
-						.Value(0.25f)
+						SNew(SBox)
+						.MinDesiredHeight(26.0f)	// Match PropertyEditorConstants::PropertyRowHeight
 						[
-							SNew(SHorizontalBox)
-							+ SHorizontalBox::Slot()
-							.AutoWidth()
-							.VAlign(VAlign_Center)
-							.Padding(FMargin(3.0f))
-							[
-								SNew(SExpanderArrow, SharedThis(this))
-							]
-							+ SHorizontalBox::Slot()
-							.Padding(FMargin(2.0f))
-							.VAlign(VAlign_Center)
+							SNew(SSplitter)
+							.Style(FAppStyle::Get(), "DetailsView.Splitter")
+							.PhysicalSplitterHandleSize(1.0f)
+							.HitDetectionSplitterHandleSize(5.0f)
+							+ SSplitter::Slot()
+							.Value(ColumnSizeData.GetNameColumnWidth())
+							.OnSlotResized(ColumnSizeData.GetOnNameColumnResized())
+							.Value(0.25f)
 							[
-								LeftSideWidget
+								SNew(SHorizontalBox)
+								+ SHorizontalBox::Slot()
+								.AutoWidth()
+								.VAlign(VAlign_Center)
+								.Padding(FMargin(3.0f))
+								[
+									SNew(SExpanderArrow, SharedThis(this))
+								]
+								+ SHorizontalBox::Slot()
+								.Padding(FMargin(2.0f))
+								.VAlign(VAlign_Center)
+								[
+									LeftSideWidget
+								]
 							]
-						]
-						+ SSplitter::Slot()
-						.Value(ColumnSizeData.GetValueColumnWidth())
-						.OnSlotResized(ColumnSizeData.GetOnValueColumnResized())
-						[
-							SNew(SHorizontalBox)
-							+ SHorizontalBox::Slot()
-							.Padding(FMargin(5.0f, 2.0f, 5.0, 2.0f))
-							.HAlign(HAlign_Left)
-							.VAlign(VAlign_Center)
+							+ SSplitter::Slot()
+							.Value(ColumnSizeData.GetValueColumnWidth())
+							.OnSlotResized(ColumnSizeData.GetOnValueColumnResized())
 							[
-								RightSideWidget
+								SNew(SHorizontalBox)
+								+ SHorizontalBox::Slot()
+								.HAlign(HAlign_Left)
+								.VAlign(VAlign_Center)
+								.Padding(RightSidePadding)
+								[
+									RightSideWidget
+								]
 							]
 						]
 					]
@@ -677,6 +323,12 @@ void SMaterialParametersOverviewTree::CreateGroupsWidget()
 				UDEditorParameterValue* Parameter = ParameterGroup.Parameters[ParamIdx];
 				if (Parameter->ParameterInfo.Association == EMaterialParameterAssociation::GlobalParameter)
 				{
+					// Skip parameters that shouldn't have widgets here
+					if (!FMaterialPropertyHelpers::ShouldCreatePropertyRowForParameter(Parameter))
+					{
+						continue;
+					}
+
 					bHasAnyParameters = true;
 					TSharedPtr<IPropertyHandle> ParametersArrayProperty = ChildHandle->GetChildHandle("Parameters");
 					TSharedPtr<IPropertyHandle> ParameterProperty = ParametersArrayProperty->GetChildHandle(ParamIdx);
@@ -684,39 +336,6 @@ void SMaterialParametersOverviewTree::CreateGroupsWidget()
 
 
 					FUnsortedParamData NonLayerProperty;
-					UDEditorScalarParameterValue* ScalarParam = Cast<UDEditorScalarParameterValue>(Parameter);
-					UDEditorVectorParameterValue* VectorParam = Cast<UDEditorVectorParameterValue>(Parameter);
-
-					if (ScalarParam && ScalarParam->SliderMax > ScalarParam->SliderMin)
-					{
-						ParameterValueProperty->SetInstanceMetaData("UIMin", FString::Printf(TEXT("%f"), ScalarParam->SliderMin));
-						ParameterValueProperty->SetInstanceMetaData("UIMax", FString::Printf(TEXT("%f"), ScalarParam->SliderMax));
-					}
-
-					if (VectorParam)
-					{
-						static const FName Red("R");
-						static const FName Green("G");
-						static const FName Blue("B");
-						static const FName Alpha("A");
-						if (!VectorParam->ChannelNames.R.IsEmpty())
-						{
-							ParameterProperty->GetChildHandle(Red)->SetPropertyDisplayName(VectorParam->ChannelNames.R);
-						}
-						if (!VectorParam->ChannelNames.G.IsEmpty())
-						{
-							ParameterProperty->GetChildHandle(Green)->SetPropertyDisplayName(VectorParam->ChannelNames.G);
-						}
-						if (!VectorParam->ChannelNames.B.IsEmpty())
-						{
-							ParameterProperty->GetChildHandle(Blue)->SetPropertyDisplayName(VectorParam->ChannelNames.B);
-						}
-						if (!VectorParam->ChannelNames.A.IsEmpty())
-						{
-							ParameterProperty->GetChildHandle(Alpha)->SetPropertyDisplayName(VectorParam->ChannelNames.A);
-						}
-					}
-
 					NonLayerProperty.Parameter = Parameter;
 					NonLayerProperty.ParameterGroup = ParameterGroup;
 					NonLayerProperty.UnsortedName = Parameter->ParameterInfo.Name;
@@ -769,23 +388,25 @@ void SMaterialParametersOverviewTree::ShowSubParameters()
 	{
 		UDEditorParameterValue* Parameter = Property.Parameter;
 		{
-			TSharedPtr<FSortedParamData> GroupProperty(new FSortedParamData());
-			GroupProperty->StackDataType = EStackDataType::Group;
-			GroupProperty->ParameterInfo.Index = Parameter->ParameterInfo.Index;
-			GroupProperty->ParameterInfo.Association = Parameter->ParameterInfo.Association;
-			GroupProperty->Group = Property.ParameterGroup;
-			GroupProperty->NodeKey = FString::FromInt(GroupProperty->ParameterInfo.Index) + FString::FromInt(GroupProperty->ParameterInfo.Association) + Property.ParameterGroup.GroupName.ToString();
+			FString GroupNodeKey = FString::FromInt(Parameter->ParameterInfo.Index) + FString::FromInt(Parameter->ParameterInfo.Association) + Property.ParameterGroup.GroupName.ToString();
 
 			bool bAddNewGroup = true;
 			for (TSharedPtr<struct FSortedParamData> GroupChild : SortedParameters)
 			{
-				if (GroupChild->NodeKey == GroupProperty->NodeKey)
+				if (GroupChild->NodeKey == GroupNodeKey)
 				{
 					bAddNewGroup = false;
 				}
 			}
 			if (bAddNewGroup)
 			{
+				TSharedPtr<FSortedParamData> GroupProperty(new FSortedParamData());
+				GroupProperty->StackDataType = EStackDataType::Group;
+				GroupProperty->ParameterInfo.Index = Parameter->ParameterInfo.Index;
+				GroupProperty->ParameterInfo.Association = Parameter->ParameterInfo.Association;
+				GroupProperty->Group = Property.ParameterGroup;
+				GroupProperty->NodeKey = GroupNodeKey;
+
 				SortedParameters.Add(GroupProperty);
 			}
 
@@ -799,9 +420,8 @@ void SMaterialParametersOverviewTree::ShowSubParameters()
 			ChildProperty->NodeKey = FString::FromInt(ChildProperty->ParameterInfo.Index) + FString::FromInt(ChildProperty->ParameterInfo.Association) +  Property.ParameterGroup.GroupName.ToString() + Property.UnsortedName.ToString();
 
 
-			UDEditorStaticComponentMaskParameterValue* CompMaskParam = Cast<UDEditorStaticComponentMaskParameterValue>(Parameter);
 			// No children for masks
-			if (!CompMaskParam)
+			if (FMaterialPropertyHelpers::ShouldCreateChildPropertiesForParameter(Parameter))
 			{
 				TArray<TSharedRef<IDetailTreeNode>> ParamChildren;
 				Property.ParameterNode->GetChildren(ParamChildren);
diff --git a/Engine/Source/Editor/MaterialEditor/Private/SMaterialParametersOverviewWidget.h b/Engine/Source/Editor/MaterialEditor/Private/SMaterialParametersOverviewWidget.h
--- a/Engine/Source/Editor/MaterialEditor/Private/SMaterialParametersOverviewWidget.h
+++ b/Engine/Source/Editor/MaterialEditor/Private/SMaterialParametersOverviewWidget.h
@@ -46,8 +46,6 @@ public:
 	void Construct(const FArguments& InArgs, const TSharedRef<STableViewBase>& InOwnerTableView);
 
 private:
-	FString GetCurvePath(class UDEditorScalarParameterValue* Parameter) const;
-
 	const FSlateBrush* GetBorderImage() const;
 
 	FSlateColor GetOuterBackgroundColor(TSharedPtr<FSortedParamData> InParamData) const;
diff --git a/Engine/Source/Editor/MaterialEditor/Public/MaterialPropertyHelpers.h b/Engine/Source/Editor/MaterialEditor/Public/MaterialPropertyHelpers.h
--- a/Engine/Source/Editor/MaterialEditor/Public/MaterialPropertyHelpers.h
+++ b/Engine/Source/Editor/MaterialEditor/Public/MaterialPropertyHelpers.h
@@ -6,6 +6,7 @@
 #include "Types/SlateEnums.h"
 #include "Layout/Visibility.h"
 #include "Materials/MaterialLayersFunctions.h"
+#include "IDetailPropertyRow.h"
 #include "Input/Reply.h"
 #include "Widgets/Layout/SSplitter.h"
 #include "MaterialEditor/MaterialEditorInstanceConstant.h"
@@ -19,6 +20,7 @@
 struct FAssetData;
 class IDetailGroup;
 class IDetailLayoutBuilder;
+class IDetailPropertyRow;
 class IDetailTreeNode;
 class IPropertyHandle;
 class UDEditorParameterValue;
@@ -150,6 +152,38 @@ public:
 	TWeakPtr<class IDraggableItem> OwningStack;
 };
 
+/** Overridable properties on material parameter property rows. */
+struct FMaterialParameterPropertyRowOverrides
+{
+public:
+	/**
+	 * Override the default EditConditionValue attribute.
+	 * If not specified, FMaterialPropertyHelpers::IsOverriddenExpression will
+	 * be used.
+	 */
+	TAttribute<bool> EditConditionValue;
+
+	/**
+	 * Override the default OnEditConditionValueChanged delegate.
+	 * If not specified, FMaterialPropertyHelpers::OnOverrideParameter will be
+	 * used.
+	 */
+	TOptional<FOnBooleanValueChanged> OnEditConditionValueChanged;
+
+	/**
+	 * Override the default Visibility attribute.
+	 * If not specified, FMaterialPropertyHelpers::ShouldShowExpression will be
+	 * used.
+	 */
+	TAttribute<EVisibility> Visibility;
+
+	/**
+	 * If Visibility is not overridden, this delegate will be passed as the
+	 * last argument to FMaterialPropertyHelpers::ShouldShowExpression.
+	 */
+	FGetShowHiddenParameters ShowHiddenDelegate;
+};
+
 /*-----------------------------------------------------------------------------
    FMaterialInstanceBaseParameterDetails
 -----------------------------------------------------------------------------*/
@@ -157,6 +191,18 @@ public:
 class FMaterialPropertyHelpers
 {
 public:
+	static UE_API bool ShouldCreatePropertyRowForParameter(UDEditorParameterValue* Parameter);
+	static UE_API bool ShouldCreateChildPropertiesForParameter(UDEditorParameterValue* Parameter);
+	static UE_API void ConfigurePropertyRowForParameter(IDetailPropertyRow& PropertyRow, UDEditorParameterValue* Parameter, TSharedPtr<IPropertyHandle> ParameterValueProperty, UMaterialEditorInstanceConstant* MaterialEditorInstance, const FMaterialParameterPropertyRowOverrides& RowOverrides = FMaterialParameterPropertyRowOverrides());
+
+	static UE_API void SetPropertyRowParameterWidget(IDetailPropertyRow& PropertyRow, UDEditorParameterValue* Parameter, TSharedPtr<IPropertyHandle> ParameterProperty, UMaterialEditorParameters* MaterialEditorParameters);
+	static UE_API void SetPropertyRowMaskParameterWidget(IDetailPropertyRow& PropertyRow, UDEditorParameterValue* Parameter, TSharedPtr<IPropertyHandle> ParameterValueProperty);
+	static UE_API void SetPropertyRowVectorChannelMaskParameterWidget(IDetailPropertyRow& PropertyRow, UDEditorParameterValue* Parameter, TSharedPtr<IPropertyHandle> ParameterValueProperty, UMaterialEditorParameters* MaterialEditorParameters);
+	static UE_API void SetPropertyRowScalarAtlasPositionParameterWidget(IDetailPropertyRow& PropertyRow, UDEditorParameterValue* Parameter, TSharedPtr<IPropertyHandle> ParameterValueProperty, UMaterialEditorParameters* MaterialEditorParameters);
+	static UE_API void SetPropertyRowScalarEnumerationParameterWidget(IDetailPropertyRow& PropertyRow, UDEditorParameterValue* Parameter, TSharedPtr<IPropertyHandle> ParameterValueProperty, const TSoftObjectPtr<UObject>& SoftEnumeration);
+	static UE_API void SetPropertyRowTextureParameterWidget(IDetailPropertyRow& PropertyRow, UDEditorParameterValue* Parameter, TSharedPtr<IPropertyHandle> ParameterValueProperty, UMaterialEditorParameters* MaterialEditorParameters);
+	static UE_API void SetPropertyRowParameterCollectionParameterWidget(IDetailPropertyRow& PropertyRow, UDEditorParameterValue* Parameter, TSharedPtr<IPropertyHandle> ParameterValueProperty, class UMaterialParameterCollection* BaseParameterCollection, TArray<TSharedRef<IPropertyHandle>>&& AllPropertyHandles);
+
 	/** Returns true if the parameter is being overridden */
 	static UE_API bool IsOverriddenExpression(UDEditorParameterValue* Parameter);
 	static bool IsOverriddenExpression(TObjectPtr<UDEditorParameterValue> Parameter) { return IsOverriddenExpression(Parameter.Get()); }
@@ -177,6 +223,7 @@ public:
 	static EVisibility ShouldShowExpression(TObjectPtr<UDEditorParameterValue> Parameter, UMaterialEditorInstanceConstant* MaterialEditorInstance, FGetShowHiddenParameters ShowHiddenDelegate) { return ShouldShowExpression(Parameter.Get(), MaterialEditorInstance, ShowHiddenDelegate); }
 
 	/** Generic material property reset to default implementation.  Resets Parameter to default */
+	static UE_API FResetToDefaultOverride CreateResetToDefaultOverride(UDEditorParameterValue* Parameter, UMaterialEditorInstanceConstant* MaterialEditorInstance);
 	static UE_API void ResetToDefault(UDEditorParameterValue* Parameter, UMaterialEditorInstanceConstant* MaterialEditorInstance);
 	static void ResetToDefault(TObjectPtr<UDEditorParameterValue> Parameter, UMaterialEditorInstanceConstant* MaterialEditorInstance) { ResetToDefault(Parameter.Get(), MaterialEditorInstance); }
 	static UE_API bool ShouldShowResetToDefault(UDEditorParameterValue* Parameter, UMaterialEditorInstanceConstant* MaterialEditorInstance);
diff --git a/Engine/Source/Editor/UnrealEd/Classes/MaterialEditor/DEditorScalarParameterValue.h b/Engine/Source/Editor/UnrealEd/Classes/MaterialEditor/DEditorScalarParameterValue.h
--- a/Engine/Source/Editor/UnrealEd/Classes/MaterialEditor/DEditorScalarParameterValue.h
+++ b/Engine/Source/Editor/UnrealEd/Classes/MaterialEditor/DEditorScalarParameterValue.h
@@ -66,5 +66,7 @@ public:
 		}
 		return false;
 	}
+
+	FString GetAtlasCurvePath() const { return AtlasData.Curve.ToString(); }
 };
 
