Disable the ZoomToFocus behavior when creating a usage node from the declaration

- Don't zoom the focus of the editor since this was visually jarring for the user.

#tests Creating declaration node, right clicking and selecting "create from usage"
[RN] minor | Rendering - Materials and Shaders

diff --git a/Engine/Source/Editor/MaterialEditor/Private/MaterialEditor.cpp b/Engine/Source/Editor/MaterialEditor/Private/MaterialEditor.cpp
--- a/Engine/Source/Editor/MaterialEditor/Private/MaterialEditor.cpp
+++ b/Engine/Source/Editor/MaterialEditor/Private/MaterialEditor.cpp
@@ -4640,7 +4640,6 @@ void FMaterialEditor::OnCreateRerouteUsageFromDeclaration()
 					}
 				}
 			}
-			FocusedGraphEd->ZoomToFit(true);
 		}
 	}
 }
