Extract MPC overrides to their own editor category. They can only be overridden as one unit, making it confusing to expose as a regular parameter that can be modified per-layer.

#rnx

diff --git a/Engine/Source/Editor/MaterialEditor/Private/MaterialEditorInstanceDetailCustomization.cpp b/Engine/Source/Editor/MaterialEditor/Private/MaterialEditorInstanceDetailCustomization.cpp
--- a/Engine/Source/Editor/MaterialEditor/Private/MaterialEditorInstanceDetailCustomization.cpp
+++ b/Engine/Source/Editor/MaterialEditor/Private/MaterialEditorInstanceDetailCustomization.cpp
@@ -145,6 +145,7 @@ void FMaterialInstanceParameterDetails::CustomizeDetails(IDetailLayoutBuilder& D
 	else
 	{
 		CreateGroupsWidget(ParameterGroupsProperty, GroupsCategory);
+		CreateParameterCollectionOverrideCategory(ParameterGroupsProperty, DetailLayout);
 
 		// Create default category for class properties
 		const FName DefaultCategoryName = NAME_None;
@@ -376,7 +377,7 @@ void FMaterialInstanceParameterDetails::CreateGroupsWidget(TSharedRef<IPropertyH
 			for (int32 ParamIdx = 0; ParamIdx < ParameterGroup.Parameters.Num() && !bCreateGroup; ++ParamIdx)
 			{
 				UDEditorParameterValue* Parameter = ParameterGroup.Parameters[ParamIdx];
-				const bool bIsVisible = MaterialEditorInstance->VisibleExpressions.Contains(Parameter->ParameterInfo) && !FMaterialPropertyHelpers::UsesCustomPrimitiveData(Parameter);
+				const bool bIsVisible = MaterialEditorInstance->VisibleExpressions.Contains(Parameter->ParameterInfo) && FMaterialPropertyHelpers::ShouldCreatePropertyRowForParameter(Parameter);
 				bCreateGroup = bIsVisible && (!MaterialEditorInstance->bShowOnlyOverrides || FMaterialPropertyHelpers::IsOverriddenExpression(Parameter));
 			}
 		
@@ -473,6 +474,93 @@ void FMaterialInstanceParameterDetails::CreateGroupsWidget(TSharedRef<IPropertyH
 	}
 }
 
+void FMaterialInstanceParameterDetails::CreateParameterCollectionOverrideCategory(TSharedRef<IPropertyHandle> ParameterGroupsProperty, IDetailLayoutBuilder& DetailLayout)
+{
+	if (!MaterialEditorInstance || !MaterialEditorInstance->SourceInstance)
+	{
+		return;
+	}
+
+	struct FParameterEntry
+	{
+		UDEditorParameterCollectionParameterValue* ParameterValue;
+		TSharedPtr<IPropertyHandle> ParameterValueProperty;
+	};
+
+	TMap<UMaterialParameterCollection*, TArray<FParameterEntry>> ParameterCollectionSharedOverrideMap;
+
+	for (int32 GroupIdx = 0; GroupIdx < MaterialEditorInstance->ParameterGroups.Num(); ++GroupIdx)
+	{
+		TSharedPtr<IPropertyHandle> ParameterGroupProperty = ParameterGroupsProperty->GetChildHandle(GroupIdx)->GetChildHandle("Parameters");
+
+		FEditorParameterGroup& ParameterGroup = MaterialEditorInstance->ParameterGroups[GroupIdx];
+		for (int32 ParamIdx = 0; ParamIdx < ParameterGroup.Parameters.Num(); ++ParamIdx)
+		{
+			UDEditorParameterValue* Parameter = ParameterGroup.Parameters[ParamIdx];
+			if (UDEditorParameterCollectionParameterValue* ParameterCollectionParam = Cast<UDEditorParameterCollectionParameterValue>(Parameter))
+			{
+				// Map each parameter collection to its referencing parameters.
+				for (const FMaterialParameterCollectionInfo& CollectionInfo : MaterialEditorInstance->SourceInstance->GetCachedExpressionData().ParameterCollectionInfos)
+				{
+					FMaterialParameterMetadata Result;
+					if (ParameterCollectionParam->GetValue(Result) && ensure(Result.Value.Type == EMaterialParameterType::ParameterCollection) &&
+						Result.Value.ParameterCollection && Result.Value.ParameterCollection->IsInstanceOf(CollectionInfo.StateId))
+					{
+						TSharedPtr<IPropertyHandle> ParameterValueProperty = ParameterGroupProperty->GetChildHandle(ParamIdx)->GetChildHandle("ParameterValue");
+
+						TArray<FParameterEntry>& ParameterIndices = ParameterCollectionSharedOverrideMap.FindOrAdd(CollectionInfo.ParameterCollection);
+						ParameterIndices.Add({ ParameterCollectionParam, ParameterValueProperty });
+						break;
+					}
+				}
+			}
+		}
+	}
+
+	if (ParameterCollectionSharedOverrideMap.Num() == 0)
+	{
+		return;
+	}
+
+	// Make a new category for MPC overrides
+	FName ParamCollectionOverridesName = TEXT("ParameterCollectionOverrides");
+	IDetailCategoryBuilder& DetailGroup = DetailLayout.EditCategory(ParamCollectionOverridesName, LOCTEXT("MICParamCollectionOverridesGroupTitle", "Parameter Collection Overrides"));
+
+	FMaterialParameterPropertyRowOverrides RowOverrides;
+	RowOverrides.ShowHiddenDelegate = ShowHiddenDelegate;
+
+	for (auto& Entry : ParameterCollectionSharedOverrideMap)
+	{
+		UMaterialParameterCollection* ParameterCollection = Entry.Key;
+		const TArray<FParameterEntry>& ParameterEntries = Entry.Value;
+		int32 NumParameterEntries = ParameterEntries.Num();
+		if (NumParameterEntries > 0)
+		{
+			const FParameterEntry& FirstParameterEntry = ParameterEntries[0];
+
+			if (FirstParameterEntry.ParameterValueProperty)
+			{
+				// Gather the relevant property handles of the collection's referencing parameters.
+				TArray<TSharedRef<IPropertyHandle>> ValuePropertyHandles{ FirstParameterEntry.ParameterValueProperty.ToSharedRef() };
+				for (int32 ParameterIndex = 1; ParameterIndex < NumParameterEntries; ++ParameterIndex)
+				{
+					const FParameterEntry& ParameterEntry = ParameterEntries[ParameterIndex];
+					if (ParameterEntry.ParameterValueProperty)
+					{
+						ValuePropertyHandles.Add(ParameterEntry.ParameterValueProperty.ToSharedRef());
+					}
+				}
+
+				// Create a parameter value widget for the first parameter.
+				IDetailPropertyRow& PropertyRow = DetailGroup.AddProperty(FirstParameterEntry.ParameterValueProperty.ToSharedRef());
+
+				FMaterialPropertyHelpers::ConfigurePropertyRowForParameter(PropertyRow, FirstParameterEntry.ParameterValue, FirstParameterEntry.ParameterValueProperty, MaterialEditorInstance, RowOverrides);
+				FMaterialPropertyHelpers::SetPropertyRowParameterCollectionParameterWidget(PropertyRow, FirstParameterEntry.ParameterValue, FirstParameterEntry.ParameterValueProperty, ParameterCollection, MoveTemp(ValuePropertyHandles));
+			}
+		}
+	}
+}
+
 void FMaterialInstanceParameterDetails::EnableGroupParameters(FEditorParameterGroup& ParameterGroup, bool ShouldEnable)
 {
 	// loop through each parameter in the group and toggle to enable/disable them all
@@ -487,8 +575,6 @@ void FMaterialInstanceParameterDetails::CreateSingleGroupWidget(FEditorParameter
 {
 	TSharedPtr<IPropertyHandle> ParametersArrayProperty = ParameterGroupProperty->GetChildHandle("Parameters");
 
-	TMap<UMaterialParameterCollection*, TArray<int32>> ParameterCollectionSharedOverrideMap;
-
 	FMaterialParameterPropertyRowOverrides RowOverrides;
 	RowOverrides.ShowHiddenDelegate = ShowHiddenDelegate;
 
@@ -502,27 +588,6 @@ void FMaterialInstanceParameterDetails::CreateSingleGroupWidget(FEditorParameter
 		if (ParameterProperty.IsValid() && ParameterValueProperty.IsValid() && ParameterValueProperty->IsValidHandle() &&
 			(GroupIndex == INDEX_NONE || Parameter->ParameterInfo.Index == GroupIndex))
 		{
-			if (UDEditorParameterCollectionParameterValue* ParameterCollectionParam = Cast<UDEditorParameterCollectionParameterValue>(Parameter))
-			{
-				if (MaterialEditorInstance && MaterialEditorInstance->SourceInstance)
-				{
-					// Map each parameter collection to its referencing parameters.
-					for (const FMaterialParameterCollectionInfo& CollectionInfo : MaterialEditorInstance->SourceInstance->GetCachedExpressionData().ParameterCollectionInfos)
-					{
-						FMaterialParameterMetadata Result;
-						if (ParameterCollectionParam->GetValue(Result) && ensure(Result.Value.Type == EMaterialParameterType::ParameterCollection) &&
-							Result.Value.ParameterCollection && Result.Value.ParameterCollection->IsInstanceOf(CollectionInfo.StateId))
-						{
-							TArray<int32>& ParameterIndices = ParameterCollectionSharedOverrideMap.FindOrAdd(CollectionInfo.ParameterCollection);
-							ParameterIndices.Add(ParamIdx);
-							break;
-						}
-					}
-				}
-
-				continue;
-			}
-
 			if (Parameter->ParameterInfo.Association == EMaterialParameterAssociation::GlobalParameter || bForceShowParam)
 			{
 				if (!FMaterialPropertyHelpers::ShouldCreatePropertyRowForParameter(Parameter))
@@ -537,40 +602,6 @@ void FMaterialInstanceParameterDetails::CreateSingleGroupWidget(FEditorParameter
 			}
 		}
 	}
-
-	for (auto& Entry : ParameterCollectionSharedOverrideMap)
-	{
-		UMaterialParameterCollection* ParameterCollection = Entry.Key;
-		const TArray<int32>& ParameterIndices = Entry.Value;
-		int32 NumParameterIndices = ParameterIndices.Num();
-		if (NumParameterIndices > 0)
-		{
-			UDEditorParameterValue* FirstParameter = ParameterGroup.Parameters[ParameterIndices[0]];
-			TSharedPtr<IPropertyHandle> FirstParameterProperty = ParametersArrayProperty->GetChildHandle(ParameterIndices[0]);
-
-			if (TSharedPtr<IPropertyHandle> FirstParameterValueProperty = FirstParameterProperty->GetChildHandle("ParameterValue"))
-			{
-				// Gather the relevant property handles of the collection's referencing parameters.
-				TArray<TSharedRef<IPropertyHandle>> PropertyHandles{ FirstParameterValueProperty.ToSharedRef() };
-				for (int32 ParameterIndex = 1; ParameterIndex < NumParameterIndices; ++ParameterIndex)
-				{
-					if (TSharedPtr<IPropertyHandle> ParameterProperty = ParametersArrayProperty->GetChildHandle(ParameterIndices[ParameterIndex]))
-					{
-						if (TSharedPtr<IPropertyHandle> ParameterValueProperty = ParameterProperty->GetChildHandle("ParameterValue"))
-						{
-							PropertyHandles.Add(ParameterValueProperty.ToSharedRef());
-						}
-					}
-				}
-
-				// Create a parameter value widget for the first parameter.
-				IDetailPropertyRow& PropertyRow = DetailGroup.AddPropertyRow(FirstParameterValueProperty.ToSharedRef());
-
-				FMaterialPropertyHelpers::ConfigurePropertyRowForParameter(PropertyRow, FirstParameter, FirstParameterValueProperty, MaterialEditorInstance, RowOverrides);
-				FMaterialPropertyHelpers::SetPropertyRowParameterCollectionParameterWidget(PropertyRow, FirstParameter, FirstParameterValueProperty, ParameterCollection, MoveTemp(PropertyHandles));
-			}
-		}
-	}
 }
 
 bool FMaterialInstanceParameterDetails::OnShouldSetAsset(const FAssetData& AssetData) const
diff --git a/Engine/Source/Editor/MaterialEditor/Private/MaterialEditorInstanceDetailCustomization.h b/Engine/Source/Editor/MaterialEditor/Private/MaterialEditorInstanceDetailCustomization.h
--- a/Engine/Source/Editor/MaterialEditor/Private/MaterialEditorInstanceDetailCustomization.h
+++ b/Engine/Source/Editor/MaterialEditor/Private/MaterialEditorInstanceDetailCustomization.h
@@ -42,6 +42,7 @@ private:
 
 	/** Builds the custom parameter groups category */
 	void CreateGroupsWidget(TSharedRef<IPropertyHandle> ParameterGroupsProperty, class IDetailCategoryBuilder& GroupsCategory);
+	void CreateParameterCollectionOverrideCategory(TSharedRef<IPropertyHandle> ParameterGroupsProperty, class IDetailLayoutBuilder& DetailLayout);
 
 	/** Builds the widget for an individual parameter group */
 	void CreateSingleGroupWidget(struct FEditorParameterGroup& ParameterGroup, TSharedPtr<IPropertyHandle> ParameterGroupProperty, class IDetailGroup& DetailGroup, int32 GroupIndex = -1, bool bForceShowParam = false);
diff --git a/Engine/Source/Editor/MaterialEditor/Private/MaterialPropertyHelpers.cpp b/Engine/Source/Editor/MaterialEditor/Private/MaterialPropertyHelpers.cpp
--- a/Engine/Source/Editor/MaterialEditor/Private/MaterialPropertyHelpers.cpp
+++ b/Engine/Source/Editor/MaterialEditor/Private/MaterialPropertyHelpers.cpp
@@ -15,6 +15,7 @@
 #include "Materials/MaterialInterface.h"
 #include "MaterialEditor/DEditorFontParameterValue.h"
 #include "MaterialEditor/DEditorMaterialLayersParameterValue.h"
+#include "MaterialEditor/DEditorParameterCollectionParameterValue.h"
 #include "MaterialEditor/DEditorRuntimeVirtualTextureParameterValue.h"
 #include "MaterialEditor/DEditorScalarParameterValue.h"
 #include "MaterialEditor/DEditorStaticComponentMaskParameterValue.h"
@@ -558,15 +559,15 @@ bool FMaterialPropertyHelpers::ShouldCreatePropertyRowForParameter(UDEditorParam
 		// No widget for layers
 		return false;
 	}
-	else if (UDEditorVectorParameterValue* VectorParam = Cast<UDEditorVectorParameterValue>(Parameter))
+	else if (UDEditorParameterCollectionParameterValue* ParamCollectionParam = Cast<UDEditorParameterCollectionParameterValue>(Parameter))
 	{
-		// No widget for CPD
-		return !VectorParam->bUseCustomPrimitiveData;
+		// No widget for MPCs, special override category in MI
+		return false;
 	}
-	else if (UDEditorScalarParameterValue* ScalarParam = Cast<UDEditorScalarParameterValue>(Parameter))
+	else if (FMaterialPropertyHelpers::UsesCustomPrimitiveData(Parameter))
 	{
 		// No widget for CPD
-		return !ScalarParam->bUseCustomPrimitiveData;
+		return false;
 	}
 
 	return true;
