Fix width of the asset picker for the Parent property in material instances

#rnx

diff --git a/Engine/Source/Editor/MaterialEditor/Private/MaterialEditorInstanceDetailCustomization.cpp b/Engine/Source/Editor/MaterialEditor/Private/MaterialEditorInstanceDetailCustomization.cpp
--- a/Engine/Source/Editor/MaterialEditor/Private/MaterialEditorInstanceDetailCustomization.cpp
+++ b/Engine/Source/Editor/MaterialEditor/Private/MaterialEditorInstanceDetailCustomization.cpp
@@ -228,39 +228,29 @@ void FMaterialInstanceParameterDetails::CustomizeDetails(IDetailLayoutBuilder& D
 			TSharedRef<IPropertyHandle> ParentPropertyHandle = DetailLayout.GetProperty("Parent");
 			IDetailPropertyRow& ParentPropertyRow = DefaultCategory.AddProperty(ParentPropertyHandle);
 
-			ParentPropertyHandle->MarkResetToDefaultCustomized();
-	
-			TSharedPtr<SWidget> NameWidget;
-			TSharedPtr<SWidget> ValueWidget;
-			FDetailWidgetRow Row;
+			TSharedRef<SObjectPropertyEntryBox> ParentValueWidget = SNew(SObjectPropertyEntryBox)
+				.PropertyHandle(ParentPropertyHandle)
+				.AllowedClass(UMaterialInterface::StaticClass())
+				.ThumbnailPool(DetailLayout.GetThumbnailPool())
+				.AllowClear(true)
+				.OnShouldSetAsset(this, &FMaterialInstanceParameterDetails::OnShouldSetAsset);
 
-			ParentPropertyRow.GetDefaultWidgets(NameWidget, ValueWidget, Row);
-	
-			ParentPropertyHandle->ClearResetToDefaultCustomized();
+			float ParentValueMinDesiredWidth, ParentValueMaxDesiredWidth;
+			ParentValueWidget->GetDesiredWidth(ParentValueMinDesiredWidth, ParentValueMaxDesiredWidth);
 
 			const bool bShowChildren = true;
 			ParentPropertyRow.CustomWidget(bShowChildren)
 				.NameContent()
-				.MinDesiredWidth(Row.NameWidget.MinWidth)
-				.MaxDesiredWidth(Row.NameWidget.MaxWidth)
 				[
-					NameWidget.ToSharedRef()
+					ParentPropertyHandle->CreatePropertyNameWidget()
 				]
 				.ValueContent()
-				.MinDesiredWidth(Row.ValueWidget.MinWidth)
-				.MaxDesiredWidth(Row.ValueWidget.MaxWidth)
+				.MinDesiredWidth(ParentValueMinDesiredWidth)
+				.MaxDesiredWidth(ParentValueMaxDesiredWidth)
 				[
-					SNew(SObjectPropertyEntryBox)
-					.PropertyHandle(ParentPropertyHandle)
-					.AllowedClass(UMaterialInterface::StaticClass())
-					.ThumbnailPool(DetailLayout.GetThumbnailPool())
-					.AllowClear(true)
-					.OnShouldSetAsset(this, &FMaterialInstanceParameterDetails::OnShouldSetAsset)
+					ParentValueWidget
 				];
 
-			ValueWidget.Reset();
-
-
 			// Add/hide other properties
 			DetailLayout.HideProperty("LightmassSettings");
 			CreateLightmassOverrideWidgets(DetailLayout);
