Add a UI button in the Material Editor to auto fixup incorrect sampler types.

- We auto assign sampler when setting the texture on the node, but sometimes the texture can change underneath the material requiring an update.
- First discover any broken nodes, and only fixup those.
- Modify the Material and Expressions so undo/redo works for this operation.

[RN] minor | Rendering - Materials and Shaders

diff --git a/Engine/Source/Editor/MaterialEditor/Private/MaterialEditor.cpp b/Engine/Source/Editor/MaterialEditor/Private/MaterialEditor.cpp
--- a/Engine/Source/Editor/MaterialEditor/Private/MaterialEditor.cpp
+++ b/Engine/Source/Editor/MaterialEditor/Private/MaterialEditor.cpp
@@ -2098,6 +2098,7 @@ void FMaterialEditor::RegisterToolBar()
 					FToolMenuSection& Section = InSubMenu->FindOrAddSection("General");
 					Section.AddEntry(FToolMenuEntry::InitMenuEntry(FMaterialEditorCommands::Get().CleanUnusedExpressions));
 					Section.AddEntry(FToolMenuEntry::InitMenuEntry(FMaterialEditorCommands::Get().ShowHideConnectors));
+					Section.AddEntry(FToolMenuEntry::InitMenuEntry(FMaterialEditorCommands::Get().FixupMismatchedSamplers));
 				}),
 			LOCTEXT("GraphCleanup_Label", "Clean Graph"),
 			LOCTEXT("GraphCleanup_Tooltip", "Tools to help clean up graph nodes and connections"),
@@ -3576,6 +3577,11 @@ void FMaterialEditor::BindCommands()
 		FExecuteAction::CreateSP(this, &FMaterialEditor::CleanUnusedExpressions),
 		FCanExecuteAction() );
 
+	ToolkitCommands->MapAction(
+		Commands.FixupMismatchedSamplers,
+		FExecuteAction::CreateSP(this, &FMaterialEditor::FixupMismatchedSamplerExpressions),
+		FCanExecuteAction() );
+
 	ToolkitCommands->MapAction(
 		Commands.ShowHideConnectors,
 		FExecuteAction::CreateSP(this, &FMaterialEditor::OnHideConnectors),
@@ -8040,6 +8046,48 @@ void FMaterialEditor::CleanUnusedExpressions()
 	}
 }
 
+void FMaterialEditor::FixupMismatchedSamplerExpressions()
+{
+	// Find all expressions that need to be fixed.
+	TArray<UMaterialExpressionTextureBase*> ExpressionsToFix;
+	for (UMaterialExpression* MaterialExpression : Material->GetExpressions())
+	{
+		if (UMaterialExpressionTextureBase* TextureBase = Cast<UMaterialExpressionTextureBase>(MaterialExpression))
+		{
+			FString SamplerTypeError;
+			if (!UMaterialExpressionTextureBase::VerifySamplerType(GEditor->PreviewPlatform.GetShaderPlatform(), nullptr, TextureBase->Texture, TextureBase->SamplerType, SamplerTypeError))
+			{
+				ExpressionsToFix.Add(TextureBase);
+			}
+		}
+	}
+
+	if (ExpressionsToFix.Num() > 0)
+	{
+		if (EAppReturnType::Yes != FMessageDialog::Open(EAppMsgType::YesNo,
+														FText::Format(
+															NSLOCTEXT("UnrealEd", "Prompt_MaterialEditorFixupTextureSamplers", "There are \"{0}\" texture sampler(s) with the incorrect sampler picked for their texture.\nWould you like to automatically fix them up?"),
+															ExpressionsToFix.Num())))
+		{
+			return;
+		}
+
+		const FScopedTransaction Transaction(NSLOCTEXT("UnrealEd", "MaterialEditorFixupSamplerExpressions", "Material Editor: Fixup Sampler Expressions"));
+			
+		{
+			ModifyMaterial();
+
+			for (UMaterialExpressionTextureBase* TextureBase : ExpressionsToFix)
+			{
+				TextureBase->Modify();
+				TextureBase->AutoSetSampleType();
+			}
+		}
+
+		UpdateMaterialAfterGraphChange();
+	}
+}
+
 void FMaterialEditor::OnWorldCleanup(UWorld* World, bool bSessionEnded, bool bCleanupResources)
 {
 	// no-op if there are no derived materials used
diff --git a/Engine/Source/Editor/MaterialEditor/Private/MaterialEditor.h b/Engine/Source/Editor/MaterialEditor/Private/MaterialEditor.h
--- a/Engine/Source/Editor/MaterialEditor/Private/MaterialEditor.h
+++ b/Engine/Source/Editor/MaterialEditor/Private/MaterialEditor.h
@@ -935,6 +935,11 @@ private:
 	 */
 	void CleanUnusedExpressions();
 
+	/**
+	* Fixes up the samplers for all broken texture sample expressions.
+	*/
+	void FixupMismatchedSamplerExpressions();
+
 	/**
 	 * Cleans up any referenced resources.
 	 */
diff --git a/Engine/Source/Editor/MaterialEditor/Private/MaterialEditorActions.cpp b/Engine/Source/Editor/MaterialEditor/Private/MaterialEditorActions.cpp
--- a/Engine/Source/Editor/MaterialEditor/Private/MaterialEditorActions.cpp
+++ b/Engine/Source/Editor/MaterialEditor/Private/MaterialEditorActions.cpp
@@ -42,6 +42,7 @@ void FMaterialEditorCommands::RegisterCommands()
 
 	UI_COMMAND( CameraHome, "Home", "Goes home on the canvas.", EUserInterfaceActionType::Button, FInputChord() );
 	UI_COMMAND( CleanUnusedExpressions, "Clean Up", "Cleans up any unused Expressions.", EUserInterfaceActionType::Button, FInputChord() );
+	UI_COMMAND( FixupMismatchedSamplers, "Fixup mismatched samplers", "Automatically fixes all texture samplers with the wrong type for their texture.", EUserInterfaceActionType::Button, FInputChord());
 	UI_COMMAND( ShowHideConnectors, "Hide Unconnected Pins", "Hides all pins with no connections", EUserInterfaceActionType::ToggleButton, FInputChord());
 	UI_COMMAND( ToggleLivePreview, "Preview Material", "Toggles real time update of the preview material.", EUserInterfaceActionType::ToggleButton, FInputChord());
 	UI_COMMAND( ToggleRealtimeExpressions, "Realtime Nodes", "Nodes impacted by time-based functions such as panners, etc. will update realtime.", EUserInterfaceActionType::ToggleButton, FInputChord() );
diff --git a/Engine/Source/Editor/MaterialEditor/Public/MaterialEditorActions.h b/Engine/Source/Editor/MaterialEditor/Public/MaterialEditorActions.h
--- a/Engine/Source/Editor/MaterialEditor/Public/MaterialEditorActions.h
+++ b/Engine/Source/Editor/MaterialEditor/Public/MaterialEditorActions.h
@@ -89,6 +89,9 @@ public:
 	
 	/** Removes any unused nodes */
 	TSharedPtr< FUICommandInfo > CleanUnusedExpressions;
+
+	/** Fixes up broken sampler nodes */
+	TSharedPtr<FUICommandInfo> FixupMismatchedSamplers;
 	
 	/** Shows or hides unused connectors */
 	TSharedPtr< FUICommandInfo > ShowHideConnectors;
