Revert due to a wrong manipulation

diff --git a/Engine/Source/Editor/MaterialEditor/Private/MaterialExpressionClasses.cpp b/Engine/Source/Editor/MaterialEditor/Private/MaterialExpressionClasses.cpp
--- a/Engine/Source/Editor/MaterialEditor/Private/MaterialExpressionClasses.cpp
+++ b/Engine/Source/Editor/MaterialEditor/Private/MaterialExpressionClasses.cpp
@@ -111,10 +111,10 @@ void MaterialExpressionClasses::InitMaterialExpressionClasses()
 						}
 
 						// Hide node types that are tagged private
-						//if(Class->HasMetaData(TEXT("Private")))
-						//{
-						//	continue;
-						//}
+						if(Class->HasMetaData(TEXT("Private")))
+						{
+							continue;
+						}
 
 						// Exclude comments from the expression list, as well as the base parameter expression and local variable usage, as they should not be used directly
 						if ( Class != UMaterialExpressionComment::StaticClass() 
