Add "Virtual texture stacks" entry to the material editor display of platform stats.
#rnx

diff --git a/Engine/Source/Editor/MaterialEditor/Private/MaterialStats.h b/Engine/Source/Editor/MaterialEditor/Private/MaterialStats.h
--- a/Engine/Source/Editor/MaterialEditor/Private/MaterialStats.h
+++ b/Engine/Source/Editor/MaterialEditor/Private/MaterialStats.h
@@ -26,6 +26,7 @@ struct FShaderStatsInfo
 	FContent InterpolatorsCount;
 	FContent TextureSampleCount;
 	FContent VirtualTextureLookupCount;
+	FContent VirtualTextureStackCount;
 	FContent ShaderCount;
 	FContent PreShaderCount;
 	FContent LWCUsage;
@@ -48,6 +49,9 @@ struct FShaderStatsInfo
 		VirtualTextureLookupCount.StrDescription = TEXT("Compiling...");
 		VirtualTextureLookupCount.StrDescriptionLong = TEXT("Compiling...");
 
+		VirtualTextureStackCount.StrDescription = TEXT("Compiling...");
+		VirtualTextureStackCount.StrDescriptionLong = TEXT("Compiling...");
+
 		ShaderCount.StrDescription = TEXT("Compiling...");
 		ShaderCount.StrDescriptionLong = TEXT("Compiling...");
 
@@ -77,6 +81,9 @@ struct FShaderStatsInfo
 		VirtualTextureLookupCount.StrDescription.Empty();
 		VirtualTextureLookupCount.StrDescriptionLong.Empty();
 
+		VirtualTextureStackCount.StrDescription.Empty();
+		VirtualTextureStackCount.StrDescriptionLong.Empty();
+
 		ShaderCount.StrDescription.Empty();
 		ShaderCount.StrDescriptionLong.Empty();
 
diff --git a/Engine/Source/Editor/MaterialEditor/Private/MaterialStatsCommon.cpp b/Engine/Source/Editor/MaterialEditor/Private/MaterialStatsCommon.cpp
--- a/Engine/Source/Editor/MaterialEditor/Private/MaterialStatsCommon.cpp
+++ b/Engine/Source/Editor/MaterialEditor/Private/MaterialStatsCommon.cpp
@@ -607,6 +607,10 @@ void FMaterialStatsUtils::ExtractMatertialStatsInfo(EShaderPlatform ShaderPlatfo
 		OutInfo.VirtualTextureLookupCount.StrDescription = FString::Printf(TEXT("%u"), NumVirtualTextureLookups);
 		OutInfo.VirtualTextureLookupCount.StrDescriptionLong = FString::Printf(TEXT("Virtual Texture Lookups (Est.): %u"), NumVirtualTextureLookups);
 
+		const uint32 NumVirtualTextureStacks = MaterialResource->GetNumVirtualTextureStacks();
+		OutInfo.VirtualTextureStackCount.StrDescription = FString::Printf(TEXT("%u"), NumVirtualTextureStacks);
+		OutInfo.VirtualTextureStackCount.StrDescriptionLong = FString::Printf(TEXT("Virtual Texture Stacks (Est.): %u"), NumVirtualTextureStacks);
+
 		// extract interpolators info
 		uint32 UVScalarsUsed, CustomInterpolatorScalarsUsed;
 		MaterialResource->GetUserInterpolatorUsage(UVScalarsUsed, CustomInterpolatorScalarsUsed);
diff --git a/Engine/Source/Editor/MaterialEditor/Private/MaterialStatsGrid.cpp b/Engine/Source/Editor/MaterialEditor/Private/MaterialStatsGrid.cpp
--- a/Engine/Source/Editor/MaterialEditor/Private/MaterialStatsGrid.cpp
+++ b/Engine/Source/Editor/MaterialEditor/Private/MaterialStatsGrid.cpp
@@ -365,6 +365,37 @@ void FStatsGridRow_NumVirtualTextureLookups::AddPlatform(TSharedPtr<FMaterialSta
 /*end FStatsGridRow_NumVirtualTextureLookups functions*/
 /*==============================================================================================================*/
 
+/*==============================================================================================================*/
+/* FStatsGridRow_NumVirtualTextureStacks functions*/
+
+void FStatsGridRow_NumVirtualTextureStacks::CreateRow(TSharedPtr<FMaterialStats> StatsManager)
+{
+	// static string in the descriptor column
+	TSharedPtr<FGridCell> HeaderCell = MakeShareable(new FGridCell_StaticString(TEXT("VT Stacks (Est.)"), TEXT("Virtual Texture Stacks (Est.)")));
+	HeaderCell->SetColor(FStyleColors::Foreground);
+	HeaderCell->SetContentBold(true);
+	AddCell(FMaterialStatsGrid::DescriptorColumnName, HeaderCell);
+
+	AddCell(FMaterialStatsGrid::ShaderColumnName, MakeShareable(new FGridCell_Empty()));
+	AddCell(FMaterialStatsGrid::ShaderStatisticColumnName, MakeShareable(new FGridCell_Empty()));
+
+	FillPlatformCellsHelper(StatsManager);
+}
+
+void FStatsGridRow_NumVirtualTextureStacks::AddPlatform(TSharedPtr<FMaterialStats> StatsManager, const TSharedPtr<FShaderPlatformSettings> Platform, const EMaterialQualityLevel::Type QualityLevel, const int32 InstanceIndex)
+{
+	// cell that will enumerate the total number of texture samples in this material
+	const FString CellContent = FMaterialStatsUtils::MaterialQualityToShortString(QualityLevel);
+	TSharedPtr<FGridCell_ShaderValue> Cell = MakeShareable(new FGridCell_ShaderValue(StatsManager, EShaderInfoType::VirtualTextureStackCount, ERepresentativeShader::Num, QualityLevel, Platform->GetPlatformShaderType(), InstanceIndex));
+
+	const FName ColumnName = FMaterialStatsGrid::MakePlatformColumnName(Platform, QualityLevel, InstanceIndex);
+	AddCell(ColumnName, Cell);
+}
+
+/*end FStatsGridRow_NumVirtualTextureStacks functions*/
+
+/*==============================================================================================================*/
+
 /*==============================================================================================================*/
 /* FStatsGridRow_NumShaders functions*/
 
@@ -648,6 +679,7 @@ void FMaterialStatsGrid::BuildRowIds()
 		BuildKeyAndInsert(ERowType::Samplers);
 		BuildKeyAndInsert(ERowType::TextureSamples);
 		BuildKeyAndInsert(ERowType::VirtualTextureLookups);
+		BuildKeyAndInsert(ERowType::VirtualTextureStacks);
 		BuildKeyAndInsert(ERowType::Interpolators);
 		BuildKeyAndInsert(ERowType::Shaders);
 		BuildKeyAndInsert(ERowType::PreShaders);
@@ -712,6 +744,11 @@ void FMaterialStatsGrid::BuildStaticRows()
 		Row->CreateRow(StatsManager);
 		StaticRows.Add(ERowType::VirtualTextureLookups, Row);
 	}
+	{
+		TSharedPtr<FStatsGridRow> Row = MakeShareable(new FStatsGridRow_NumVirtualTextureStacks());
+		Row->CreateRow(StatsManager);
+		StaticRows.Add(ERowType::VirtualTextureStacks, Row);
+	}
 	{
 		TSharedPtr<FStatsGridRow> Row = MakeShareable(new FStatsGridRow_Interpolators());
 		Row->CreateRow(StatsManager);
@@ -1072,6 +1109,9 @@ FString FGridCell_ShaderValue::InternalGetContent(bool bLongContent)
 		case EShaderInfoType::VirtualTextureLookupCount:
 			return bLongContent ? InstanceData.ShaderStatsInfo.VirtualTextureLookupCount.StrDescriptionLong : InstanceData.ShaderStatsInfo.VirtualTextureLookupCount.StrDescription;
 
+		case EShaderInfoType::VirtualTextureStackCount:
+			return bLongContent ? InstanceData.ShaderStatsInfo.VirtualTextureStackCount.StrDescriptionLong : InstanceData.ShaderStatsInfo.VirtualTextureStackCount.StrDescription;
+
 		case EShaderInfoType::ShaderCount:
 			return bLongContent ? InstanceData.ShaderStatsInfo.ShaderCount.StrDescriptionLong : InstanceData.ShaderStatsInfo.ShaderCount.StrDescription;
 
diff --git a/Engine/Source/Editor/MaterialEditor/Private/MaterialStatsGrid.h b/Engine/Source/Editor/MaterialEditor/Private/MaterialStatsGrid.h
--- a/Engine/Source/Editor/MaterialEditor/Private/MaterialStatsGrid.h
+++ b/Engine/Source/Editor/MaterialEditor/Private/MaterialStatsGrid.h
@@ -86,6 +86,7 @@ enum class EShaderInfoType
 	InterpolatorsCount,
 	TextureSampleCount,
 	VirtualTextureLookupCount,
+	VirtualTextureStackCount,
 	ShaderCount,
 	PreShaderCount,
 	LWCUsage,
@@ -241,6 +242,15 @@ public:
 	void AddPlatform(TSharedPtr<FMaterialStats> StatsManager, const TSharedPtr<FShaderPlatformSettings> Platform, const EMaterialQualityLevel::Type QualityLevel, const int32 InstanceIndex) override;
 };
 
+/** this row will display the global number of virtual texture stacks present in the material for a specified platform */
+class FStatsGridRow_NumVirtualTextureStacks : public FStatsGridRow
+{
+public:
+	void CreateRow(TSharedPtr<FMaterialStats> StatsManager) override;
+
+	void AddPlatform(TSharedPtr<FMaterialStats> StatsManager, const TSharedPtr<FShaderPlatformSettings> Platform, const EMaterialQualityLevel::Type QualityLevel, const int32 InstanceIndex) override;
+};
+
 /** this row will display the total number of shaders present in the material for a specified platform */
 class FStatsGridRow_NumShaders : public FStatsGridRow
 {
@@ -292,6 +302,7 @@ class FMaterialStatsGrid
 		Interpolators,
 		TextureSamples,
 		VirtualTextureLookups,
+		VirtualTextureStacks,
 		Shaders,
 		PreShaders,
 		LWCUsage,
