Allow for copy/paste of material errors in the PlatformStats panel.

- Simply make the text read-only "editable" to allow the user to select the text and get automatic copy/paste functionality.

[RN] minor | Rendering - Materials and Shaders

diff --git a/Engine/Source/Editor/MaterialEditor/Private/SMaterialEditorStatsWidget.cpp b/Engine/Source/Editor/MaterialEditor/Private/SMaterialEditorStatsWidget.cpp
--- a/Engine/Source/Editor/MaterialEditor/Private/SMaterialEditorStatsWidget.cpp
+++ b/Engine/Source/Editor/MaterialEditor/Private/SMaterialEditorStatsWidget.cpp
@@ -3,6 +3,7 @@
 #include "SMaterialEditorStatsWidget.h"
 #include "Widgets/Text/STextBlock.h"
 #include "Widgets/Layout/SSpacer.h"
+#include "Widgets/Input/SEditableTextBox.h"
 #include "Widgets/Layout/SExpandableArea.h"
 #include "Widgets/Layout/SScrollBox.h"
 #include "Widgets/Views/SListView.h"
@@ -626,11 +627,12 @@ void SMaterialEditorStatsWidget::AddMessage(const FString& Message, const bool b
 				.AutoWidth()
 				.VAlign(VAlign_Center)
 				[
-					SNew(STextBlock)
-					.TextStyle(FAppStyle::Get(), TEXT("RichTextBlock.Bold"))
+					SNew(SEditableText)
+					.Style(FAppStyle::Get(), TEXT("RichTextBlock.Bold"))
 					.ColorAndOpacity(bIsError ? FStyleColors::Error : FStyleColors::Warning)
 					.Text(FText::FromString(Message))
 					.ToolTipText(FText::FromString(Message))
+					.IsReadOnly(true)
 				]
 			]
 		];
