Add clarification about Substrate features when using format=0.

diff --git a/Engine/Source/Editor/MaterialEditor/Private/SMaterialEditorSubstrateWidget.cpp b/Engine/Source/Editor/MaterialEditor/Private/SMaterialEditorSubstrateWidget.cpp
--- a/Engine/Source/Editor/MaterialEditor/Private/SMaterialEditorSubstrateWidget.cpp
+++ b/Engine/Source/Editor/MaterialEditor/Private/SMaterialEditorSubstrateWidget.cpp
@@ -42,6 +42,27 @@ void SMaterialEditorSubstrateWidget::Construct(const FArguments& InArgs, TWeakPt
 		.ShadowOffset(FVector2D::UnitVector)
 		.Text(LOCTEXT("DescriptionTextBlock_Default", "Shader is compiling"));
 
+	MaterialFeatures0TextBlock = SNew(STextBlock)
+		.TextStyle(FAppStyle::Get(), "Log.Normal")
+		.ColorAndOpacity(FLinearColor::White)
+		.ShadowColorAndOpacity(FLinearColor::Black)
+		.ShadowOffset(FVector2D::UnitVector)
+		.Text(LOCTEXT("DescriptionTextBlock_Default", "Shader is compiling"));
+
+	MaterialFeatures1TextBlock = SNew(STextBlock)
+		.TextStyle(FAppStyle::Get(), "Log.Normal")
+		.ColorAndOpacity(FLinearColor::White)
+		.ShadowColorAndOpacity(FLinearColor::Black)
+		.ShadowOffset(FVector2D::UnitVector)
+		.Text(LOCTEXT("DescriptionTextBlock_Default", "Shader is compiling"));
+
+	MaterialFeatures2TextBlock = SNew(STextBlock)
+		.TextStyle(FAppStyle::Get(), "Log.Normal")
+		.ColorAndOpacity(FLinearColor::White)
+		.ShadowColorAndOpacity(FLinearColor::Black)
+		.ShadowOffset(FVector2D::UnitVector)
+		.Text(LOCTEXT("DescriptionTextBlock_Default", "Shader is compiling"));
+
 	DescriptionTextBlock = SNew(STextBlock)
 		.TextStyle(FAppStyle::Get(), "Log.Normal")
 		.ColorAndOpacity(FLinearColor::White)
@@ -343,7 +364,59 @@ void SMaterialEditorSubstrateWidget::Construct(const FArguments& InArgs, TWeakPt
 						]
 					]
 
+					+SVerticalBox::Slot()
+					.AutoHeight()
+					.Padding(0.0f, 15.0f, 0.0f, 0.0f)
+					[
+						SNew(SWrapBox)
+						.UseAllottedSize(true)
+						+ SWrapBox::Slot()
+						.Padding(0.0f)
+						.HAlign(HAlign_Left)
+						.VAlign(VAlign_Center)
+						[
+							SNew(STextBlock)
+							.ColorAndOpacity(FLinearColor::White)
+							.ShadowColorAndOpacity(FLinearColor::Black)
+							.ShadowOffset(FVector2D::UnitVector)
+							.Text(LOCTEXT("MaterialFeaturesTextBlock", "Material Features"))
+						]
+					]
 
+					+SVerticalBox::Slot()
+					.AutoHeight()
+					.Padding(0.0f, 0.0f, 0.0f, 0.0f)
+					[
+						SNew(SBorder)
+						.Padding(FMargin(5.0f, 5.0f, 5.0f, 5.0f))
+						[
+							SNew(SVerticalBox)
+							+ SVerticalBox::Slot()
+							.AutoHeight()
+							.Padding(FMargin(5.0f, 0.0f, 0.0f, 0.0f))
+							.HAlign(HAlign_Left)
+							.VAlign(VAlign_Top)
+							[
+								MaterialFeatures0TextBlock->AsShared()
+							]
+							+ SVerticalBox::Slot()
+							.AutoHeight()
+							.Padding(FMargin(5.0f, 0.0f, 0.0f, 0.0f))
+							.HAlign(HAlign_Left)
+							.VAlign(VAlign_Top)
+							[
+								MaterialFeatures1TextBlock->AsShared()
+							]
+							+ SVerticalBox::Slot()
+							.AutoHeight()
+							.Padding(FMargin(5.0f, 0.0f, 0.0f, 0.0f))
+							.HAlign(HAlign_Left)
+							.VAlign(VAlign_Top)
+							[
+								MaterialFeatures2TextBlock->AsShared()
+							]
+						]
+					]
 					
 					+SVerticalBox::Slot()
 					.AutoHeight()
@@ -462,6 +535,9 @@ void SMaterialEditorSubstrateWidget::Tick(const FGeometry& AllottedGeometry, con
 		{
 			FString MaterialDescription;
 			FString MaterialBudget;
+			FString MaterialFeatures0;
+			FString MaterialFeatures1;
+			FString MaterialFeatures2;
 
 			bool bMaterialOutOfBudgetHasBeenSimplified = false;
 			FMaterialShaderMap* ShaderMap = MaterialResource->GetGameThreadShaderMap();
@@ -472,12 +548,82 @@ void SMaterialEditorSubstrateWidget::Tick(const FGeometry& AllottedGeometry, con
 				const uint32 FinalPixelClosureCount = CompilationOutput.SubstrateClosureCount;
 				bMaterialOutOfBudgetHasBeenSimplified = CompilationOutput.bMaterialOutOfBudgetHasBeenSimplified > 0;
 
+				// Features
+				{	
+					// Actives features
+					const ESubstrateBsdfFeature TrackedFeatures0 = 
+						  ESubstrateBsdfFeature::SSS
+						| ESubstrateBsdfFeature::Fuzz
+						| ESubstrateBsdfFeature::SecondRoughnessOrSimpleClearCoat
+						| ESubstrateBsdfFeature::Anisotropy;
+
+					const ESubstrateBsdfFeature TrackedFeatures1 = 
+						  ESubstrateBsdfFeature::EdgeColor
+						| ESubstrateBsdfFeature::Glint
+						| ESubstrateBsdfFeature::SpecularProfile;
+
+					ESubstrateBsdfFeature TrackedFeatures = ESubstrateBsdfFeature::None;
+					EnumAddFlags(TrackedFeatures, TrackedFeatures0);
+					EnumAddFlags(TrackedFeatures, TrackedFeatures1);
+
+					// Untracked features:
+					// - ESubstrateBsdfFeature::MFPPluggedIn
+					// - ESubstrateBsdfFeature::Eye
+					// - ESubstrateBsdfFeature::Hair
+					if (EnumHasAnyFlags(CompilationOutput.SubstrateMaterialBsdfFeatures, TrackedFeatures))
+					{
+						// Selection & Sorting based on logic in:
+						// - SubstrateGetLegacyShadingModels(..)
+						// - SubstrateMaterialExportOut(..)					
+						FString& PotentiallySupported   = Substrate::IsSubstrateBlendableGBufferEnabled(GMaxRHIShaderPlatform) ? MaterialFeatures1 : MaterialFeatures0;
+						FString& PotentiallyUnsupported = Substrate::IsSubstrateBlendableGBufferEnabled(GMaxRHIShaderPlatform) ? MaterialFeatures2 : MaterialFeatures0;
+
+						if (Substrate::IsSubstrateBlendableGBufferEnabled(GMaxRHIShaderPlatform))
+						{
+							if (EnumHasAnyFlags(CompilationOutput.SubstrateMaterialBsdfFeatures, TrackedFeatures0))
+							{
+								MaterialFeatures1 += FString::Printf(TEXT("Enabled features: \n"));
+							}
+							if (EnumHasAnyFlags(CompilationOutput.SubstrateMaterialBsdfFeatures, TrackedFeatures1))
+							{
+								MaterialFeatures2 += FString::Printf(TEXT("Disabled features: \n"));
+							}
+						}
+						else
+						{
+							MaterialFeatures0 += FString::Printf(TEXT("Enabled features: \n"));
+						}
+
+						if (EnumHasAnyFlags(CompilationOutput.SubstrateMaterialBsdfFeatures, ESubstrateBsdfFeature::SSS                              )) { PotentiallySupported 	+= TEXT(" - SSS\n"); }
+						if (EnumHasAnyFlags(CompilationOutput.SubstrateMaterialBsdfFeatures, ESubstrateBsdfFeature::Fuzz                             )) { PotentiallySupported 	+= TEXT(" - Fuzz\n"); }
+						if (EnumHasAnyFlags(CompilationOutput.SubstrateMaterialBsdfFeatures, ESubstrateBsdfFeature::SecondRoughnessOrSimpleClearCoat )) { PotentiallySupported 	+= TEXT(" - ClearCoat\n"); }
+						if (EnumHasAnyFlags(CompilationOutput.SubstrateMaterialBsdfFeatures, ESubstrateBsdfFeature::Anisotropy                       )) { PotentiallySupported 	+= TEXT(" - Anisotropy\n"); }
+						if (EnumHasAnyFlags(CompilationOutput.SubstrateMaterialBsdfFeatures, ESubstrateBsdfFeature::EdgeColor                        )) { PotentiallyUnsupported+= TEXT(" - F90\n"); }
+						if (EnumHasAnyFlags(CompilationOutput.SubstrateMaterialBsdfFeatures, ESubstrateBsdfFeature::Glint                            )) { PotentiallyUnsupported+= TEXT(" - Glint\n"); }
+						if (EnumHasAnyFlags(CompilationOutput.SubstrateMaterialBsdfFeatures, ESubstrateBsdfFeature::SpecularProfile                  )) { PotentiallyUnsupported+= TEXT(" - Specular Profile\n"); }
+
+						if (Substrate::IsSubstrateBlendableGBufferEnabled(GMaxRHIShaderPlatform))
+						{
+							if (EnumHasAnyFlags(CompilationOutput.SubstrateMaterialBsdfFeatures, TrackedFeatures0))
+							{
+								MaterialFeatures1 += FString::Printf(TEXT("\nFormat : Blendable GBuffer => Only one feature will be active, sorted by priority. The other features will be disabled.\n"));
+							}
+							if (EnumHasAnyFlags(CompilationOutput.SubstrateMaterialBsdfFeatures, TrackedFeatures1))
+							{
+								MaterialFeatures2 += FString::Printf(TEXT("\nFormat : Blendable GBuffer => Unsupported features.\n"));
+							}
+						}
+					}
+				}
+
+				// Budget
 				if (bMaterialOutOfBudgetHasBeenSimplified)
 				{
-					MaterialBudget += FString::Printf(TEXT("The material was OUT-OF-BUDGET so it has been simplified: Requested bytes = %i / budget = %i  -  Requested Closures = %i / budget = %i\r\n"),
+					MaterialBudget += FString::Printf(TEXT("Format: %s\n\nThe material was OUT-OF-BUDGET so it has been simplified:\n  - Requested Bytes = %i / budget = %i\n  - Requested Closures = %i / budget = %i\r\n\n"),
+						Substrate::IsSubstrateBlendableGBufferEnabled(GMaxRHIShaderPlatform) ? TEXT("Blendable GBuffer") : TEXT("Adaptive GBuffer"),
 						CompilationOutput.RequestedBytePerPixel, CompilationOutput.PlatformBytePerPixel,
 						CompilationOutput.RequestedClosurePerPixel, CompilationOutput.PlatformClosurePixel);
-					MaterialBudget += FString::Printf(TEXT("Final per pixel byte count   = %i\r\n"), FinalPixelByteCount);
+					MaterialBudget += FString::Printf(TEXT("Final per pixel byte count      = %i\r\n"), FinalPixelByteCount);
 					MaterialBudget += FString::Printf(TEXT("Final per pixel closure count   = %i\r\n"), FinalPixelClosureCount);
 				}
 				else
@@ -487,6 +633,8 @@ void SMaterialEditorSubstrateWidget::Tick(const FGeometry& AllottedGeometry, con
 					MaterialBudget += FString::Printf(TEXT("Material per pixel closure count = %i / budget = %i\r\n"),
 						FinalPixelClosureCount, CompilationOutput.PlatformClosurePixel);
 				}
+
+				// Description
 				MaterialDescription += FString::Printf(TEXT("Closures Count               = %i\r\n"), CompilationOutput.SubstrateClosureCount);
 				MaterialDescription += FString::Printf(TEXT("Local bases Count            = %i\r\n"), CompilationOutput.SharedLocalBasesCount);
 				MaterialDescription += FString::Printf(TEXT("Material complexity          = %s\r\n"), *SubstrateMaterialTypeToString(CompilationOutput.SubstrateMaterialType));
@@ -588,8 +736,17 @@ void SMaterialEditorSubstrateWidget::Tick(const FGeometry& AllottedGeometry, con
 			{
 				MaterialBudgetTextBlock->SetColorAndOpacity(FLinearColor::White);
 			}
-			MaterialBudgetTextBlock->SetText(FText::FromString(MaterialBudget));
 
+			{
+				MaterialFeatures0TextBlock->SetColorAndOpacity(FLinearColor::White);
+				MaterialFeatures1TextBlock->SetColorAndOpacity(FLinearColor(0.8, 0.5f, 0.07f, 1.0f));
+				MaterialFeatures2TextBlock->SetColorAndOpacity(FLinearColor(0.8, 0.07f, 0.07f, 1.0f));
+			}
+
+			MaterialBudgetTextBlock->SetText(FText::FromString(MaterialBudget));
+			MaterialFeatures0TextBlock->SetText(FText::FromString(MaterialFeatures0));
+			MaterialFeatures1TextBlock->SetText(FText::FromString(MaterialFeatures1));
+			MaterialFeatures2TextBlock->SetText(FText::FromString(MaterialFeatures2));
 			DescriptionTextBlock->SetText(FText::FromString(MaterialDescription));
 		}
 	}
diff --git a/Engine/Source/Editor/MaterialEditor/Private/SMaterialEditorSubstrateWidget.h b/Engine/Source/Editor/MaterialEditor/Private/SMaterialEditorSubstrateWidget.h
--- a/Engine/Source/Editor/MaterialEditor/Private/SMaterialEditorSubstrateWidget.h
+++ b/Engine/Source/Editor/MaterialEditor/Private/SMaterialEditorSubstrateWidget.h
@@ -48,6 +48,9 @@ private:
 	TSharedPtr<class SButton> ButtonApplyToPreview;
 
 	TSharedPtr<class STextBlock> MaterialBudgetTextBlock;
+	TSharedPtr<class STextBlock> MaterialFeatures0TextBlock;
+	TSharedPtr<class STextBlock> MaterialFeatures1TextBlock;
+	TSharedPtr<class STextBlock> MaterialFeatures2TextBlock;
 	TSharedPtr<class STextBlock> DescriptionTextBlock;
 
 	TSharedPtr<class SBox> MaterialBox;
