Fix Upper Right buttons on material graph nodes not properly aligning to the right of the node title in all cases (particularly smaller nodes)

- set up our upper right buttons by overriding SGraphNodeMaterialBase::CreateTitleRightWidget() instead of ::SetDefaultTitleAreaWidget(TSharedRef<SOverlay> DefaultTitleAreaWidget) which properly allows that content to add desired size to the node title area as needed

- Have title area ColorSpill actually fill the entire title area in SGraphNode.cpp, rather than just the first HBox slot

diff --git a/Engine/Source/Editor/GraphEditor/Private/SGraphNode.cpp b/Engine/Source/Editor/GraphEditor/Private/SGraphNode.cpp
--- a/Engine/Source/Editor/GraphEditor/Private/SGraphNode.cpp
+++ b/Engine/Source/Editor/GraphEditor/Private/SGraphNode.cpp
@@ -878,14 +878,14 @@ void SGraphNode::UpdateGraphNode()
 		.HAlign(HAlign_Fill)
 		.VAlign(VAlign_Fill)
 		[
-			SNew(SHorizontalBox)
-			+ SHorizontalBox::Slot()
-			.HAlign(HAlign_Fill)
+			SNew(SBorder)
+			.BorderImage( FAppStyle::GetBrush("Graph.Node.ColorSpill") )
+			.Padding(TitleBorderMargin)
+			.BorderBackgroundColor(this, &SGraphNode::GetNodeTitleColor)
 			[
-				SNew(SBorder)
-				.BorderImage( FAppStyle::GetBrush("Graph.Node.ColorSpill") )
-				.Padding(TitleBorderMargin)
-				.BorderBackgroundColor( this, &SGraphNode::GetNodeTitleColor )
+				SNew(SHorizontalBox)
+				+ SHorizontalBox::Slot()
+				.HAlign(HAlign_Fill)
 				[
 					SNew(SHorizontalBox)
 					+ SHorizontalBox::Slot()
@@ -910,14 +910,14 @@ void SGraphNode::UpdateGraphNode()
 						]
 					]
 				]
-			]
-			+ SHorizontalBox::Slot()
-			.HAlign(HAlign_Right)
-			.VAlign(VAlign_Center)
-			.Padding(0, 0, 5, 0)
-			.AutoWidth()
-			[
-				CreateTitleRightWidget()
+				+ SHorizontalBox::Slot()
+				.HAlign(HAlign_Right)
+				.VAlign(VAlign_Center)
+				.Padding(0, 0, 5, 0)
+				.AutoWidth()
+				[
+					CreateTitleRightWidget()
+				]
 			]
 		]
 		+SOverlay::Slot()
diff --git a/Engine/Source/Editor/MaterialEditor/Private/MaterialNodes/SGraphNodeMaterialBase.cpp b/Engine/Source/Editor/MaterialEditor/Private/MaterialNodes/SGraphNodeMaterialBase.cpp
--- a/Engine/Source/Editor/MaterialEditor/Private/MaterialNodes/SGraphNodeMaterialBase.cpp
+++ b/Engine/Source/Editor/MaterialEditor/Private/MaterialNodes/SGraphNodeMaterialBase.cpp
@@ -302,7 +302,7 @@ void FPreviewElement::Draw_RenderThread(FRDGBuilder& GraphBuilder, const FDrawPa
 
 void SGraphNodeMaterialBase::Construct(const FArguments& InArgs, UMaterialGraphNode* InNode)
 {
-	this->TitleBorderMargin = FMargin(10.f, 5.f, 50.f, 3.f);
+	this->TitleBorderMargin = FMargin(10.f, 5.f, 10.f, 3.f);
 	this->GraphNode = InNode;
 	this->MaterialNode = InNode;
 
@@ -548,55 +548,51 @@ void SGraphNodeMaterialBase::CreateBelowPinControls(TSharedPtr<SVerticalBox> Mai
 	}
 }
 
-void SGraphNodeMaterialBase::SetDefaultTitleAreaWidget(TSharedRef<SOverlay> DefaultTitleAreaWidget)
+TSharedRef<SWidget> SGraphNodeMaterialBase::CreateTitleRightWidget()
 {
-	if (!MaterialNode->MaterialExpression->bHidePreviewWindow)
+	if (MaterialNode->MaterialExpression->bHidePreviewWindow)
 	{
-		DefaultTitleAreaWidget->AddSlot()
-		.HAlign(HAlign_Right)
-		.VAlign(VAlign_Top)
-		.Padding(FMargin(5))
+		return SNullWidget::NullWidget;
+	}
+
+	return SNew(SHorizontalBox)
+	+ SHorizontalBox::Slot()
+	.AutoWidth()
+	[
+		SNew(SButton)
+		.ButtonStyle(FAppStyle::Get(), "HoverHintOnly")
+		.OnClicked(this, &SGraphNodeMaterialBase::OnTogglePreviewClicked)
+		.Cursor(EMouseCursor::Default)
 		[
 			SNew(SHorizontalBox)
-			+ SHorizontalBox::Slot()
-			.AutoWidth()
+			+SHorizontalBox::Slot()
+			.VAlign(VAlign_Center)
+			.HAlign(HAlign_Center)
 			[
-				SNew(SButton)
-				.ButtonStyle(FAppStyle::Get(), "HoverHintOnly")
-				.OnClicked(this, &SGraphNodeMaterialBase::OnTogglePreviewClicked)
-				.Cursor(EMouseCursor::Default)
-				[
-					SNew(SHorizontalBox)
-					+SHorizontalBox::Slot()
-					.VAlign(VAlign_Center)
-					.HAlign(HAlign_Center)
-					[
-						SNew(SImage)
-						.Image(FAppStyle::GetBrush(TEXT("Icons.Preview")))
-					]
-				]
+				SNew(SImage)
+				.Image(FAppStyle::GetBrush(TEXT("Icons.Preview")))
 			]
-			+ SHorizontalBox::Slot()
-			.AutoWidth()
+		]
+	]
+	+ SHorizontalBox::Slot()
+	.AutoWidth()
+	[
+		SNew(SCheckBox)
+		.OnCheckStateChanged( this, &SGraphNodeMaterialBase::OnExpressionPreviewChanged )
+		.IsChecked( IsExpressionPreviewChecked() )
+		.Cursor(EMouseCursor::Default)
+		.Style(FAppStyle::Get(), "Graph.Node.AdvancedView")
+		[
+			SNew(SHorizontalBox)
+			+SHorizontalBox::Slot()
+			.VAlign(VAlign_Center)
+			.HAlign(HAlign_Center)
 			[
-				SNew(SCheckBox)
-				.OnCheckStateChanged( this, &SGraphNodeMaterialBase::OnExpressionPreviewChanged )
-				.IsChecked( IsExpressionPreviewChecked() )
-				.Cursor(EMouseCursor::Default)
-				.Style(FAppStyle::Get(), "Graph.Node.AdvancedView")
-				[
-					SNew(SHorizontalBox)
-					+SHorizontalBox::Slot()
-					.VAlign(VAlign_Center)
-					.HAlign(HAlign_Center)
-					[
-						SNew(SImage)
-						.Image(GetExpressionPreviewArrow())
-					]
-				]
+				SNew(SImage)
+				.Image(GetExpressionPreviewArrow())
 			]
-		];
-	}
+		]
+	];
 }
 
 TSharedRef<SWidget> SGraphNodeMaterialBase::CreateNodeContentArea()
diff --git a/Engine/Source/Editor/MaterialEditor/Public/MaterialEditor/MaterialNodes/SGraphNodeMaterialBase.h b/Engine/Source/Editor/MaterialEditor/Public/MaterialEditor/MaterialNodes/SGraphNodeMaterialBase.h
--- a/Engine/Source/Editor/MaterialEditor/Public/MaterialEditor/MaterialNodes/SGraphNodeMaterialBase.h
+++ b/Engine/Source/Editor/MaterialEditor/Public/MaterialEditor/MaterialNodes/SGraphNodeMaterialBase.h
@@ -115,7 +115,8 @@ protected:
 	// SGraphNode interface
 	UE_API virtual void AddPin( const TSharedRef<SGraphPin>& PinToAdd ) override;
 	UE_API virtual void CreateBelowPinControls(TSharedPtr<SVerticalBox> MainBox) override;
-	UE_API virtual void SetDefaultTitleAreaWidget(TSharedRef<SOverlay> DefaultTitleAreaWidget) override;
+	UE_API virtual TSharedRef<SWidget> CreateTitleRightWidget() override;
+
 	UE_API virtual TSharedRef<SWidget> CreateNodeContentArea() override;
 	UE_API virtual void OnAdvancedViewChanged(const ECheckBoxState NewCheckedState) override;
 	// End of SGraphNode interface
