6s -> 1s optimization when modifying material vector and scalar parameters on MIs with complex parents.

- Some materials were taking 6 seconds after _every_ material parameter change.  The bulk of the time was recursively searching the Base Material (and all its MFs) for all texture parameters.
- Ideally, we find a faster way to lookup texture expressions in order to perform validation, but for now we should only do this validation when a texture parameter property on the material instance changes.
- Put a cvar killswitch to completely turn this code off to get a further speed up if needed.

TODO: Long term this validation should be done on a per-property basis for further optimization.  Also don't display these warnings in the viewport, but on the properties themselves.

[RN] minor | Rendering - Materials and Shaders

diff --git a/Engine/Source/Editor/MaterialEditor/Private/MaterialInstanceEditor.cpp b/Engine/Source/Editor/MaterialEditor/Private/MaterialInstanceEditor.cpp
--- a/Engine/Source/Editor/MaterialEditor/Private/MaterialInstanceEditor.cpp
+++ b/Engine/Source/Editor/MaterialEditor/Private/MaterialInstanceEditor.cpp
@@ -72,6 +72,13 @@
 
 DEFINE_LOG_CATEGORY_STATIC(LogMaterialInstanceEditor, Log, All);
 
+static TAutoConsoleVariable<bool> CVarEnableTextureValidationInMaterialInstanceEditor(
+	TEXT("r.MaterialEditor.EnableTextureValidation"),
+	true,
+	TEXT("Enables expensive texture validation in the material instance editor.  On by default.  Disable this to improve performance when changing texture paramters on large materials."),
+	ECVF_Default
+);
+
 //////////////////////////////////////////////////////////////////////////
 // SMaterialTreeWidgetItem
 class SMaterialTreeWidgetItem : public SMultiColumnTableRow< TWeakObjectPtr<UMaterialInterface> >
@@ -1370,7 +1377,13 @@ void FMaterialInstanceEditor::NotifyPostChange( const FPropertyChangedEvent& Pro
 		UpdatePropertyWindow();
 	}
 
-	RefreshOnScreenMessages();
+	if (PropertyClass && 
+	   (PropertyClass->GetName() == TEXT("DEditorTextureParameterValue") || 
+		PropertyClass->GetName() == TEXT("DEditorRuntimeVirtualTextureParameterValue") ||
+		PropertyClass->GetName() == TEXT("DEditorSparseVolumeTextureParameterValue")))
+	{
+		RefreshOnScreenTextureWarnings();
+	}
 
 	// something was changed in the material so we need to reflect this in the stats
 	MaterialStatsManager->SignalMaterialChanged();
@@ -1505,11 +1518,13 @@ void FMaterialInstanceEditor::DrawMessages( FViewport* Viewport, FCanvas* Canvas
 	Canvas->PopTransform();
 }
 
-void FMaterialInstanceEditor::RefreshOnScreenMessages()
+void FMaterialInstanceEditor::RefreshOnScreenTextureWarnings()
 {
 	OnScreenMessages.Reset();
 
-	if (MaterialEditorInstance->SourceInstance)
+	const bool bValidateMaterials = CVarEnableTextureValidationInMaterialInstanceEditor.GetValueOnAnyThread();
+
+	if (bValidateMaterials && MaterialEditorInstance->SourceInstance)
 	{
 		UMaterial* BaseMaterial = MaterialEditorInstance->SourceInstance->GetMaterial();
 		if (BaseMaterial)
@@ -1827,7 +1842,7 @@ void FMaterialInstanceEditor::Refresh()
 	
 	UpdatePropertyWindow();
 
-	RefreshOnScreenMessages();
+	RefreshOnScreenTextureWarnings();
 
 	if (Substrate::IsMaterialLayeringSupportEnabled() && SubstrateWidget)
 	{
diff --git a/Engine/Source/Editor/MaterialEditor/Private/MaterialInstanceEditor.h b/Engine/Source/Editor/MaterialEditor/Private/MaterialInstanceEditor.h
--- a/Engine/Source/Editor/MaterialEditor/Private/MaterialInstanceEditor.h
+++ b/Engine/Source/Editor/MaterialEditor/Private/MaterialInstanceEditor.h
@@ -213,7 +213,7 @@ private:
 	/** Refreshes the preview asset */
 	void RefreshPreviewAsset();
 
-	void RefreshOnScreenMessages();
+	void RefreshOnScreenTextureWarnings();
 
 	//~ Begin FEditorUndoClient Interface
 	virtual void PostUndo( bool bSuccess ) override;
