Add MaterialEditingLibrary function called RenameMaterialParameterGroup and RenameMaterialFunctionParameterGroup which bulk renames all the Group Names on many Material Parameters at once.

- Change LogMaterialEditingLibrary verbosity to include `Log` messages.
- Simply loop over all material expressions that are parameters and if the OldGroupName matches change the name to the new one.
- Call RecompileMaterial to force the material to recompile so MIs update through the editor.
- Unfortunately the `Group` property of material expressions is spread across many different classes.  To avoid checking all the different material expression types (there are 8 different UMaterialExpression classes with their own unique Group), I've opted for a property system approach that looks for the property by name.

[RN] minor | Rendering - Materials and Shaders

diff --git a/Engine/Source/Editor/MaterialEditor/Private/MaterialEditingLibrary.cpp b/Engine/Source/Editor/MaterialEditor/Private/MaterialEditingLibrary.cpp
--- a/Engine/Source/Editor/MaterialEditor/Private/MaterialEditingLibrary.cpp
+++ b/Engine/Source/Editor/MaterialEditor/Private/MaterialEditingLibrary.cpp
@@ -39,7 +39,7 @@
 
 #include UE_INLINE_GENERATED_CPP_BY_NAME(MaterialEditingLibrary)
 
-DEFINE_LOG_CATEGORY_STATIC(LogMaterialEditingLibrary, Warning, All);
+DEFINE_LOG_CATEGORY_STATIC(LogMaterialEditingLibrary, Log, All);
 
 /** Util to find expression  */
 static FExpressionInput* GetExpressionInputByName(UMaterialExpression* Expression, const FName InputName)
@@ -1430,6 +1430,86 @@ bool UMaterialEditingLibrary::GetTextureParameterSource(UMaterialInterface* Mate
 	return false;
 }
 
+template <typename AssetType>
+static bool RenameInternal(AssetType* Asset, const FName OldGroupName, const FName NewGroupName)
+{
+	if (!Asset)
+	{
+		UE_LOG(LogMaterialEditingLibrary, Warning, TEXT("Invalid asset"));
+		return false;
+	}
+
+	if (Asset->HasAnyFlags(RF_ClassDefaultObject))
+	{
+		UE_LOG(LogMaterialEditingLibrary, Warning, TEXT("Can't rename groups on a Class Default Object."));
+		return false;
+	}
+
+	if (OldGroupName == NewGroupName)
+	{
+		UE_LOG(LogMaterialEditingLibrary, Log, TEXT("Old and new group names are identical: %s. Nothing to do."), *NewGroupName.ToString());
+		return false;
+	}
+
+	int NumModified = 0;
+	for (UMaterialExpression* Expression : Asset->GetExpressions())
+	{
+		if (!Expression)
+		{
+			continue;
+		}
+
+		// Use the property system to look for the `Group` parameter since many distinct parameter expressions
+		// can have a `Group`.
+		if (FProperty* Prop = Expression->GetClass()->FindPropertyByName(TEXT("Group")))
+		{
+			if (FNameProperty* NameProp = CastField<FNameProperty>(Prop))
+			{
+				FName* CurrentGroupName = NameProp->ContainerPtrToValuePtr<FName>(Expression);
+				if (CurrentGroupName && *CurrentGroupName == OldGroupName)
+				{
+					*CurrentGroupName = NewGroupName;
+					UE_LOG(LogMaterialEditingLibrary, Verbose, TEXT("Moved %s to %s."), *Expression->GetClass()->GetName(), *NewGroupName.ToString());
+					NumModified++;
+				}
+			}
+		}
+	}
+
+	if (NumModified > 0)
+	{
+		UE_LOG(LogMaterialEditingLibrary, Log, TEXT("Moved %d parameters from %s into group %s"), NumModified, *OldGroupName.ToString(), *NewGroupName.ToString());
+	}
+	else
+	{
+		UE_LOG(LogMaterialEditingLibrary, Log, TEXT("No parameters found with the group name: %s"), *OldGroupName.ToString());
+	}
+
+	return NumModified > 0;
+}
+
+bool UMaterialEditingLibrary::RenameMaterialParameterGroup(UMaterial* Material, const FName OldGroupName, const FName NewGroupName)
+{
+	if (RenameInternal(Material, OldGroupName, NewGroupName))
+	{
+		UMaterialEditingLibrary::RecompileMaterial(Material);
+		return true;
+	}
+
+	return false;
+}
+
+bool UMaterialEditingLibrary::RenameMaterialFunctionParameterGroup(UMaterialFunctionInterface* MaterialFunction, const FName OldGroupName, const FName NewGroupName)
+{
+	if (RenameInternal(MaterialFunction, OldGroupName, NewGroupName))
+	{
+		UMaterialEditingLibrary::UpdateMaterialFunction(MaterialFunction, nullptr);
+		return true;
+	}
+
+	return false;
+}
+
 bool UMaterialEditingLibrary::GetStaticSwitchParameterSource(UMaterialInterface* Material, const FName ParameterName, FSoftObjectPath& ParameterSource)
 {
 	if (Material)
diff --git a/Engine/Source/Editor/MaterialEditor/Public/MaterialEditingLibrary.h b/Engine/Source/Editor/MaterialEditor/Public/MaterialEditingLibrary.h
--- a/Engine/Source/Editor/MaterialEditor/Public/MaterialEditingLibrary.h
+++ b/Engine/Source/Editor/MaterialEditor/Public/MaterialEditingLibrary.h
@@ -398,6 +398,52 @@ public:
 	UFUNCTION(BlueprintCallable, Category = "MaterialEditing")
 	static UE_API bool GetTextureParameterSource(UMaterialInterface* Material, const FName ParameterName, FSoftObjectPath& ParameterSource);
 
+	/**
+	 * Renames a material parameter group within the specified material.
+	 *
+	 * This function allows you to rename an existing parameter group in a material.
+	 * It iterates all parameters within the material, finds all the one belonging
+	 * to the `OldGroupName` group and switches those parameters to be in the `NewGroupName`
+	 * group.  This function only affects parameters that belong to the specified group.   
+	 * Group name comparisons are case sensitive.  To remove the groups from the parameters 
+	 * the new group name can be 'None'.  If the OldGroupName does not exist in the material, 
+	 * the function will return false.   If the NewGroupName already exists, the parameters will 
+	 * be "merged" into the existing group.
+	 *
+	 * @param Material       The material asset in which the parameter group resides.
+	 * @param OldGroupName   The current name of the parameter group to rename.
+	 * @param NewGroupName   The new name to assign to the parameter group.
+	 *
+	 * @return true if the rename operation was successful; false otherwise.
+	 * 
+	 * @see RenameMaterialFunctionParameterGroup
+	 */
+	UFUNCTION(BlueprintCallable, Category = "MaterialEditing")
+	static UE_API bool RenameMaterialParameterGroup(UMaterial* Material, const FName OldGroupName, const FName NewGroupName);
+
+	/**
+	 * Renames a material parameter group within the specified material function.
+	 *
+	 * This function allows you to rename an existing parameter group in a material function.
+	 * It iterates all parameters within the material, finds all the one belonging
+	 * to the `OldGroupName` group and switches those parameters to be in the `NewGroupName`
+	 * group.  This function only affects parameters that belong to the specified group.
+	 * Group name comparisons are case sensitive.  To remove the groups from the parameters
+	 * the new group name can be 'None'.  If the OldGroupName does not exist in the material,
+	 * the function will return false.   If the NewGroupName already exists, the parameters will
+	 * be "merged" into the existing group.
+	 *
+	 * @param MaterialFunction	The material function asset in which the parameter group resides.
+	 * @param OldGroupName		The current name of the parameter group to rename.
+	 * @param NewGroupName		The new name to assign to the parameter group.
+	 *
+	 * @return true if the rename operation was successful; false otherwise.
+	 *
+	 * @see RenameMaterialParameterGroup
+	 */
+	UFUNCTION(BlueprintCallable, Category = "MaterialEditing")
+	static UE_API bool RenameMaterialFunctionParameterGroup(UMaterialFunctionInterface* MaterialFunction, const FName OldGroupName, const FName NewGroupName);
+
 	/**
 	*	Returns the path of the asset where the parameter originated, as well as true/false if it was found
 	*	@param	Material	The material or material instance you want to look up a parameter from
