Fix documentation.  The Material Parameter Group name is case-insensitive since it is dealing with FNames.

#rnx

diff --git a/Engine/Source/Editor/MaterialEditor/Public/MaterialEditingLibrary.h b/Engine/Source/Editor/MaterialEditor/Public/MaterialEditingLibrary.h
--- a/Engine/Source/Editor/MaterialEditor/Public/MaterialEditingLibrary.h
+++ b/Engine/Source/Editor/MaterialEditor/Public/MaterialEditingLibrary.h
@@ -405,10 +405,9 @@ public:
 	 * It iterates all parameters within the material, finds all the one belonging
 	 * to the `OldGroupName` group and switches those parameters to be in the `NewGroupName`
 	 * group.  This function only affects parameters that belong to the specified group.   
-	 * Group name comparisons are case sensitive.  To remove the groups from the parameters 
-	 * the new group name can be 'None'.  If the OldGroupName does not exist in the material, 
-	 * the function will return false.   If the NewGroupName already exists, the parameters will 
-	 * be "merged" into the existing group.
+	 * To remove the groups from the parameters the new group name can be 'None'.  If the 
+	 * OldGroupName does not exist in the material, the function will return false.   If 
+	 * the NewGroupName already exists, the parameters will be "merged" into the existing group.
 	 *
 	 * @param Material       The material asset in which the parameter group resides.
 	 * @param OldGroupName   The current name of the parameter group to rename.
@@ -425,13 +424,12 @@ public:
 	 * Renames a material parameter group within the specified material function.
 	 *
 	 * This function allows you to rename an existing parameter group in a material function.
-	 * It iterates all parameters within the material, finds all the one belonging
+	 * It iterates all parameters within the material function, finds all the one belonging
 	 * to the `OldGroupName` group and switches those parameters to be in the `NewGroupName`
 	 * group.  This function only affects parameters that belong to the specified group.
-	 * Group name comparisons are case sensitive.  To remove the groups from the parameters
-	 * the new group name can be 'None'.  If the OldGroupName does not exist in the material,
-	 * the function will return false.   If the NewGroupName already exists, the parameters will
-	 * be "merged" into the existing group.
+	 * To remove the groups from the parameters the new group name can be 'None'.  If the 
+	 * OldGroupName does not exist in the material, the function will return false.   If the 
+	 * NewGroupName already exists, the parameters will be "merged" into the existing group.
 	 *
 	 * @param MaterialFunction	The material function asset in which the parameter group resides.
 	 * @param OldGroupName		The current name of the parameter group to rename.
