Add notification hook to layer material instance editor properties.
This allows the platform stats display to update when modifying material instance layer properties.
#rnx

diff --git a/Engine/Source/Editor/MaterialEditor/Private/MaterialInstanceEditor.cpp b/Engine/Source/Editor/MaterialEditor/Private/MaterialInstanceEditor.cpp
--- a/Engine/Source/Editor/MaterialEditor/Private/MaterialInstanceEditor.cpp
+++ b/Engine/Source/Editor/MaterialEditor/Private/MaterialInstanceEditor.cpp
@@ -928,7 +928,8 @@ void FMaterialInstanceEditor::CreateInternalWidgets()
 	{
 		MaterialLayersFunctionsInstance = SNew(SMaterialLayersFunctionsInstanceWrapper)
 			.InMaterialEditorInstance(MaterialEditorInstance)
-			.InShowHiddenDelegate(FGetShowHiddenParameters::CreateSP(this, &FMaterialInstanceEditor::GetShowHiddenParameters));
+			.InShowHiddenDelegate(FGetShowHiddenParameters::CreateSP(this, &FMaterialInstanceEditor::GetShowHiddenParameters))
+			.InNotifyHook(this);
 		MaterialLayersFunctionsInstance->OnLayerPropertyChanged.BindSP(this, &FMaterialInstanceEditor::Refresh);
 	}
 	
diff --git a/Engine/Source/Editor/MaterialEditor/Private/SMaterialLayersFunctionsTree.cpp b/Engine/Source/Editor/MaterialEditor/Private/SMaterialLayersFunctionsTree.cpp
--- a/Engine/Source/Editor/MaterialEditor/Private/SMaterialLayersFunctionsTree.cpp
+++ b/Engine/Source/Editor/MaterialEditor/Private/SMaterialLayersFunctionsTree.cpp
@@ -805,6 +805,7 @@ void SMaterialLayersFunctionsInstanceTree::Construct(const FArguments& InArgs)
 	MaterialEditorInstance = InArgs._InMaterialEditorInstance;
 	Wrapper = InArgs._InWrapper;
 	ShowHiddenDelegate = InArgs._InShowHiddenDelegate;
+	NotifyHook = InArgs._InNotifyHook;
 	CreateGroupsWidget();
 
 #if WITH_EDITOR
@@ -1047,6 +1048,8 @@ void SMaterialLayersFunctionsInstanceTree::CreateGroupsWidget()
 	if (!Generator.IsValid())
 	{
 		FPropertyRowGeneratorArgs Args;
+		Args.NotifyHook = NotifyHook;
+
 		Generator = Module.CreatePropertyRowGenerator(Args);
 		// the sizes of the parameter lists are only based on the parent material and not changed out from under the details panel 
 		// When a parameter is added open MI editors are refreshed
@@ -1681,7 +1684,8 @@ void SMaterialLayersFunctionsInstanceWrapper::Construct(const FArguments& InArgs
 		NestedTree = SNew(SMaterialLayersFunctionsInstanceTree)
 			.InMaterialEditorInstance(Cast<UMaterialEditorInstanceConstant>(InArgs._InMaterialEditorInstance))
 			.InWrapper(this)
-			.InShowHiddenDelegate(InArgs._InShowHiddenDelegate);
+			.InShowHiddenDelegate(InArgs._InShowHiddenDelegate)
+			.InNotifyHook(InArgs._InNotifyHook);
 	}
 	
 
diff --git a/Engine/Source/Editor/MaterialEditor/Private/SMaterialLayersFunctionsTree.h b/Engine/Source/Editor/MaterialEditor/Private/SMaterialLayersFunctionsTree.h
--- a/Engine/Source/Editor/MaterialEditor/Private/SMaterialLayersFunctionsTree.h
+++ b/Engine/Source/Editor/MaterialEditor/Private/SMaterialLayersFunctionsTree.h
@@ -23,6 +23,7 @@ class IPropertyHandle;
 class SMaterialLayersFunctionsInstanceTree;
 class UDEditorParameterValue;
 class UMaterialEditorInstanceConstant;
+class FNotifyHook;
 
 class SMaterialLayersFunctionsInstanceTreeItem : public STableRow< TSharedPtr<FSortedParamData> >, public IDraggableItem
 {
@@ -110,12 +111,14 @@ public:
 	SLATE_BEGIN_ARGS(SMaterialLayersFunctionsInstanceWrapper)
 		: _InMaterialEditorInstance(nullptr),
 		_InGenerator(nullptr),
-		_InShowHiddenDelegate()
+		_InShowHiddenDelegate(),
+		_InNotifyHook(nullptr)
 	{}
 
 	SLATE_ARGUMENT(UMaterialEditorParameters*, InMaterialEditorInstance)
 	SLATE_ARGUMENT(TSharedPtr<class IPropertyRowGenerator>, InGenerator)
 	SLATE_ARGUMENT(FGetShowHiddenParameters, InShowHiddenDelegate)
+	SLATE_ARGUMENT(FNotifyHook*, InNotifyHook)
 	SLATE_END_ARGS()
 	
 	void Construct(const FArguments& InArgs);
@@ -154,11 +157,13 @@ class SMaterialLayersFunctionsInstanceTree : public SMaterialLayersTree
 public:
 	SLATE_BEGIN_ARGS(SMaterialLayersFunctionsInstanceTree)
 		: _InMaterialEditorInstance(nullptr)
+		, _InNotifyHook(nullptr)
 	{}
 
 	SLATE_ARGUMENT(UMaterialEditorInstanceConstant*, InMaterialEditorInstance)
 	SLATE_ARGUMENT(SMaterialLayersFunctionsInstanceWrapper*, InWrapper)
 	SLATE_ARGUMENT(FGetShowHiddenParameters, InShowHiddenDelegate)
+	SLATE_ARGUMENT(FNotifyHook*, InNotifyHook)
 	SLATE_END_ARGS()
 
 	/** Constructs this widget with InArgs */
@@ -225,6 +230,7 @@ private:
 	/** Delegate to call to determine if hidden parameters should be shown */
 	FGetShowHiddenParameters ShowHiddenDelegate;
 
+	FNotifyHook* NotifyHook = nullptr;
 };
 
 class UMaterialEditorPreviewParameters;
