Add a basic diff tool for Materail and Material Instance assets

[RN] Material and Material Instance assets can now be diff'd using the normal diff controls

diff --git a/Engine/Plugins/Editor/EngineAssetDefinitions/Source/Private/Material/AssetDefinition_Material.cpp b/Engine/Plugins/Editor/EngineAssetDefinitions/Source/Private/Material/AssetDefinition_Material.cpp
--- a/Engine/Plugins/Editor/EngineAssetDefinitions/Source/Private/Material/AssetDefinition_Material.cpp
+++ b/Engine/Plugins/Editor/EngineAssetDefinitions/Source/Private/Material/AssetDefinition_Material.cpp
@@ -2,6 +2,11 @@
 
 #include "AssetDefinition_Material.h"
 #include "MaterialEditorModule.h"
+#include "BlueprintEditor.h"
+#include "DiffTool/Widgets/SMaterialDiff.h"
+#include "MaterialGraph/MaterialGraph.h"
+#include "MaterialGraph/MaterialGraphSchema.h"
+#include "Kismet2/BlueprintEditorUtils.h"
 
 #include UE_INLINE_GENERATED_CPP_BY_NAME(AssetDefinition_Material)
 
@@ -15,3 +20,29 @@ EAssetCommandResult UAssetDefinition_Material::OpenAssets(const FAssetOpenArgs&
 
 	return EAssetCommandResult::Handled;
 }
+
+EAssetCommandResult UAssetDefinition_Material::PerformAssetDiff(const FAssetDiffArgs& DiffArgs) const
+{
+	UMaterial* OldMaterial = Cast<UMaterial>(DiffArgs.OldAsset);
+	UMaterial* NewMaterial = Cast<UMaterial>(DiffArgs.NewAsset);
+
+	if (OldMaterial && NewMaterial)
+	{
+		/** We create transient graphs so that the material asset is not dirtied by graph/package modifications */
+		UMaterialGraph* OldMaterialGraph = CastChecked<UMaterialGraph>(FBlueprintEditorUtils::CreateNewGraph(GetTransientPackage(), MakeUniqueObjectName(GetTransientPackage(), UMaterialGraph::StaticClass()), UMaterialGraph::StaticClass(), UMaterialGraphSchema::StaticClass()));
+		UMaterialGraph* NewMaterialGraph = CastChecked<UMaterialGraph>(FBlueprintEditorUtils::CreateNewGraph(GetTransientPackage(), MakeUniqueObjectName(GetTransientPackage(), UMaterialGraph::StaticClass()), UMaterialGraph::StaticClass(), UMaterialGraphSchema::StaticClass()));
+		OldMaterialGraph->Material = OldMaterial;
+		OldMaterialGraph->MaterialFunction = nullptr;
+
+		NewMaterialGraph->Material = NewMaterial;
+		NewMaterialGraph->MaterialFunction = nullptr;
+
+		OldMaterialGraph->RebuildGraph();
+		NewMaterialGraph->RebuildGraph();
+
+		SMaterialDiff::CreateDiffWindow(OldMaterialGraph, NewMaterialGraph, DiffArgs.OldRevision, DiffArgs.NewRevision, GetAssetClass().Get());
+		return EAssetCommandResult::Handled;
+	}
+
+	return EAssetCommandResult::Unhandled;
+}
diff --git a/Engine/Plugins/Editor/EngineAssetDefinitions/Source/Private/Material/AssetDefinition_Material.h b/Engine/Plugins/Editor/EngineAssetDefinitions/Source/Private/Material/AssetDefinition_Material.h
--- a/Engine/Plugins/Editor/EngineAssetDefinitions/Source/Private/Material/AssetDefinition_Material.h
+++ b/Engine/Plugins/Editor/EngineAssetDefinitions/Source/Private/Material/AssetDefinition_Material.h
@@ -24,5 +24,7 @@ public:
 		return Categories;
 	}
 	virtual EAssetCommandResult OpenAssets(const FAssetOpenArgs& OpenArgs) const override;
+	
+	virtual EAssetCommandResult PerformAssetDiff(const FAssetDiffArgs& DiffArgs) const override;
 	// UAssetDefinition End
 };
diff --git a/Engine/Plugins/Editor/EngineAssetDefinitions/Source/Private/Material/AssetDefinition_MaterialFunction.cpp b/Engine/Plugins/Editor/EngineAssetDefinitions/Source/Private/Material/AssetDefinition_MaterialFunction.cpp
--- a/Engine/Plugins/Editor/EngineAssetDefinitions/Source/Private/Material/AssetDefinition_MaterialFunction.cpp
+++ b/Engine/Plugins/Editor/EngineAssetDefinitions/Source/Private/Material/AssetDefinition_MaterialFunction.cpp
@@ -12,6 +12,11 @@
 #include "Misc/PackageName.h"
 #include "ToolMenu.h"
 #include "ToolMenuSection.h"
+#include "BlueprintEditor.h"
+#include "DiffTool/Widgets/SMaterialDiff.h"
+#include "Kismet2/BlueprintEditorUtils.h"
+#include "MaterialGraph/MaterialGraph.h"
+#include "MaterialGraph/MaterialGraphSchema.h"
 
 #include UE_INLINE_GENERATED_CPP_BY_NAME(AssetDefinition_MaterialFunction)
 
@@ -39,6 +44,47 @@ UThumbnailInfo* UAssetDefinition_MaterialFunction::LoadThumbnailInfo(const FAsse
 	return UE::Editor::FindOrCreateThumbnailInfo(InAsset.GetAsset(), USceneThumbnailInfoWithPrimitive::StaticClass());
 }
 
+EAssetCommandResult UAssetDefinition_MaterialFunction::PerformAssetDiff(const FAssetDiffArgs& DiffArgs) const
+{
+	UMaterialFunction* OldMaterialFunction = Cast<UMaterialFunction>(DiffArgs.OldAsset);
+	UMaterialFunction* NewMaterialFunction = Cast<UMaterialFunction>(DiffArgs.NewAsset);
+
+	if (OldMaterialFunction && NewMaterialFunction)
+	{
+		/** We create transient graphs so that the material asset is not dirtied by graph/package modifications */
+		UMaterialGraph* OldMaterialGraph = CastChecked<UMaterialGraph>(FBlueprintEditorUtils::CreateNewGraph(GetTransientPackage(), MakeUniqueObjectName(GetTransientPackage(), UMaterialGraph::StaticClass()), UMaterialGraph::StaticClass(), UMaterialGraphSchema::StaticClass()));
+		UMaterialGraph* NewMaterialGraph = CastChecked<UMaterialGraph>(FBlueprintEditorUtils::CreateNewGraph(GetTransientPackage(), MakeUniqueObjectName(GetTransientPackage(), UMaterialGraph::StaticClass()), UMaterialGraph::StaticClass(), UMaterialGraphSchema::StaticClass()));
+
+		OldMaterialGraph->MaterialFunction = OldMaterialFunction;
+		OldMaterialGraph->Material = NewObject<UMaterial>();
+		{
+			FArchiveUObject DummyArchive;
+			// Hack: serialize the new material with an archive that does nothing so that its material resources are created
+			OldMaterialGraph->Material->Serialize(DummyArchive);
+		}
+		OldMaterialGraph->Material->SetShadingModel(MSM_Unlit);
+		OldMaterialGraph->Material->AssignExpressionCollection(OldMaterialGraph->MaterialFunction->GetExpressionCollection());
+
+		NewMaterialGraph->MaterialFunction = NewMaterialFunction;
+		NewMaterialGraph->Material = NewObject<UMaterial>();
+		{
+			FArchiveUObject DummyArchive;
+			// Hack: serialize the new material with an archive that does nothing so that its material resources are created
+			NewMaterialGraph->Material->Serialize(DummyArchive);
+		}
+		NewMaterialGraph->Material->SetShadingModel(MSM_Unlit);
+		NewMaterialGraph->Material->AssignExpressionCollection(NewMaterialGraph->MaterialFunction->GetExpressionCollection());
+
+		OldMaterialGraph->RebuildGraph();
+		NewMaterialGraph->RebuildGraph();
+
+		SMaterialDiff::CreateDiffWindow(OldMaterialGraph, NewMaterialGraph, DiffArgs.OldRevision, DiffArgs.NewRevision, GetAssetClass().Get());
+		return EAssetCommandResult::Handled;
+	}
+
+	return EAssetCommandResult::Unhandled;
+}
+
 // Menu Extensions
 //--------------------------------------------------------------------
 
diff --git a/Engine/Plugins/Editor/EngineAssetDefinitions/Source/Private/Material/AssetDefinition_MaterialFunction.h b/Engine/Plugins/Editor/EngineAssetDefinitions/Source/Private/Material/AssetDefinition_MaterialFunction.h
--- a/Engine/Plugins/Editor/EngineAssetDefinitions/Source/Private/Material/AssetDefinition_MaterialFunction.h
+++ b/Engine/Plugins/Editor/EngineAssetDefinitions/Source/Private/Material/AssetDefinition_MaterialFunction.h
@@ -24,6 +24,7 @@ public:
 		return Categories;
 	}
 	virtual UThumbnailInfo* LoadThumbnailInfo(const FAssetData& InAsset) const override;
+	virtual EAssetCommandResult PerformAssetDiff(const FAssetDiffArgs& DiffArgs) const override;
 	virtual EAssetCommandResult OpenAssets(const FAssetOpenArgs& OpenArgs) const override;
 	// UAssetDefinition End
 };
diff --git a/Engine/Plugins/Editor/EngineAssetDefinitions/Source/Private/Material/AssetDefinition_MaterialInstanceConstant.cpp b/Engine/Plugins/Editor/EngineAssetDefinitions/Source/Private/Material/AssetDefinition_MaterialInstanceConstant.cpp
--- a/Engine/Plugins/Editor/EngineAssetDefinitions/Source/Private/Material/AssetDefinition_MaterialInstanceConstant.cpp
+++ b/Engine/Plugins/Editor/EngineAssetDefinitions/Source/Private/Material/AssetDefinition_MaterialInstanceConstant.cpp
@@ -8,6 +8,7 @@
 #include "MaterialEditorModule.h"
 #include "ToolMenu.h"
 #include "ToolMenuSection.h"
+#include "DiffTool/Widgets/SMaterialInstanceDiff.h"
 
 #include UE_INLINE_GENERATED_CPP_BY_NAME(AssetDefinition_MaterialInstanceConstant)
 
@@ -24,6 +25,21 @@ EAssetCommandResult UAssetDefinition_MaterialInstanceConstant::OpenAssets(const
 	return EAssetCommandResult::Handled;
 }
 
+EAssetCommandResult UAssetDefinition_MaterialInstanceConstant::PerformAssetDiff(const FAssetDiffArgs& DiffArgs) const
+{
+	UMaterialInstance* OldMaterialInstance = Cast<UMaterialInstance>(DiffArgs.OldAsset);
+	UMaterialInstance* NewMaterialInstance = Cast<UMaterialInstance>(DiffArgs.NewAsset);
+
+	if (OldMaterialInstance && NewMaterialInstance)
+	{
+		SMaterialInstanceDiff::CreateDiffWindow(OldMaterialInstance, NewMaterialInstance, DiffArgs.OldRevision, DiffArgs.NewRevision, GetAssetClass().Get());
+
+		return EAssetCommandResult::Handled;
+	}
+
+	return EAssetCommandResult::Unhandled;
+}
+
 // Menu Extensions
 //--------------------------------------------------------------------
 
diff --git a/Engine/Plugins/Editor/EngineAssetDefinitions/Source/Private/Material/AssetDefinition_MaterialInstanceConstant.h b/Engine/Plugins/Editor/EngineAssetDefinitions/Source/Private/Material/AssetDefinition_MaterialInstanceConstant.h
--- a/Engine/Plugins/Editor/EngineAssetDefinitions/Source/Private/Material/AssetDefinition_MaterialInstanceConstant.h
+++ b/Engine/Plugins/Editor/EngineAssetDefinitions/Source/Private/Material/AssetDefinition_MaterialInstanceConstant.h
@@ -28,6 +28,7 @@ public:
 		return Categories;
 	}
 	virtual EAssetCommandResult OpenAssets(const FAssetOpenArgs& OpenArgs) const override;
+	virtual EAssetCommandResult PerformAssetDiff(const FAssetDiffArgs& DiffArgs) const override;
 	virtual bool CanImport() const override { return true; }
 	// UAssetDefinition End
 };
diff --git a/Engine/Plugins/Runtime/StateTree/Source/StateTreeEditorModule/Private/StateTreeDiffControl.h b/Engine/Plugins/Runtime/StateTree/Source/StateTreeEditorModule/Private/StateTreeDiffControl.h
--- a/Engine/Plugins/Runtime/StateTree/Source/StateTreeEditorModule/Private/StateTreeDiffControl.h
+++ b/Engine/Plugins/Runtime/StateTree/Source/StateTreeEditorModule/Private/StateTreeDiffControl.h
@@ -4,7 +4,7 @@
 
 #include "AsyncStateTreeDiff.h"
 #include "DiffUtils.h"
-#include "Editor/Kismet/Private/DiffControl.h"
+#include "DiffControl.h"
 #include "Internationalization/Text.h"
 #include "UObject/ObjectKey.h"
 #include "UObject/StrongObjectPtr.h"
diff --git a/Engine/Source/Editor/GraphEditor/Private/GraphDiffControl.cpp b/Engine/Source/Editor/GraphEditor/Private/GraphDiffControl.cpp
--- a/Engine/Source/Editor/GraphEditor/Private/GraphDiffControl.cpp
+++ b/Engine/Source/Editor/GraphEditor/Private/GraphDiffControl.cpp
@@ -487,9 +487,39 @@ static UEdGraphPin* FindOtherLink(TArray<UEdGraphPin*>& Links2, int32 OriginalIn
 	return BestMatch;
 }
 
+static bool CheckGraphNodePin(UEdGraphPin* OldLinkedPin, UEdGraphPin* NewLinkedPin)
+{
+	if (!OldLinkedPin || !NewLinkedPin)
+	{
+		return false;
+	}
+
+	for (int32 i = 0; i < OldLinkedPin->LinkedTo.Num(); ++i)
+	{
+		UEdGraphPin* OldLinkedToPin = OldLinkedPin->LinkedTo[i];
+		UEdGraphPin* NewLinkedToPin = NewLinkedPin->LinkedTo.IsValidIndex(i) ? FindOtherLink(NewLinkedPin->LinkedTo, i, OldLinkedToPin) : nullptr;
+
+		if (!OldLinkedToPin || !NewLinkedToPin)
+		{
+			continue;
+		}
+
+		UEdGraphNode* OldLinkedToNode = OldLinkedToPin->GetOwningNode();
+		UEdGraphNode* NewLinkedToNode = NewLinkedToPin->GetOwningNode();
+
+		if (OldLinkedToNode && NewLinkedToNode && !FGraphDiffControl::IsNodeMatch(OldLinkedToNode, NewLinkedToNode))
+		{
+			return false;
+		}
+	}
+
+	return true;
+}
+
 /** Determine if the LinkedTo pins are the same */
 static bool LinkedToDifferent(UEdGraphPin* OldPin, UEdGraphPin* NewPin, const TArray<UEdGraphPin*>& OldLinks, TArray<UEdGraphPin*>& NewLinks, FDiffResults& Results)
 {
+	const UPackage* const TransientPackage = GetTransientPackage();
 	const int32 Size = OldLinks.Num();
 	bool bHasResult = false;
 	for(int32 i = 0;i<Size;++i)
@@ -512,6 +542,13 @@ static bool LinkedToDifferent(UEdGraphPin* OldPin, UEdGraphPin* NewPin, const TA
 			DiffR_LinkedToNode(Results, OldPin, NewPin, OldNode, NewNode);
 			KEEP_GOING_IF_RESULTS()
 		}
+		// ignore PinId for generated graphs, like the material graphs - instead do a more expensive comparison
+		// of the pin's node to detect changed links:
+		else if (OldNode->GetPackage() == TransientPackage && NewNode->GetPackage() == TransientPackage
+			&& CheckGraphNodePin(OldLinkedPin, NewLinkedPin))
+		{
+			continue;
+		}
 		else if (OldLinkedPin->PinId != NewLinkedPin->PinId)
 		{
 			DiffR_LinkedToPin(Results, OldPin, NewPin, OldLinkedPin, NewLinkedPin);
diff --git a/Engine/Source/Editor/Kismet/Internal/DiffControl.h b/Engine/Source/Editor/Kismet/Internal/DiffControl.h
new file mode 100644
--- /dev/null
+++ b/Engine/Source/Editor/Kismet/Internal/DiffControl.h
@@ -0,0 +1,418 @@
+// Copyright Epic Games, Inc. All Rights Reserved.
+#pragma once
+
+#include "SCSDiff.h"
+#if UE_ENABLE_INCLUDE_ORDER_DEPRECATED_IN_5_5
+#include "CoreMinimal.h"
+#include "Algo/Transform.h"
+#endif // UE_ENABLE_INCLUDE_ORDER_DEPRECATED_IN_5_5
+#include "DetailsDiff.h"
+#include "DiffResults.h"
+#include "IAssetTypeActions.h"
+#include "IDetailsView.h"
+#include "ReviewComments.h"
+#include "Widgets/SWidget.h"
+#include "DiffUtils.h"
+
+struct FReviewCommentsDiffControl;
+class SMultiLineEditableTextBox;
+struct FDiffResultItem;
+class SBlueprintDiff;
+class UEdGraph;
+struct FEdGraphEditAction;
+class SCheckBox;
+
+class FAsyncDetailViewDiff;
+
+/** Interface responsible for generating FBlueprintDifferenceTreeEntry's for visual diff tools */
+class IDiffControl
+{
+public:
+	virtual ~IDiffControl() {}
+
+	virtual void Tick() {};
+	
+	/** Adds widgets to the tree of differences to show */
+	virtual void GenerateTreeEntries(TArray< TSharedPtr<FBlueprintDifferenceTreeEntry> >& OutTreeEntries, TArray< TSharedPtr<FBlueprintDifferenceTreeEntry> >& OutRealDifferences) = 0;
+	
+	static KISMET_API FText RightRevision;
+	static KISMET_API TSharedRef<SWidget> GenerateSimpleDiffWidget(FText DiffText);
+	static KISMET_API TSharedRef<SWidget> GenerateObjectDiffWidget(FSingleObjectDiffEntry DiffEntry, FText ObjectName);
+
+	// to support comment posting, set this to the tree view that contains the comments.
+	KISMET_API void EnableComments(TWeakPtr<STreeView<TSharedPtr<FBlueprintDifferenceTreeEntry>>> TreeView, const UObject* OldObject, const UObject* NewObject);
+	KISMET_API void EnableComments(TWeakPtr<STreeView<TSharedPtr<FBlueprintDifferenceTreeEntry>>> TreeView, const TArray<const UObject*>& Objects);
+	virtual void EnableComments(TWeakPtr<STreeView<TSharedPtr<FBlueprintDifferenceTreeEntry>>> TreeView) {}
+	
+protected:
+	// to be called at the end of GenerateTreeEntries if you want this diff control to support review comments
+	KISMET_API void GenerateCategoryCommentTreeEntries(TArray< TSharedPtr<FBlueprintDifferenceTreeEntry> >& , TArray< TSharedPtr<FBlueprintDifferenceTreeEntry> >& OutRealDifferences, FString CategoryKey);
+	
+	TSharedPtr<FReviewCommentsDiffControl> ReviewCommentsDiffControl;
+};
+
+/** Shows all differences for the blueprint structure itself that aren't picked up elsewhere */
+class FMyBlueprintDiffControl : public TSharedFromThis<FMyBlueprintDiffControl>, public IDiffControl
+{
+public:
+	FMyBlueprintDiffControl(const UBlueprint* InOldBlueprint, const UBlueprint* InNewBlueprint, FOnDiffEntryFocused InSelectionCallback);
+
+	virtual void GenerateTreeEntries(TArray< TSharedPtr<FBlueprintDifferenceTreeEntry> >& OutTreeEntries, TArray< TSharedPtr<FBlueprintDifferenceTreeEntry> >& OutRealDifferences) override;
+
+	// to support comment posting, set this to the tree view that contains the comments.
+	virtual void EnableComments(TWeakPtr<STreeView<TSharedPtr<FBlueprintDifferenceTreeEntry>>> TreeView) override;
+
+private:
+	FOnDiffEntryFocused SelectionCallback;
+	const UBlueprint* OldBlueprint;
+	const UBlueprint* NewBlueprint;
+};
+
+/** 
+ * Each difference in the tree will either be a tree node that is added in one Blueprint 
+ * or a tree node and an FName of a property that has been added or edited in one Blueprint
+ */
+class FSCSDiffControl : public TSharedFromThis<FSCSDiffControl>, public IDiffControl
+{
+public:
+	FSCSDiffControl(const UBlueprint* InOldBlueprint, const UBlueprint* InNewBlueprint, FOnDiffEntryFocused InSelectionCallback);
+
+	virtual void GenerateTreeEntries(TArray< TSharedPtr<FBlueprintDifferenceTreeEntry> >& OutTreeEntries, TArray< TSharedPtr<FBlueprintDifferenceTreeEntry> >& OutRealDifferences) override;
+
+	TSharedRef<SWidget> OldTreeWidget() { return OldSCS.TreeWidget(); }
+	TSharedRef<SWidget> NewTreeWidget() { return NewSCS.TreeWidget(); }
+
+	// to support comment posting, set this to the tree view that contains the comments.
+	virtual void EnableComments(TWeakPtr<STreeView<TSharedPtr<FBlueprintDifferenceTreeEntry>>> TreeView) override;
+
+private:
+	FOnDiffEntryFocused SelectionCallback;
+	FSCSDiffRoot DifferingProperties;
+
+	FSCSDiff OldSCS;
+	FSCSDiff NewSCS;
+};
+
+/** Generic wrapper around a details view template parameter determines whether TreeEntries are populated */
+class FDetailsDiffControl : public TSharedFromThis<FDetailsDiffControl>, public IDiffControl
+{
+public:
+	KISMET_API FDetailsDiffControl(const UObject* InOldObject, const UObject* InNewObject, FOnDiffEntryFocused InSelectionCallback, bool bPopulateOutTreeEntries);
+	KISMET_API ~FDetailsDiffControl();
+
+	KISMET_API virtual void Tick() override;
+	
+	KISMET_API virtual void GenerateTreeEntries(TArray< TSharedPtr<FBlueprintDifferenceTreeEntry> >& OutTreeEntries, TArray< TSharedPtr<FBlueprintDifferenceTreeEntry> >& OutRealDifferences) override;
+	KISMET_API void GenerateTreeEntriesWithoutComments(TArray< TSharedPtr<FBlueprintDifferenceTreeEntry> >& OutTreeEntries, TArray< TSharedPtr<FBlueprintDifferenceTreeEntry> >& OutRealDifferences);
+
+	KISMET_API TSharedRef<IDetailsView> InsertObject(const UObject* Object, bool bScrollbarOnLeft = false, int32 Index = INDEX_NONE);
+	KISMET_API TSharedRef<IDetailsView> GetDetailsWidget(const UObject* Object) const;
+	KISMET_API TSharedPtr<IDetailsView> TryGetDetailsWidget(const UObject* Object) const;
+	KISMET_API TSharedPtr<FAsyncDetailViewDiff> GetDifferencesWithLeft(const UObject* Object) const;
+	KISMET_API TSharedPtr<FAsyncDetailViewDiff> GetDifferencesWithRight(const UObject* Object) const;
+	KISMET_API int32 IndexOfObject(const UObject* Object) const;
+
+	// to support comment posting, set this to the tree view that contains the comments.
+	KISMET_API virtual void EnableComments(TWeakPtr<STreeView<TSharedPtr<FBlueprintDifferenceTreeEntry>>> TreeView) override;
+
+	DiffUtils::FOnGenerateCustomDiffEntries GenerateCustomEntriesCallback;
+	DiffUtils::FOnGenerateCustomDiffEntryWidget GenerateCustomEntryWidgetCallback;
+	DiffUtils::FOnOrganizeDiffEntries OrganizeEntriesCallback;
+protected:
+	KISMET_API virtual void OnSelectDiffEntry(FPropertySoftPath PropertyName);
+	KISMET_API TAttribute<TArray<FVector2f>> GetLinkedScrollRateAttribute(const TSharedRef<IDetailsView>& OldDetailsView, const TSharedRef<IDetailsView>& NewDetailsView);
+	
+
+	// helper function that analyzes two details views and determines the rate they should scroll relative to one another to be in sync
+	KISMET_API TArray<FVector2f> GetLinkedScrollRate(TSharedRef<IDetailsView> LeftDetailsView, TSharedRef<IDetailsView> RightDetailsView) const;
+private:
+	KISMET_API void HandlePropertiesChanged();
+
+	FDetailsDiffControl& operator=(const FDetailsDiffControl&) = delete;
+	FDetailsDiffControl( const FDetailsDiffControl& ) = delete;
+	FDetailsDiffControl& operator=(FDetailsDiffControl&&) = delete;
+	FDetailsDiffControl( FDetailsDiffControl&& ) = delete;
+protected:
+
+	FOnDiffEntryFocused SelectionCallback;
+	
+	TMap<const UObject*,FDetailsDiff> DetailsDiffs;
+	TArray<const UObject*> ObjectDisplayOrder;
+	struct FPropertyTreeDiffPairs
+	{
+		TSharedPtr<FAsyncDetailViewDiff> Left;
+		TSharedPtr<FAsyncDetailViewDiff> Right;
+	};
+	TMap<const UObject*, FPropertyTreeDiffPairs> PropertyTreeDifferences;
+
+	TArray< TSharedPtr<FBlueprintDifferenceTreeEntry> > Children;
+	const bool bPopulateOutTreeEntries;
+	
+	TSet<FPropertyPath> PropertyAllowList;
+};
+
+/** Override for CDO special case */
+class FCDODiffControl : public FDetailsDiffControl
+{
+public:
+	FCDODiffControl(const UObject* InOldObject, const UObject* InNewObject, FOnDiffEntryFocused InSelectionCallback);
+
+	virtual void GenerateTreeEntries(TArray< TSharedPtr<FBlueprintDifferenceTreeEntry> >& OutTreeEntries, TArray< TSharedPtr<FBlueprintDifferenceTreeEntry> >& OutRealDifferences) override;
+};
+
+/** Override for class class settings */
+class FClassSettingsDiffControl : public FDetailsDiffControl
+{
+public:
+	FClassSettingsDiffControl(const UObject* InOldObject, const UObject* InNewObject, FOnDiffEntryFocused InSelectionCallback);
+
+	virtual void GenerateTreeEntries(TArray< TSharedPtr<FBlueprintDifferenceTreeEntry> >& OutTreeEntries, TArray< TSharedPtr<FBlueprintDifferenceTreeEntry> >& OutRealDifferences) override;
+};
+
+/** Diff control to handle finding type-specific differences */
+struct FBlueprintTypeDiffControl : public TSharedFromThis<FBlueprintTypeDiffControl>, public IDiffControl
+{
+	struct FSubObjectDiff
+	{
+		FDiffSingleResult SourceResult;
+		FDetailsDiff OldDetails;
+		FDetailsDiff NewDetails;
+		TArray<TSharedPtr<FBlueprintDifferenceTreeEntry>> Diffs;
+
+		FSubObjectDiff(const FDiffSingleResult& InSourceResult, const UObject* OldObject, const UObject* NewObject)
+			: SourceResult(InSourceResult)
+			, OldDetails(OldObject)
+			, NewDetails(NewObject)
+		{}
+	};
+
+	struct FDiffCustomObject
+	{
+		FDiffSingleResult SourceResult;
+		TSharedPtr<IDiffCustomObject> DiffWidget;
+		TArray<TSharedPtr<FBlueprintDifferenceTreeEntry>> Diffs;
+		FDiffCustomObject(const FDiffSingleResult& InSourceResult, TSharedPtr<IDiffCustomObject> DiffWidget)
+			: SourceResult(InSourceResult)
+			, DiffWidget(DiffWidget)
+		{}
+	};
+
+	FBlueprintTypeDiffControl(const UBlueprint* InBlueprintOld, const UBlueprint* InBlueprintNew, FOnDiffEntryFocused InSelectionCallback);
+
+	/** Generate difference tree widgets */
+	virtual void GenerateTreeEntries(TArray< TSharedPtr<FBlueprintDifferenceTreeEntry> >& OutTreeEntries, TArray< TSharedPtr<FBlueprintDifferenceTreeEntry> >& OutRealDifferences) override;
+	
+	// to support comment posting, set this to the tree view that contains the comments.
+	virtual void EnableComments(TWeakPtr<STreeView<TSharedPtr<FBlueprintDifferenceTreeEntry>>> TreeView) override;
+	
+	/** The old blueprint (left) */
+	const UBlueprint* BlueprintOld;
+
+	/** The new blueprint(right) */
+	const UBlueprint* BlueprintNew;
+
+	/** Boxes that will display the details diffs */
+	TSharedPtr<SBox> OldDetailsBox;
+	TSharedPtr<SBox> NewDetailsBox;
+
+private:
+	/** Generate Widget for top category */
+	TSharedRef<SWidget> GenerateCategoryWidget(bool bHasRealDiffs);
+
+	/** Build up the Diff Source Array*/
+	void BuildDiffSourceArray();
+
+	/** Handle selecting a diff */
+	void OnSelectSubobjectDiff(FPropertySoftPath Identifier, TSharedPtr<FSubObjectDiff> SubObjectDiff);
+
+	/** Handle selecting a custom object diff */
+	void OnSelectDiffCustomObject(TSharedPtr<IDiffCustomProperty> CustomProperty, TSharedPtr<FDiffCustomObject> DiffCustomObject);
+
+	/** List of objects with differences */
+	TArray<TSharedPtr<FSubObjectDiff>> SubObjectDiffs;
+	
+	/** List of custom object differences */
+	TArray<TSharedPtr<FDiffCustomObject>> DiffCustomObjects;
+
+	/** Source for list view */
+	TArray<TSharedPtr<FDiffResultItem>> DiffListSource;
+
+	/** Selection callback */
+	FOnDiffEntryFocused SelectionCallback;
+
+	/** Did diff generation succeed? */
+	bool bDiffSucceeded;
+};
+
+
+/** Category list item for a graph*/
+struct FGraphToDiff	: public TSharedFromThis<FGraphToDiff>, IDiffControl
+{
+	FGraphToDiff(SBlueprintDiff* DiffWidget, UEdGraph* GraphOld, UEdGraph* GraphNew, const FRevisionInfo& RevisionOld, const FRevisionInfo& RevisionNew);
+	virtual ~FGraphToDiff() override;
+
+	/** Add widgets to the differences tree */
+	virtual void GenerateTreeEntries(TArray< TSharedPtr<FBlueprintDifferenceTreeEntry> >& OutTreeEntries, TArray< TSharedPtr<FBlueprintDifferenceTreeEntry> >& OutRealDifferences) override;
+
+	/** Get old(left) graph*/
+	UEdGraph* GetGraphOld() const;
+
+	/** Get new(right) graph*/
+	UEdGraph* GetGraphNew() const;
+
+	// to support comment posting, set this to the tree view that contains the comments.
+	virtual void EnableComments(TWeakPtr<STreeView<TSharedPtr<FBlueprintDifferenceTreeEntry>>> TreeView) override;
+	
+	/** Source for list view */
+	TArray<TSharedPtr<FDiffResultItem>> DiffListSource;
+	TSharedPtr<TArray<FDiffSingleResult>> FoundDiffs;
+
+	/** Index of the first item in RealDifferences that was generated by this graph */
+	int32 RealDifferencesStartIndex = INDEX_NONE;
+
+private:
+	/** Get tooltip for category */
+	FText GetToolTip();
+
+	/** Generate Widget for category list */
+	TSharedRef<SWidget> GenerateCategoryWidget();
+
+	/** Called when the Newer Graph is modified*/
+	void OnGraphChanged(const FEdGraphEditAction& Action);
+
+	/** Build up the Diff Source Array*/
+	void BuildDiffSourceArray();
+
+	/** Diff widget */
+	class SBlueprintDiff* DiffWidget;
+
+	/** The old graph(left)*/
+	UEdGraph* GraphOld;
+
+	/** The new graph(right)*/
+	UEdGraph* GraphNew;
+
+	/** Description of Old and new graph*/
+	FRevisionInfo	RevisionOld, RevisionNew;
+
+	/** Handle to the registered OnGraphChanged delegate. */
+	FDelegateHandle OnGraphChangedDelegateHandle;
+};
+
+class FCommentTreeEntry : public FBlueprintDifferenceTreeEntry, public TSharedFromThis<FCommentTreeEntry>
+{
+public:
+	FCommentTreeEntry(TWeakPtr<FReviewCommentsDiffControl> InCommentsControl, const FReviewComment& InComment, const TArray<TSharedPtr<FBlueprintDifferenceTreeEntry>>& InChildren);
+	virtual ~FCommentTreeEntry() override;
+	static TSharedRef<FCommentTreeEntry> Make(TWeakPtr<FReviewCommentsDiffControl> CommentsControl, const FReviewComment& Comment, const TArray<TSharedPtr<FBlueprintDifferenceTreeEntry>>& Children = {});
+
+	int32 GetCommentIDChecked() const;
+	void AwaitCommentPost();
+private:
+	// FBlueprintDifferenceTreeEntry callbacks
+	TSharedRef<SWidget> CreateWidget();
+
+	// slate callbacks
+	bool IsCommentTextBoxReadonly() const;
+	FReply OnClickReply();
+	FReply OnClickEdit();
+	FReply OnLikeToggle();
+	FText GetLikeTooltip() const;
+	const FSlateBrush* GetLikeIcon() const;
+	FSlateColor GetUsernameColor() const;
+	FSlateColor GetLikeIconColor() const;
+	EVisibility GetEditReplyButtonGroupVisibility() const;
+	EVisibility GetEditButtonVisibility() const;
+	EVisibility GetSubmitCancelButtonGroupVisibility() const;
+	bool IsSubmitButtonEnabled() const;
+	FReply OnEditSubmitClicked();
+	FReply OnEditCancelClicked();
+
+	// comment api callbacks
+	void OnCommentPosted(const FReviewComment& Comment);
+
+	FString GetCommentString() const;
+	void SetCommentString(const FString& NewComment);
+
+	// during a pending edit, returns whether the comment is different from the saved version
+	bool HasCommentStringChanged() const;
+	
+	bool IsEditMode() const;
+	void SetEditMode(bool bIsEditMode);
+	
+	FReviewComment Comment;
+
+	TSharedPtr<SWidget> Content;
+	TSharedPtr<SHorizontalBox> EditReplyButtonGroup;
+	TSharedPtr<SHorizontalBox> SubmitCancelButtonGroup;
+	TSharedPtr<SMultiLineEditableTextBox> CommentTextBox;
+	
+	TWeakPtr<FReviewCommentsDiffControl> CommentsControl;
+	FDelegateHandle OnCommentPostedHandle;
+
+	// prefer the use of IsEditMode().
+	bool bExpectedEditMode = false;
+};
+
+class FCommentDraftTreeEntry : public FBlueprintDifferenceTreeEntry, public TSharedFromThis<FCommentDraftTreeEntry>
+{
+public:
+	FCommentDraftTreeEntry(TWeakPtr<FReviewCommentsDiffControl> InCommentsControl, TArray<TSharedPtr<FBlueprintDifferenceTreeEntry>>* InSiblings, int32 InReplyID = -1);
+	static TSharedRef<FCommentDraftTreeEntry> MakeCommentDraft(TWeakPtr<FReviewCommentsDiffControl> CommentsControl, TArray<TSharedPtr<FBlueprintDifferenceTreeEntry>>* Siblings);
+	static TSharedRef<FCommentDraftTreeEntry> MakeReplyDraft(TWeakPtr<FReviewCommentsDiffControl> CommentsControl, TSharedPtr<FCommentTreeEntry> InParent);
+
+	void ReassignReplyParent(TSharedPtr<FCommentTreeEntry> InParent);
+
+	TSharedPtr<SMultiLineEditableTextBox> GetCommentTextBox() const {return CommentTextBox;}
+
+	bool IsReply() const;
+private:
+	// FBlueprintDifferenceTreeEntry callbacks
+	TSharedRef<SWidget> CreateWidget();
+	
+	// slate callbacks
+	bool IsPostButtonEnabled() const;
+	FReply OnCommentPostClicked();
+	
+	void ReassignSiblings(TArray<TSharedPtr<FBlueprintDifferenceTreeEntry>>* InSiblings, int32 InReplyID);
+	
+	TSharedPtr<SWidget> Content;
+	TSharedPtr<SMultiLineEditableTextBox> CommentTextBox;
+	TSharedPtr<SCheckBox> FlagAsTaskCheckBox;
+	
+	TWeakPtr<FReviewCommentsDiffControl> CommentsControl;
+	
+	TArray<TSharedPtr<FBlueprintDifferenceTreeEntry>>* Siblings;
+	
+	// if this is a reply, store the comment id of it's parent
+	int32 ReplyID = -1;
+};
+
+/** Category list item for a graph*/
+struct FReviewCommentsDiffControl : public TSharedFromThis<FReviewCommentsDiffControl>, IDiffControl
+{
+	FReviewCommentsDiffControl(const FString& InCommentFilePath, TWeakPtr<STreeView<TSharedPtr<FBlueprintDifferenceTreeEntry>>> TreeView);
+
+	virtual void GenerateTreeEntries(TArray< TSharedPtr<FBlueprintDifferenceTreeEntry> >& OutTreeEntries, TArray< TSharedPtr<FBlueprintDifferenceTreeEntry> >& OutRealDifferences) override;
+
+	void SetCategory(const FString& CategoryKey);
+
+	void PostComment(FReviewComment& Comment);
+	void DraftReply(TSharedPtr<FCommentTreeEntry> ParentComment);
+	void RebuildListView() const;
+
+	const FString& GetCommentFilePath() {return CommentFilePath;}
+	const FString& GetCommentCategory() {return CommentCategory;}
+protected:
+	
+	void GenerateCommentThreadRecursive(const FReviewComment& Comment,
+		const TMap<int32, TArray<const FReviewComment*>>& CommentReplyMap,
+		TArray<TSharedPtr<FBlueprintDifferenceTreeEntry>>& OutTreeEntries);
+	
+	FString CommentFilePath;
+	FString CommentCategory;
+	TWeakPtr<STreeView<TSharedPtr<FBlueprintDifferenceTreeEntry>>> CommentsTreeView;
+
+	// we only want one reply draft widget in existence so we can use a static to track it
+	static TWeakPtr<FCommentDraftTreeEntry> ReplyDraftEntry;
+};
diff --git a/Engine/Source/Editor/Kismet/Private/DiffControl.h b/Engine/Source/Editor/Kismet/Private/DiffControl.h
deleted file mode 100644
diff --git a/Engine/Source/Editor/MaterialEditor/Internal/DiffTool/MaterialDiffPanel.h b/Engine/Source/Editor/MaterialEditor/Internal/DiffTool/MaterialDiffPanel.h
new file mode 100644
--- /dev/null
+++ b/Engine/Source/Editor/MaterialEditor/Internal/DiffTool/MaterialDiffPanel.h
@@ -0,0 +1,102 @@
+// Copyright Epic Games, Inc. All Rights Reserved.
+#pragma once
+
+#include "AssetDefinition.h"
+#include "CoreMinimal.h"
+#include "GraphEditor.h"
+#include "PropertyPath.h"
+
+class IDetailsView;
+class SBox;
+class SKismetInspector;
+class SMaterialEditor3DPreviewViewport;
+class SMaterialEditorUIPreviewViewport;
+class SMyBlueprint;
+class UMaterial;
+class UMaterialExpression;
+class UMaterialGraph;
+
+struct FRevisionInfo;
+
+/** Panel used to display the MaterialGraph */
+struct FMaterialDiffPanel : public FGCObject
+{
+	/** Generate a panel for NewGraph diffed against OldGraph */
+	void GeneratePanel(UEdGraph* NewGraph, UEdGraph* OldGraph);
+
+	/** Generate a panel that displays the Graph and reflects the items in the DiffResults */
+	void GeneratePanel(UEdGraph* Graph, TSharedPtr<TArray<FDiffSingleResult>> DiffResults, TAttribute<int32> FocusedDiffResult);
+
+	/** Called when user hits keyboard shortcut to copy nodes */
+	void CopySelectedNodes();
+
+	/** Gets whatever nodes are selected in the Graph Editor */
+	FGraphPanelSelectionSet GetSelectedNodes() const;
+
+	/** Can user copy any of the selected nodes? */
+	bool CanCopyNodes() const;
+
+	void SelectDownstreamNodes() const;
+	void SelectUpstreamNodes() const;
+	
+	void OnPreviewNode();
+
+	/** Functions used to focus/find a particular change in a diff result */
+	void FocusDiff(UEdGraphPin& Pin);
+	void FocusDiff(UEdGraphNode& Node);
+
+	TSharedRef<SWidget> GetMaterialNodeDetailsViewWidget() const;
+	
+	void SetViewportToDisplay();
+	
+	TSharedRef<SWidget> GetViewportToDisplay() const;
+	
+	void SetViewportVisibility(bool bShowViewport);
+
+	void SetPreviewExpression(UMaterialExpression* NewPreviewExpression);
+
+	void SetPreviewMaterial(UMaterialInterface* InMaterialInterface) const;
+
+	void UpdatePreviewMaterial() const;
+
+	/** The MaterialGraph we are showing */
+	TObjectPtr<UMaterialGraph> MaterialGraph = nullptr;
+
+	/** The box around the graph editor, used to change the content when new graphs are set */
+	TSharedPtr<SBox> GraphEditorBox;
+
+	/** The details view associated with the MaterialNode of the focused MaterialGraphNode */
+	TSharedPtr<IDetailsView> MaterialNodeDetailsView;
+
+	TSharedPtr<SMaterialEditor3DPreviewViewport> Preview3DViewport;
+
+	TSharedPtr<SMaterialEditorUIPreviewViewport> Preview2DViewport;
+
+	TObjectPtr<UMaterial> ExpressionPreviewMaterial = nullptr;
+
+	TObjectPtr<UMaterialExpression> PreviewExpression = nullptr;
+
+	/** The graph editor which does the work of displaying the graph */
+	TWeakPtr<SGraphEditor> GraphEditor;
+
+	/** Revision information for this Material */
+	FRevisionInfo RevisionInfo;
+
+	/** True if we should show a name identifying which asset this panel is displaying */
+	bool bShowAssetName = true;
+
+	/** The widget that contains the revision info in graph mode */
+	TSharedPtr<SWidget> OverlayGraphRevisionInfo;
+
+	/** The property that we will highlight in MaterialNodeDetailsView */
+	FPropertyPath PropertyToFocus;
+	
+	// Begin FGCObject interface
+	virtual void AddReferencedObjects(FReferenceCollector& Collector) override;
+	virtual FString GetReferencerName() const override;
+	// End FGCObject interface
+
+private:
+	/** Command list for this diff panel */
+	TSharedPtr<FUICommandList> GraphEditorCommands;
+};
\ No newline at end of file
diff --git a/Engine/Source/Editor/MaterialEditor/Internal/DiffTool/MaterialInstanceDiffPanel.h b/Engine/Source/Editor/MaterialEditor/Internal/DiffTool/MaterialInstanceDiffPanel.h
new file mode 100644
--- /dev/null
+++ b/Engine/Source/Editor/MaterialEditor/Internal/DiffTool/MaterialInstanceDiffPanel.h
@@ -0,0 +1,49 @@
+// Copyright Epic Games, Inc. All Rights Reserved.
+
+#pragma once
+
+#include "AssetDefinition.h"
+#include "CoreMinimal.h"
+#include "PropertyPath.h"
+#include "Widgets/SWidget.h"
+
+class UMaterialInterface;
+class SMaterialEditorUIPreviewViewport;
+class SMaterialEditor3DPreviewViewport;
+class UMaterial;
+class UMaterialExpression;
+class UMaterialInstance;
+
+struct FRevisionInfo;
+
+/** Panel used to display the MaterialInstance */
+struct FMaterialInstanceDiffPanel : public FGCObject
+{
+	void SetViewportToDisplay();
+	
+	TSharedRef<SWidget> GetViewportToDisplay() const;
+
+	void SetPreviewExpression(UMaterialExpression* NewPreviewExpression);
+
+	void SetPreviewMaterial(UMaterialInterface* InMaterialInterface) const;
+
+	void UpdatePreviewMaterial() const;
+
+	TSharedPtr<SMaterialEditor3DPreviewViewport> Preview3DViewport;
+
+	TSharedPtr<SMaterialEditorUIPreviewViewport> Preview2DViewport;
+
+	TObjectPtr<UMaterial> ExpressionPreviewMaterial = nullptr;
+
+	TObjectPtr<UMaterialExpression> PreviewExpression = nullptr;
+
+	TObjectPtr<UMaterialInstance> MaterialInstance = nullptr;
+
+	/** Revision information for this MaterialInstance */
+	FRevisionInfo RevisionInfo;
+	
+	// Begin FGCObject interface
+	virtual void AddReferencedObjects(FReferenceCollector& Collector) override;
+	virtual FString GetReferencerName() const override;
+	// End FGCObject interface
+};
\ No newline at end of file
diff --git a/Engine/Source/Editor/MaterialEditor/Internal/DiffTool/Widgets/SMaterialDiff.h b/Engine/Source/Editor/MaterialEditor/Internal/DiffTool/Widgets/SMaterialDiff.h
new file mode 100644
--- /dev/null
+++ b/Engine/Source/Editor/MaterialEditor/Internal/DiffTool/Widgets/SMaterialDiff.h
@@ -0,0 +1,185 @@
+// Copyright Epic Games, Inc. All Rights Reserved.
+
+#pragma once
+
+#include "AssetDefinition.h"
+#include "CoreMinimal.h"
+#include "DiffControl.h"
+#include "DiffTool/MaterialDiffPanel.h"
+#include "SBlueprintDiff.h"
+#include "Widgets/SCompoundWidget.h"
+
+class FBlueprintDifferenceTreeEntry;
+class UEdGraph;
+class UMaterialGraph;
+
+struct FMaterialToDiff;
+
+enum class EAssetEditorCloseReason : uint8;
+
+struct FMaterialDiffResultItem : public FDiffResultItem
+{
+	/** Not empty if the DiffResultItem is a property changed */
+	FPropertyPath Property;
+};
+
+/* Visual Diff between two Material Graphs */
+class MATERIALEDITOR_API SMaterialDiff : public SCompoundWidget
+{
+public:
+	SLATE_BEGIN_ARGS(SMaterialDiff) {}
+		SLATE_ARGUMENT(TObjectPtr<UMaterialGraph>, OldMaterialGraph)
+		SLATE_ARGUMENT(TObjectPtr<UMaterialGraph>, NewMaterialGraph)
+		SLATE_ARGUMENT(FRevisionInfo, OldRevision)
+		SLATE_ARGUMENT(FRevisionInfo, NewRevision)
+		SLATE_ARGUMENT(bool, ShowAssetNames)
+		SLATE_ARGUMENT(TSharedPtr<SWindow>, ParentWindow)
+	SLATE_END_ARGS()
+
+	void Construct(const FArguments& InArgs);
+	virtual ~SMaterialDiff() override;
+
+	/** Called when a new Graph is clicked on by user */
+	void OnGraphChanged(FMaterialToDiff* Diff);
+
+	/** Called when user clicks on a new graph list item */
+	void OnGraphSelectionChanged(TSharedPtr<FMaterialToDiff> Item, ESelectInfo::Type SelectionType);
+
+	/** Called when user clicks on an entry in the listview of differences */
+	void OnDiffListSelectionChanged(TSharedPtr<FMaterialDiffResultItem> TheDiff);
+
+	/** Helper function for generating an empty widget */
+	static TSharedRef<SWidget> DefaultEmptyPanel();
+
+	/** Helper function to create a window that holds a diff widget */
+	static TSharedPtr<SWindow> CreateDiffWindow(FText WindowTitle, TObjectPtr<UMaterialGraph> OldMaterialGraph, TObjectPtr<UMaterialGraph> NewMaterialGraph, const FRevisionInfo& OldRevision, const FRevisionInfo& NewRevision);
+
+	/** Helper function to create a window that holds a diff widget (default window title) */
+	static TSharedPtr<SWindow> CreateDiffWindow(TObjectPtr<UMaterialGraph> OldMaterialGraph, TObjectPtr<UMaterialGraph> NewMaterialGraph, const FRevisionInfo& OldRevision, const FRevisionInfo& NewRevision, const UClass* ObjectClass);
+
+	DECLARE_MULTICAST_DELEGATE_OneParam(FOnWindowClosedEvent, TSharedRef<SMaterialDiff>)
+	FOnWindowClosedEvent OnWindowClosedEvent;
+
+protected:
+	/** Called when user clicks button to go to next difference */
+	void NextDiff();
+
+	/** Called when user clicks button to go to prev difference */
+	void PrevDiff();
+
+	/** Called to determine whether we have a list of differences to cycle through */
+	bool HasNextDiff() const;
+	bool HasPrevDiff() const;
+
+	/** User toggles the option to show/hide preview viewports */
+	void ToggleViewport();
+
+	/** Get the image to show for the toggle viewport option */
+	FSlateIcon GetViewportImage() const;
+
+	/** Find the FMaterialToDiff that displays the graph with GraphPath relative path */
+	FMaterialToDiff* FindGraphToDiffEntry(const FString& GraphPath);
+
+	/** Bring these revisions of graph into focus on main display */
+	void FocusOnMaterialGraphRevisions(FMaterialToDiff* Diff);
+
+	/** User toggles the option to lock the views between the two material graphs */
+	void OnToggleLockView();
+
+	/** User toggles the option to change the split view mode betwwen vertical and horizontal */
+	void OnToggleSplitViewMode();
+
+	/** Get the image to show for the toggle lock option */
+	FSlateIcon GetLockViewImage() const;
+
+	/** Get the image to show for the toggle split view mode option */
+	FSlateIcon GetSplitViewModeImage() const;
+
+	/** Reset the graph editor, called when user switches graphs to display */
+	void ResetGraphEditors();
+
+	/** Material Graph to diff, is added to panel */
+	TSharedPtr<FMaterialToDiff> MaterialGraphToDiff;
+
+	/** Get Graph editor associated with this Graph */
+	FMaterialDiffPanel& GetDiffPanelForNode(UEdGraphNode& Node);
+
+	/** Event handler that updates the graph view when user selects a new graph */
+	void HandleGraphChanged(const FString& GraphPath);
+
+	/** Function used to generate the list of differences and the widgets needed to calculate that list */
+	void GenerateDifferencesList();
+
+	/** Called when editor may need to be closed */
+	void OnCloseAssetEditor(UObject* Asset, EAssetEditorCloseReason CloseReason);
+
+	struct FMaterialGraphDiffControl
+	{
+		FMaterialGraphDiffControl() : Widget(), DiffControl(nullptr)
+		{
+		}
+
+		TSharedPtr<SWidget> Widget;
+		TSharedPtr<IDiffControl> DiffControl;
+	};
+
+	FMaterialGraphDiffControl GenerateMaterialGraphPanel();
+
+	TSharedRef<SOverlay> GenerateMaterialGraphWidgetForPanel(FMaterialDiffPanel& OutDiffPanel) const;
+	TSharedRef<SBox> GenerateRevisionInfoWidgetForPanel(TSharedPtr<SWidget>& OutGeneratedWidget, const FText& InRevisionText) const;
+
+	/** Accessor and event handler for toggling between diff view modes - only GraphMode for now: */
+	void SetCurrentMode(FName NewMode);
+
+	FName GetCurrentMode() const { return CurrentMode; }
+
+	void OnModeChanged(const FName& InNewViewMode) const;
+
+	void UpdateTopSectionVisibility(const FName& InNewViewMode) const;
+
+	void SetCurrentWidgetIndex(int32 Index);
+	int32 GetCurrentWidgetIndex() const;
+
+	FName CurrentMode;
+
+	FMaterialDiffPanel PanelOld, PanelNew;
+
+	/** If the two views should be locked */
+	bool bLockViews = true;
+
+	/** If the view on Graph Mode should be divided vertically */
+	bool bVerticalSplitGraphMode = true;
+
+	/** If the preview Viewports should be shown */
+	bool bShowViewport = true;
+
+	/** Contents widget that we swap when mode changes (defaults, components, etc) */
+	TSharedPtr<SBox> ModeContents;
+
+	TSharedPtr<SSplitter> TopRevisionInfoWidget;
+
+	TSharedPtr<SSplitter> DiffGraphSplitter;
+
+	TSharedPtr<SSplitter> GraphToolBarWidget;
+
+	TSharedPtr<SWidgetSwitcher> WidgetSwitcher;
+
+	int32 CurrentWidgetIndex = 0;
+
+	/** Tree of differences collected across all panels: */
+	TArray<TSharedPtr<FBlueprintDifferenceTreeEntry>> PrimaryDifferencesList;
+
+	/** List of all differences, cached so that we can iterate only the differences and not labels, etc: */
+	TArray<TSharedPtr<FBlueprintDifferenceTreeEntry>> RealDifferences;
+
+	/** Tree view that displays the differences, cached for the buttons that iterate the differences: */
+	TSharedPtr<STreeView<TSharedPtr<FBlueprintDifferenceTreeEntry>>> DifferencesTreeView;
+
+	/** Stored references to widgets used to display various parts of a material, from the mode name */
+	TMap<FName, FMaterialGraphDiffControl> ModePanels;
+
+	/** A pointer to the window holding this */
+	TWeakPtr<SWindow> WeakParentWindow;
+
+	FDelegateHandle AssetEditorCloseDelegate;
+};
diff --git a/Engine/Source/Editor/MaterialEditor/Internal/DiffTool/Widgets/SMaterialInstanceDiff.h b/Engine/Source/Editor/MaterialEditor/Internal/DiffTool/Widgets/SMaterialInstanceDiff.h
new file mode 100644
--- /dev/null
+++ b/Engine/Source/Editor/MaterialEditor/Internal/DiffTool/Widgets/SMaterialInstanceDiff.h
@@ -0,0 +1,123 @@
+// Copyright Epic Games, Inc. All Rights Reserved.
+
+#pragma once
+
+#include "AssetDefinition.h"
+#include "CoreMinimal.h"
+#include "DiffTool/Widgets/SMaterialDiff.h"
+#include "DiffControl.h"
+#include "Materials/MaterialInstance.h"
+#include "SBlueprintDiff.h"
+#include "DiffTool/MaterialInstanceDiffPanel.h"
+#include "Widgets/SCompoundWidget.h"
+
+class SDetailsSplitter;
+class FBlueprintDifferenceTreeEntry;
+
+enum class EAssetEditorCloseReason : uint8;
+
+/* Visual Diff between two Material Instances */
+class MATERIALEDITOR_API SMaterialInstanceDiff : public SCompoundWidget
+{
+public:
+	SLATE_BEGIN_ARGS(SMaterialInstanceDiff) {}
+		SLATE_ARGUMENT(TObjectPtr<UMaterialInstance>, OldMaterialInstance)
+		SLATE_ARGUMENT(TObjectPtr<UMaterialInstance>, NewMaterialInstance)
+		SLATE_ARGUMENT(FRevisionInfo, OldRevision)
+		SLATE_ARGUMENT(FRevisionInfo, NewRevision)
+		SLATE_ARGUMENT(TSharedPtr<SWindow>, ParentWindow)
+	SLATE_END_ARGS()
+
+	void Construct(const FArguments& InArgs);
+	virtual ~SMaterialInstanceDiff() override;
+
+	/** Called when user clicks on an entry in the listview of differences */
+	void OnDiffListSelectionChanged(TSharedPtr<FMaterialDiffResultItem> TheDiff);
+
+	/** Helper function for generating an empty widget */
+	static TSharedRef<SWidget> DefaultEmptyPanel();
+
+	/** Helper function to create a window that holds a diff widget */
+	static TSharedPtr<SWindow> CreateDiffWindow(FText WindowTitle, TObjectPtr<UMaterialInstance>  OldMaterialInstance, TObjectPtr<UMaterialInstance>  NewMaterialInstance, const FRevisionInfo& OldRevision, const FRevisionInfo& NewRevision);
+
+	/** Helper function to create a window that holds a diff widget (default window title) */
+	static TSharedPtr<SWindow> CreateDiffWindow(TObjectPtr<UMaterialInstance>  OldMaterialInstance, TObjectPtr<UMaterialInstance>  NewMaterialInstance, const FRevisionInfo& OldRevision, const FRevisionInfo& NewRevision, const UClass* ObjectClass);
+
+	DECLARE_MULTICAST_DELEGATE_OneParam(FOnWindowClosedEvent, TSharedRef<SMaterialInstanceDiff>)
+	FOnWindowClosedEvent OnWindowClosedEvent;
+
+protected:
+	/** Called when user clicks button to go to next difference */
+	void NextDiff();
+
+	/** Called when user clicks button to go to prev difference */
+	void PrevDiff();
+
+	/** Called to determine whether we have a list of differences to cycle through */
+	bool HasNextDiff() const;
+	bool HasPrevDiff() const;
+
+	/** User toggles the option to show/hide preview viewports */
+	void ToggleViewport();
+
+	/** Get the image to show for the toggle viewport option */
+	FSlateIcon GetViewportImage() const;
+
+	/** Function used to generate the list of differences and the widgets needed to calculate that list */
+	void GenerateDifferencesList();
+
+	/** Called when editor may need to be closed */
+	void OnCloseAssetEditor(UObject* Asset, EAssetEditorCloseReason CloseReason);
+
+	struct FMaterialInstanceDiffControl
+	{
+		FMaterialInstanceDiffControl()
+			: Widget()
+			, DiffControl(nullptr)
+		{
+		}
+
+		TSharedPtr<SWidget> Widget;
+		TSharedPtr<IDiffControl> DiffControl;
+	};
+
+	FMaterialInstanceDiffControl GenerateMaterialInstancePanel();
+
+	/** Accessor and event handler for toggling between diff view modes - only details mode for now: */
+	void SetCurrentMode(FName NewMode);
+
+	FName GetCurrentMode() const { return CurrentMode; }
+
+	void OnModeChanged(const FName& InNewViewMode) const;
+
+	FName CurrentMode;
+
+	FMaterialInstanceDiffPanel PanelOld, PanelNew;
+
+	/** If the preview Viewports should be shown */
+	bool bShowViewport = true;
+
+	/** Contents widget that we swap when mode changes (defaults, components, etc) */
+	TSharedPtr<SBox> ModeContents;
+
+	TSharedPtr<SSplitter> TopRevisionInfoWidget;
+
+	TSharedPtr<SDetailsSplitter> DiffDetailSplitter;
+
+	/** Tree of differences collected across all panels: */
+	TArray<TSharedPtr<FBlueprintDifferenceTreeEntry>> PrimaryDifferencesList;
+
+	/** List of all differences, cached so that we can iterate only the differences and not labels, etc: */
+	TArray<TSharedPtr<FBlueprintDifferenceTreeEntry>> RealDifferences;
+
+	/** Tree view that displays the differences, cached for the buttons that iterate the differences: */
+	TSharedPtr<STreeView<TSharedPtr<FBlueprintDifferenceTreeEntry>>> DifferencesTreeView;
+
+	/** Stored references to widgets used to display various parts of a blueprint, from the mode name */
+	TMap<FName, FMaterialInstanceDiffControl> ModePanels;
+
+	/** A pointer to the window holding this */
+	TWeakPtr<SWindow> WeakParentWindow;
+
+	FDelegateHandle AssetEditorCloseDelegate;
+};
diff --git a/Engine/Source/Editor/MaterialEditor/MaterialEditor.Build.cs b/Engine/Source/Editor/MaterialEditor/MaterialEditor.Build.cs
--- a/Engine/Source/Editor/MaterialEditor/MaterialEditor.Build.cs
+++ b/Engine/Source/Editor/MaterialEditor/MaterialEditor.Build.cs
@@ -15,7 +15,6 @@ public class MaterialEditor : ModuleRules
 			new string[] 
 			{
 				"AssetTools",
-				"Kismet",
 				"EditorWidgets",
 				"MessageLog",
 			}
@@ -44,7 +43,9 @@ public class MaterialEditor : ModuleRules
                 "AssetRegistry",
 				"ToolMenus",
 				"MainFrame",
-				"Landscape"
+				"Landscape",
+				"Kismet", // for rich diffing machinery which was mostly implemented in the kismet module
+				"SourceControl"
 			}
 		);
 
diff --git a/Engine/Source/Editor/MaterialEditor/Private/DiffTool/MaterialDiffPanel.cpp b/Engine/Source/Editor/MaterialEditor/Private/DiffTool/MaterialDiffPanel.cpp
new file mode 100644
--- /dev/null
+++ b/Engine/Source/Editor/MaterialEditor/Private/DiffTool/MaterialDiffPanel.cpp
@@ -0,0 +1,586 @@
+// Copyright Epic Games, Inc. All Rights Reserved.
+
+#include "DiffTool/MaterialDiffPanel.h"
+
+#include "EdGraphUtilities.h"
+#include "Framework/Commands/GenericCommands.h"
+#include "GraphDiffControl.h"
+#include "HAL/PlatformApplicationMisc.h"
+#include "IDetailsView.h"
+#include "MaterialDomain.h"
+#include "MaterialEditorActions.h"
+#include "MaterialShared.h"
+#include "MaterialEditor/PreviewMaterial.h"
+#include "MaterialGraph/MaterialGraphNode.h"
+#include "MaterialGraph/MaterialGraphNode_Comment.h"
+#include "Materials/MaterialExpression.h"
+#include "Materials/MaterialExpressionBreakMaterialAttributes.h"
+#include "Materials/MaterialExpressionFunctionOutput.h"
+#include "Materials/MaterialExpressionStaticBool.h"
+#include "Materials/MaterialExpressionStaticBoolParameter.h"
+#include "SMaterialEditorViewport.h"
+#include "SMyBlueprint.h"
+#include "Widgets/Layout/SBox.h"
+
+#define LOCTEXT_NAMESPACE "MaterialDiffPanel"
+
+DEFINE_LOG_CATEGORY_STATIC(LogMaterialEditorDiff, Log, All);
+
+void FMaterialDiffPanel::GeneratePanel(UEdGraph* NewGraph, UEdGraph* OldGraph)
+{
+	const TSharedPtr<TArray<FDiffSingleResult>> Diff = MakeShared<TArray<FDiffSingleResult>>();
+	FGraphDiffControl::DiffGraphs(OldGraph, NewGraph, *Diff);
+	GeneratePanel(NewGraph, Diff, TAttribute<int32>{});
+}
+
+void FMaterialDiffPanel::GeneratePanel(UEdGraph* Graph, TSharedPtr<TArray<FDiffSingleResult>> DiffResults, TAttribute<int32> FocusedDiffResult)
+{
+	if (GraphEditor.IsValid() && GraphEditor.Pin()->GetCurrentGraph() == Graph)
+	{
+		return;
+	}
+
+	// clang-format off
+	TSharedPtr<SWidget> Widget = SNew(SBorder)
+									 .HAlign(HAlign_Center)
+									 .VAlign(VAlign_Center)
+									 [
+										SNew(STextBlock)
+										 .Text(LOCTEXT("MaterialGraphDiffPanelNoGraphTip", "Graph does not exist in this revision"))
+									 ];
+	// clang-format on
+
+	if (Graph)
+	{
+		SGraphEditor::FGraphEditorEvents InEvents;
+		{
+			const auto SelectionChangedHandler = [this](const FGraphPanelSelectionSet& SelectionSet) {
+				if (MaterialNodeDetailsView)
+				{
+					if (SelectionSet.Array().Num() == 1)
+					{
+						if (UMaterialGraphNode_Base* MaterialNode = Cast<UMaterialGraphNode_Base>(SelectionSet.Array()[0]))
+						{
+							MaterialNodeDetailsView->SetObject(MaterialNode->GetMaterialNodeOwner());
+							MaterialNodeDetailsView->HighlightProperty(PropertyToFocus);
+							MaterialNodeDetailsView->SetVisibility(EVisibility::Visible);
+							MaterialNodeDetailsView->SetIsPropertyEditingEnabledDelegate(
+								FIsPropertyEditingEnabled::CreateLambda([]() { return false; }));
+						}
+						else if (UMaterialGraphNode_Comment* CommentNode = Cast<UMaterialGraphNode_Comment>(SelectionSet.Array()[0]))
+						{
+							MaterialNodeDetailsView->SetObject(CommentNode);
+							MaterialNodeDetailsView->SetVisibility(EVisibility::Visible);
+						}
+						else
+						{
+							MaterialNodeDetailsView->SetVisibility(EVisibility::Hidden);
+						}
+					}
+					else
+					{
+						MaterialNodeDetailsView->SetVisibility(EVisibility::Hidden);
+					}
+				}
+			};
+
+			const auto ContextMenuHandler = [this](UEdGraph* CurrentGraph, const UEdGraphNode* InGraphNode, const UEdGraphPin* InGraphPin, FMenuBuilder* MenuBuilder, bool bIsDebugging) {
+				if (MenuBuilder)
+				{
+					MenuBuilder->AddMenuEntry(FGenericCommands::Get().Copy);
+					MenuBuilder->AddMenuEntry(FMaterialEditorCommands::Get().SelectDownstreamNodes);
+					MenuBuilder->AddMenuEntry(FMaterialEditorCommands::Get().SelectUpstreamNodes);
+					if (const UMaterialGraphNode* MaterialGraphNode = Cast<UMaterialGraphNode>(InGraphNode))
+					{
+						if (UMaterialExpression* MaterialExpression = MaterialGraphNode->MaterialExpression)
+						{
+							// Don't show preview option for bools
+							if (!MaterialExpression->IsA(UMaterialExpressionStaticBool::StaticClass()) && !MaterialExpression->IsA(UMaterialExpressionStaticBoolParameter::StaticClass()))
+							{
+								// Add a preview node option if only one node is selected
+								if (MaterialGraphNode->MaterialExpression == PreviewExpression)
+								{
+									// If we are already previewing the selected node, the menu option should tell the user that this will stop previewing
+									MenuBuilder->AddMenuEntry(FMaterialEditorCommands::Get().StopPreviewNode);
+								}
+								else
+								{
+									// The menu option should tell the user this node will be previewed.
+									MenuBuilder->AddMenuEntry(FMaterialEditorCommands::Get().StartPreviewNode);
+								}
+							}
+						}
+					}
+					
+					return FActionMenuContent(MenuBuilder->MakeWidget());
+				}
+
+				return FActionMenuContent();
+			};
+
+			InEvents.OnSelectionChanged = SGraphEditor::FOnSelectionChanged::CreateLambda(SelectionChangedHandler);
+			InEvents.OnCreateNodeOrPinMenu = SGraphEditor::FOnCreateNodeOrPinMenu::CreateLambda(ContextMenuHandler);
+		}
+
+		if (!GraphEditorCommands.IsValid())
+		{
+			GraphEditorCommands = MakeShared<FUICommandList>();
+
+			GraphEditorCommands->MapAction(FGenericCommands::Get().Copy,
+										   FExecuteAction::CreateRaw(this, &FMaterialDiffPanel::CopySelectedNodes),
+										   FCanExecuteAction::CreateRaw(this, &FMaterialDiffPanel::CanCopyNodes));
+			GraphEditorCommands->MapAction(FMaterialEditorCommands::Get().SelectDownstreamNodes,
+										   FExecuteAction::CreateRaw(this, &FMaterialDiffPanel::SelectDownstreamNodes));
+			GraphEditorCommands->MapAction(FMaterialEditorCommands::Get().SelectUpstreamNodes,
+										   FExecuteAction::CreateRaw(this, &FMaterialDiffPanel::SelectUpstreamNodes));
+			GraphEditorCommands->MapAction(FMaterialEditorCommands::Get().StartPreviewNode,
+										   FExecuteAction::CreateRaw(this, &FMaterialDiffPanel::OnPreviewNode));
+			GraphEditorCommands->MapAction(FMaterialEditorCommands::Get().StopPreviewNode,
+										   FExecuteAction::CreateRaw(this, &FMaterialDiffPanel::OnPreviewNode));
+		}
+
+		TSharedRef<SGraphEditor> Editor = SNew(SGraphEditor)
+											  .AdditionalCommands(GraphEditorCommands)
+											  .GraphToEdit(Graph)
+											  .GraphToDiff(nullptr)
+											  .DiffResults(DiffResults)
+											  .FocusedDiffResult(FocusedDiffResult)
+											  .IsEditable(false)
+											  .GraphEvents(InEvents);
+
+		GraphEditor = Editor;
+		Widget = Editor;
+	}
+
+	GraphEditorBox->SetContent(Widget.ToSharedRef());
+}
+
+FGraphPanelSelectionSet FMaterialDiffPanel::GetSelectedNodes() const
+{
+	FGraphPanelSelectionSet CurrentSelection{};
+	TSharedPtr<SGraphEditor> FocusedGraphEd = GraphEditor.Pin();
+
+	if (FocusedGraphEd.IsValid())
+	{
+		CurrentSelection = FocusedGraphEd->GetSelectedNodes();
+	}
+
+	return CurrentSelection;
+}
+
+void FMaterialDiffPanel::CopySelectedNodes()
+{
+	// Export the selected nodes and place the text on the clipboard
+	const FGraphPanelSelectionSet SelectedNodes = GetSelectedNodes();
+
+	FString ExportedText;
+	FEdGraphUtilities::ExportNodesToText(SelectedNodes, ExportedText);
+	FPlatformApplicationMisc::ClipboardCopy(*ExportedText);
+}
+
+bool FMaterialDiffPanel::CanCopyNodes() const
+{
+	// If any of the nodes can be duplicated then we should allow copying
+	const FGraphPanelSelectionSet SelectedNodes = GetSelectedNodes();
+	for (FGraphPanelSelectionSet::TConstIterator SelectedIter(SelectedNodes); SelectedIter; ++SelectedIter)
+	{
+		UEdGraphNode* Node = Cast<UEdGraphNode>(*SelectedIter);
+		if (Node && Node->CanDuplicateNode())
+		{
+			return true;
+		}
+	}
+
+	return false;
+}
+
+void FMaterialDiffPanel::SelectDownstreamNodes() const
+{
+	TArray<UMaterialGraphNode*> NodesToCheck;
+	TArray<UMaterialGraphNode*> CheckedNodes;
+	TArray<UMaterialGraphNode*> NodesToSelect;
+
+	const FGraphPanelSelectionSet SelectedNodes = GetSelectedNodes();
+
+	for (FGraphPanelSelectionSet::TConstIterator NodeIt(SelectedNodes); NodeIt; ++NodeIt)
+	{
+		UMaterialGraphNode* GraphNode = Cast<UMaterialGraphNode>(*NodeIt);
+		if (GraphNode)
+		{
+			NodesToCheck.Add(GraphNode);
+		}
+	}
+
+	int32 WatchDogWhile = 0;
+	static constexpr int32 LimitIteration = 100;
+	while (WatchDogWhile++ < LimitIteration && NodesToCheck.Num() > 0)
+	{
+		UMaterialGraphNode* CurrentNode = NodesToCheck.Last();
+		for (UEdGraphPin* Pin : CurrentNode->Pins)
+		{
+			if (Pin->Direction == EGPD_Output)
+			{
+				for (int32 LinkIndex = 0; LinkIndex < Pin->LinkedTo.Num(); ++LinkIndex)
+				{
+					UMaterialGraphNode* LinkedNode = Cast<UMaterialGraphNode>(Pin->LinkedTo[LinkIndex]->GetOwningNode());
+					if (LinkedNode)
+					{
+						int32 FoundIndex = -1;
+						CheckedNodes.Find(LinkedNode, FoundIndex);
+
+						if (FoundIndex < 0)
+						{
+							NodesToSelect.Add(LinkedNode);
+							NodesToCheck.Add(LinkedNode);
+						}
+					}
+				}
+			}
+		}
+
+		// This graph node has now been examined
+		CheckedNodes.Add(CurrentNode);
+		NodesToCheck.Remove(CurrentNode);
+	}
+
+	if (WatchDogWhile >= LimitIteration)
+	{
+		UE_LOG(LogMaterialEditorDiff, Warning, TEXT("FMaterialDiffPanel::SelectDownstreamNodes - Infinite loop encounter"));
+		return;
+	}
+
+	if (TSharedPtr<SGraphEditor> FocusedGraphEd = GraphEditor.Pin())
+	{
+		for (int32 Index = 0; Index < NodesToSelect.Num(); ++Index)
+		{
+			FocusedGraphEd->SetNodeSelection(NodesToSelect[Index], true);
+		}
+	}
+}
+
+void FMaterialDiffPanel::SelectUpstreamNodes() const
+{
+	TArray<UMaterialGraphNode*> NodesToCheck;
+	TArray<UMaterialGraphNode*> CheckedNodes;
+	TArray<UMaterialGraphNode*> NodesToSelect;
+
+	const FGraphPanelSelectionSet SelectedNodes = GetSelectedNodes();
+
+	for (FGraphPanelSelectionSet::TConstIterator NodeIt(SelectedNodes); NodeIt; ++NodeIt)
+	{
+		UMaterialGraphNode* GraphNode = Cast<UMaterialGraphNode>(*NodeIt);
+		if (GraphNode)
+		{
+			NodesToCheck.Add(GraphNode);
+		}
+	}
+
+	int32 WatchDogWhile = 0;
+	static constexpr int32 LimitIteration = 100;
+	while (WatchDogWhile++ < LimitIteration && NodesToCheck.Num() > 0)
+	{
+		UMaterialGraphNode* CurrentNode = NodesToCheck.Last();
+		for (UEdGraphPin* Pin : CurrentNode->Pins)
+		{
+			if (Pin->Direction == EGPD_Input)
+			{
+				for (int32 LinkIndex = 0; LinkIndex < Pin->LinkedTo.Num(); ++LinkIndex)
+				{
+					UMaterialGraphNode* LinkedNode = Cast<UMaterialGraphNode>(Pin->LinkedTo[LinkIndex]->GetOwningNode());
+					if (LinkedNode)
+					{
+						int32 FoundIndex = -1;
+						CheckedNodes.Find(LinkedNode, FoundIndex);
+
+						if (FoundIndex < 0)
+						{
+							NodesToSelect.Add(LinkedNode);
+							NodesToCheck.Add(LinkedNode);
+						}
+					}
+				}
+			}
+		}
+
+		// This graph node has now been examined
+		CheckedNodes.Add(CurrentNode);
+		NodesToCheck.Remove(CurrentNode);
+	}
+
+	if (WatchDogWhile >= LimitIteration)
+	{
+		UE_LOG(LogMaterialEditorDiff, Warning, TEXT("FMaterialDiffPanel::SelectUpstreamNodes - Infinite loop encounter"));
+		return;
+	}
+
+	if (TSharedPtr<SGraphEditor> FocusedGraphEd = GraphEditor.Pin())
+	{
+		for (int32 Index = 0; Index < NodesToSelect.Num(); ++Index)
+		{
+			FocusedGraphEd->SetNodeSelection(NodesToSelect[Index], true);
+		}
+	}
+}
+
+void FMaterialDiffPanel::OnPreviewNode()
+{
+	const FGraphPanelSelectionSet SelectedNodes = GetSelectedNodes();
+	if (SelectedNodes.Num() == 1)
+	{
+		for (FGraphPanelSelectionSet::TConstIterator NodeIt(SelectedNodes); NodeIt; ++NodeIt)
+		{
+			UMaterialGraphNode* GraphNode = Cast<UMaterialGraphNode>(*NodeIt);
+			if (GraphNode)
+			{
+				if (TSharedPtr<SGraphEditor> FocusedGraphEd = GraphEditor.Pin())
+				{
+					FocusedGraphEd->NotifyGraphChanged();
+				}
+
+				SetPreviewExpression(GraphNode->MaterialExpression);
+			}
+		}
+	}
+}
+
+void FMaterialDiffPanel::FocusDiff(UEdGraphPin& Pin)
+{
+	if (GraphEditor.IsValid())
+	{
+		GraphEditor.Pin()->JumpToPin(&Pin);
+	}
+}
+
+void FMaterialDiffPanel::FocusDiff(UEdGraphNode& Node)
+{
+	if (GraphEditor.IsValid())
+	{
+		GraphEditor.Pin()->JumpToNode(&Node, false);
+	}
+}
+
+TSharedRef<SWidget> FMaterialDiffPanel::GetMaterialNodeDetailsViewWidget() const
+{
+	// If this panel is displaying a null object, return an empty boarder instead of a IDetailsView
+	return MaterialNodeDetailsView ? MaterialNodeDetailsView.ToSharedRef() : TSharedRef<SWidget>(SNew(SBorder));
+}
+
+void FMaterialDiffPanel::SetViewportToDisplay()
+{
+	if (!MaterialGraph)
+	{
+		return;
+	}
+
+	if (MaterialGraph->Material && MaterialGraph->Material->IsUIMaterial())
+	{
+		SAssignNew(Preview2DViewport, SMaterialEditorUIPreviewViewport, MaterialGraph->Material);
+	}
+
+	SAssignNew(Preview3DViewport, SMaterialEditor3DPreviewViewport)
+		.PreviewMaterial(MaterialGraph->Material);
+
+	// See FMaterialEditor::InitMaterialEditor
+	if (MaterialGraph->MaterialFunction && MaterialGraph->Material)
+	{
+		bool bSetPreviewExpression = false;
+		UMaterialExpressionFunctionOutput* FirstOutput = nullptr;
+		for (int32 ExpressionIndex = MaterialGraph->Material->GetExpressions().Num() - 1; ExpressionIndex >= 0; ExpressionIndex--)
+		{
+			UMaterialExpression* Expression = MaterialGraph->Material->GetExpressions()[ExpressionIndex];
+
+			// Setup the expression to be used with the preview material instead of the function
+			Expression->Function = nullptr;
+			Expression->Material = MaterialGraph->Material;
+
+			UMaterialExpressionFunctionOutput* FunctionOutput = Cast<UMaterialExpressionFunctionOutput>(Expression);
+			if (FunctionOutput)
+			{
+				FirstOutput = FunctionOutput;
+				if (FunctionOutput->bLastPreviewed)
+				{
+					bSetPreviewExpression = true;
+
+					// Preview the last output previewed
+					SetPreviewExpression(FunctionOutput);
+				}
+			}
+		}
+
+		if (!bSetPreviewExpression && FirstOutput)
+		{
+			SetPreviewExpression(FirstOutput);
+		}
+	}
+}
+
+TSharedRef<SWidget> FMaterialDiffPanel::GetViewportToDisplay() const
+{
+	if (Preview2DViewport)
+	{
+		return Preview2DViewport.ToSharedRef();
+	}
+
+	return Preview3DViewport ? Preview3DViewport.ToSharedRef() : TSharedRef<SWidget>(SNew(SBorder));
+}
+
+void FMaterialDiffPanel::SetViewportVisibility(bool bShowViewport)
+{
+	if (!MaterialGraph)
+	{
+		return;
+	}
+
+	if (Preview2DViewport)
+	{
+		Preview2DViewport->SetVisibility(bShowViewport ? EVisibility::Visible : EVisibility::Collapsed);
+	}
+
+	if (Preview3DViewport)
+	{
+		Preview3DViewport->SetVisibility(bShowViewport ? EVisibility::Visible : EVisibility::Collapsed);
+	}
+}
+
+void FMaterialDiffPanel::SetPreviewExpression(UMaterialExpression* NewPreviewExpression)
+{
+	if (!MaterialGraph)
+	{
+		return;
+	}
+
+	UMaterialExpressionFunctionOutput* FunctionOutput = Cast<UMaterialExpressionFunctionOutput>(NewPreviewExpression);
+
+	const UMaterial* Material = MaterialGraph->Material;
+
+	if (!NewPreviewExpression || PreviewExpression == NewPreviewExpression)
+	{
+		if (FunctionOutput)
+		{
+			FunctionOutput->bLastPreviewed = false;
+		}
+		// If we are already previewing the selected expression toggle previewing off
+		PreviewExpression = nullptr;
+		if(ExpressionPreviewMaterial)
+		{
+			ExpressionPreviewMaterial->GetExpressionCollection().Empty();
+		}
+		SetPreviewMaterial(MaterialGraph->Material);
+		// Recompile the preview material to get changes that might have been made during previewing
+		UpdatePreviewMaterial();
+	}
+	else
+	{
+		if (ExpressionPreviewMaterial == nullptr)
+		{
+			// Create the expression preview material if it hasnt already been created
+			ExpressionPreviewMaterial = NewObject<UPreviewMaterial>(GetTransientPackage(), NAME_None, RF_Public);
+			ExpressionPreviewMaterial->bIsPreviewMaterial = true;
+			ExpressionPreviewMaterial->bEnableNewHLSLGenerator = Material->IsUsingNewHLSLGenerator();
+			// MODAVI_TODO
+			//ExpressionPreviewMaterial->bEnableExecWire = Material->IsUsingControlFlow();
+			if (Material->IsUIMaterial())
+			{
+				ExpressionPreviewMaterial->MaterialDomain = MD_UI;
+			}
+			else if (Material->IsPostProcessMaterial())
+			{
+				ExpressionPreviewMaterial->MaterialDomain = MD_PostProcess;
+			}
+		}
+
+		if (FunctionOutput)
+		{
+			FunctionOutput->bLastPreviewed = true;
+		}
+		else
+		{
+			// Hooking up the output of the break expression doesn't make much sense, preview the expression feeding it instead.
+			UMaterialExpressionBreakMaterialAttributes* BreakExpr = Cast<UMaterialExpressionBreakMaterialAttributes>(NewPreviewExpression);
+			if (BreakExpr && BreakExpr->GetInput(0) && BreakExpr->GetInput(0)->Expression)
+			{
+				NewPreviewExpression = BreakExpr->GetInput(0)->Expression;
+			}
+		}
+
+		// The expression preview material's expressions array must stay up to date before recompiling
+		// So that RebuildMaterialFunctionInfo will see all the nested material functions that may need to be updated
+		ExpressionPreviewMaterial->AssignExpressionCollection(Material->GetExpressionCollection());
+
+		// The preview window should now show the expression preview material
+		SetPreviewMaterial(ExpressionPreviewMaterial);
+
+		// Set the preview expression
+		PreviewExpression = NewPreviewExpression;
+
+		// Recompile the preview material
+		UpdatePreviewMaterial();
+	}
+}
+
+
+void FMaterialDiffPanel::SetPreviewMaterial(UMaterialInterface* InMaterialInterface) const
+{
+	if (Preview2DViewport)
+	{
+		Preview2DViewport->SetPreviewMaterial(InMaterialInterface);
+	}
+	
+	if (Preview3DViewport)
+	{
+		Preview3DViewport->SetPreviewMaterial(InMaterialInterface);
+	}
+}
+
+void FMaterialDiffPanel::UpdatePreviewMaterial() const
+{
+	if (!MaterialGraph)
+	{
+		return;
+	}
+
+	if (PreviewExpression && ExpressionPreviewMaterial)
+	{
+		ExpressionPreviewMaterial->UpdateCachedExpressionData();
+		PreviewExpression->ConnectToPreviewMaterial(ExpressionPreviewMaterial, 0);
+	}
+
+	const UMaterial* Material = MaterialGraph->Material;
+
+	if (PreviewExpression)
+	{
+		check(ExpressionPreviewMaterial);
+
+		// The preview material's expressions array must stay up to date before recompiling
+		// So that RebuildMaterialFunctionInfo will see all the nested material functions that may need to be updated
+		ExpressionPreviewMaterial->AssignExpressionCollection(Material->GetExpressionCollection());
+		ExpressionPreviewMaterial->bEnableNewHLSLGenerator = Material->IsUsingNewHLSLGenerator();
+
+		FMaterialUpdateContext UpdateContext(FMaterialUpdateContext::EOptions::SyncWithRenderingThread);
+		UpdateContext.AddMaterial(ExpressionPreviewMaterial);
+
+		// If we are previewing an expression, update the expression preview material
+		ExpressionPreviewMaterial->PreEditChange(nullptr);
+		ExpressionPreviewMaterial->PostEditChange();
+	}
+
+	// Reregister all components that use the preview material, since UMaterial::PEC does not reregister components using a bIsPreviewMaterial=true material
+	if (Preview3DViewport)
+	{
+		Preview3DViewport->RefreshViewport();
+	}
+}
+
+void FMaterialDiffPanel::AddReferencedObjects(FReferenceCollector& Collector)
+{
+	Collector.AddReferencedObject(ExpressionPreviewMaterial);
+	Collector.AddReferencedObject(PreviewExpression);
+	Collector.AddReferencedObject(MaterialGraph);
+}
+
+FString FMaterialDiffPanel::GetReferencerName() const
+{
+	return TEXT("FMaterialDiffPanel");
+}
+
+#undef LOCTEXT_NAMESPACE
diff --git a/Engine/Source/Editor/MaterialEditor/Private/DiffTool/MaterialInstanceDiffPanel.cpp b/Engine/Source/Editor/MaterialEditor/Private/DiffTool/MaterialInstanceDiffPanel.cpp
new file mode 100644
--- /dev/null
+++ b/Engine/Source/Editor/MaterialEditor/Private/DiffTool/MaterialInstanceDiffPanel.cpp
@@ -0,0 +1,177 @@
+// Copyright Epic Games, Inc. All Rights Reserved.
+
+#include "DiffTool/MaterialInstanceDiffPanel.h"
+
+#include "HAL/PlatformApplicationMisc.h"
+#include "MaterialDomain.h"
+#include "MaterialEditor/PreviewMaterial.h"
+#include "MaterialEditorActions.h"
+#include "MaterialShared.h"
+#include "Materials/MaterialExpression.h"
+#include "Materials/MaterialExpressionBreakMaterialAttributes.h"
+#include "Materials/MaterialExpressionFunctionOutput.h"
+#include "Materials/MaterialInstance.h"
+#include "SMaterialEditorViewport.h"
+#include "Widgets/Layout/SBorder.h"
+
+#define LOCTEXT_NAMESPACE "MaterialInstanceDiffPanel"
+
+void FMaterialInstanceDiffPanel::SetViewportToDisplay()
+{
+	if (!MaterialInstance)
+	{
+		return;
+	}
+
+	if (MaterialInstance && MaterialInstance->IsUIMaterial())
+	{
+		SAssignNew(Preview2DViewport, SMaterialEditorUIPreviewViewport, MaterialInstance);
+	}
+
+	SAssignNew(Preview3DViewport, SMaterialEditor3DPreviewViewport)
+		.PreviewMaterial(MaterialInstance);
+}
+
+TSharedRef<SWidget> FMaterialInstanceDiffPanel::GetViewportToDisplay() const
+{
+	if (Preview2DViewport)
+	{
+		return Preview2DViewport.ToSharedRef();
+	}
+
+	return Preview3DViewport ? Preview3DViewport.ToSharedRef() : TSharedRef<SWidget>(SNew(SBorder));
+}
+
+void FMaterialInstanceDiffPanel::SetPreviewExpression(UMaterialExpression* NewPreviewExpression)
+{
+	UMaterialExpressionFunctionOutput* FunctionOutput = Cast<UMaterialExpressionFunctionOutput>(NewPreviewExpression);
+
+	check(MaterialInstance);
+
+	if (!NewPreviewExpression || PreviewExpression == NewPreviewExpression)
+	{
+		if (FunctionOutput)
+		{
+			FunctionOutput->bLastPreviewed = false;
+		}
+		// If we are already previewing the selected expression toggle previewing off
+		PreviewExpression = nullptr;
+		if(ExpressionPreviewMaterial)
+		{
+			ExpressionPreviewMaterial->GetExpressionCollection().Empty();
+		}
+		SetPreviewMaterial(MaterialInstance);
+		// Recompile the preview material to get changes that might have been made during previewing
+		UpdatePreviewMaterial();
+	}
+	else
+	{
+		if (ExpressionPreviewMaterial == nullptr)
+		{
+			// Create the expression preview material if it hasnt already been created
+			ExpressionPreviewMaterial = NewObject<UPreviewMaterial>(GetTransientPackage(), NAME_None, RF_Public);
+			ExpressionPreviewMaterial->bIsPreviewMaterial = true;
+			ExpressionPreviewMaterial->bEnableNewHLSLGenerator = MaterialInstance->IsUsingNewHLSLGenerator();
+			if (MaterialInstance->IsUIMaterial())
+			{
+				ExpressionPreviewMaterial->MaterialDomain = MD_UI;
+			}
+			else if (MaterialInstance->IsPostProcessMaterial())
+			{
+				ExpressionPreviewMaterial->MaterialDomain = MD_PostProcess;
+			}
+		}
+
+		if (FunctionOutput)
+		{
+			FunctionOutput->bLastPreviewed = true;
+		}
+		else
+		{
+			// Hooking up the output of the break expression doesn't make much sense, preview the expression feeding it instead.
+			UMaterialExpressionBreakMaterialAttributes* BreakExpr = Cast<UMaterialExpressionBreakMaterialAttributes>(NewPreviewExpression);
+			if (BreakExpr && BreakExpr->GetInput(0) && BreakExpr->GetInput(0)->Expression)
+			{
+				NewPreviewExpression = BreakExpr->GetInput(0)->Expression;
+			}
+		}
+
+		// The expression preview material's expressions array must stay up to date before recompiling
+		// So that RebuildMaterialFunctionInfo will see all the nested material functions that may need to be updated
+		const UMaterial* Material = MaterialInstance->GetMaterial();
+		const FMaterialExpressionCollection& ExpressionCollection = Material ? Material->GetExpressionCollection() : FMaterialExpressionCollection{};
+		ExpressionPreviewMaterial->AssignExpressionCollection(ExpressionCollection);
+
+		// The preview window should now show the expression preview material
+		SetPreviewMaterial(ExpressionPreviewMaterial);
+
+		// Set the preview expression
+		PreviewExpression = NewPreviewExpression;
+
+		// Recompile the preview material
+		UpdatePreviewMaterial();
+	}
+}
+
+
+void FMaterialInstanceDiffPanel::SetPreviewMaterial(UMaterialInterface* InMaterialInterface) const
+{
+	if (Preview2DViewport)
+	{
+		Preview2DViewport->SetPreviewMaterial(InMaterialInterface);
+	}
+
+	if(Preview3DViewport)
+	{
+		Preview3DViewport->SetPreviewMaterial(InMaterialInterface);
+	}
+}
+
+void FMaterialInstanceDiffPanel::UpdatePreviewMaterial() const
+{
+	if (PreviewExpression && ExpressionPreviewMaterial)
+	{
+		ExpressionPreviewMaterial->UpdateCachedExpressionData();
+		PreviewExpression->ConnectToPreviewMaterial(ExpressionPreviewMaterial, 0);
+	}
+
+	if (PreviewExpression)
+	{
+		check(ExpressionPreviewMaterial);
+
+		// The preview material's expressions array must stay up to date before recompiling
+		// So that RebuildMaterialFunctionInfo will see all the nested material functions that may need to be updated
+		const UMaterial* Material = MaterialInstance->GetMaterial();
+		const FMaterialExpressionCollection& ExpressionCollection = Material ? Material->GetExpressionCollection() : FMaterialExpressionCollection{};
+		ExpressionPreviewMaterial->AssignExpressionCollection(ExpressionCollection);
+		ExpressionPreviewMaterial->bEnableNewHLSLGenerator = MaterialInstance->IsUsingNewHLSLGenerator();
+
+		FMaterialUpdateContext UpdateContext(FMaterialUpdateContext::EOptions::SyncWithRenderingThread);
+		UpdateContext.AddMaterial(ExpressionPreviewMaterial);
+
+		// If we are previewing an expression, update the expression preview material
+		ExpressionPreviewMaterial->PreEditChange(nullptr);
+		ExpressionPreviewMaterial->PostEditChange();
+	}
+
+	// Reregister all components that use the preview material, since UMaterial::PEC does not reregister components using a bIsPreviewMaterial=true material
+	if(Preview3DViewport)
+	{
+		Preview3DViewport->RefreshViewport();
+	}
+}
+
+void FMaterialInstanceDiffPanel::AddReferencedObjects(FReferenceCollector& Collector)
+{
+	Collector.AddReferencedObject(ExpressionPreviewMaterial);
+	Collector.AddReferencedObject(PreviewExpression);
+	Collector.AddReferencedObject(MaterialInstance);
+}
+
+FString FMaterialInstanceDiffPanel::GetReferencerName() const
+{
+	return TEXT("FMaterialInstanceDiffPanel");
+}
+
+
+#undef LOCTEXT_NAMESPACE
diff --git a/Engine/Source/Editor/MaterialEditor/Private/DiffTool/MaterialToDiff.cpp b/Engine/Source/Editor/MaterialEditor/Private/DiffTool/MaterialToDiff.cpp
new file mode 100644
--- /dev/null
+++ b/Engine/Source/Editor/MaterialEditor/Private/DiffTool/MaterialToDiff.cpp
@@ -0,0 +1,191 @@
+// Copyright Epic Games, Inc. All Rights Reserved.
+
+#include "DiffTool/MaterialToDiff.h"
+
+#include "EdGraph/EdGraph.h"
+#include "EdGraph/EdGraphSchema.h"
+#include "GraphDiffControl.h"
+#include "MaterialGraph/MaterialGraph.h"
+#include "SBlueprintDiff.h"
+#include "Widgets/SCompoundWidget.h"
+#include "DiffTool/Widgets/SMaterialDiff.h"
+
+#define LOCTEXT_NAMESPACE "MaterialToDiff"
+
+FMaterialToDiff::FMaterialToDiff(SMaterialDiff* InDiffWidget, UMaterialGraph* InOldGraph, UMaterialGraph* InNewGraph, const FRevisionInfo& InOldRevision, const FRevisionInfo& InNewRevision)
+	: FoundDiffs(MakeShared<TArray<FDiffSingleResult>>())
+	, DiffWidget(InDiffWidget)
+	, OldGraph(InOldGraph)
+	, NewGraph(InNewGraph)
+	, OldRevision(InOldRevision)
+	, NewRevision(InNewRevision)
+{
+	check(InOldGraph || InNewGraph);
+	
+	BuildDiffSourceArray();
+}
+
+FMaterialToDiff::~FMaterialToDiff()
+{
+}
+
+void FMaterialToDiff::GenerateTreeEntries(TArray<TSharedPtr<FBlueprintDifferenceTreeEntry>>& OutTreeEntries, TArray<TSharedPtr<FBlueprintDifferenceTreeEntry>>& OutRealDifferences)
+{
+	if (!DiffListSource.IsEmpty())
+	{
+		RealDifferencesStartIndex = OutRealDifferences.Num();
+	}
+
+	TArray<TSharedPtr<FBlueprintDifferenceTreeEntry>> Children;
+	for (const TSharedPtr<FMaterialDiffResultItem>& Difference : DiffListSource)
+	{
+		TSharedPtr<FBlueprintDifferenceTreeEntry> ChildEntry = MakeShared<FBlueprintDifferenceTreeEntry>(
+			FOnDiffEntryFocused::CreateSP(DiffWidget, &SMaterialDiff::OnDiffListSelectionChanged, Difference),
+			FGenerateDiffEntryWidget::CreateSP(Difference.ToSharedRef(), &FMaterialDiffResultItem::GenerateWidget));
+		Children.Push(ChildEntry);
+		OutRealDifferences.Push(ChildEntry);
+	}
+
+	if (Children.Num() == 0)
+	{
+		// Make one child informing the user that there are no differences
+		Children.Push(FBlueprintDifferenceTreeEntry::NoDifferencesEntry());
+	}
+
+	TSharedPtr<FBlueprintDifferenceTreeEntry> Entry = MakeShared<FBlueprintDifferenceTreeEntry>(
+		FOnDiffEntryFocused::CreateSP(DiffWidget, &SMaterialDiff::OnGraphSelectionChanged, TSharedPtr<FMaterialToDiff>(AsShared()), ESelectInfo::Direct),
+		FGenerateDiffEntryWidget::CreateSP(AsShared(), &FMaterialToDiff::GenerateCategoryWidget),
+		Children);
+	OutTreeEntries.Push(Entry);
+
+	// Add comments as children to this category
+	FString GraphName;
+	const UMaterialGraph* Graph = OldGraph ? OldGraph : NewGraph;
+	if (const UEdGraphSchema* Schema = Graph->GetSchema())
+	{
+		FGraphDisplayInfo DisplayInfo;
+		Schema->GetGraphDisplayInformation(*Graph, DisplayInfo);
+
+		GraphName = DisplayInfo.DisplayName.ToString();
+	}
+	else
+	{
+		GraphName = Graph->GetFName().ToString();
+	}
+	GenerateCategoryCommentTreeEntries(OutTreeEntries.Last()->Children, OutRealDifferences, GraphName);
+}
+
+UMaterialGraph* FMaterialToDiff::GetOldGraph() const
+{
+	return OldGraph;
+}
+
+UMaterialGraph* FMaterialToDiff::GetNewGraph() const
+{
+	return NewGraph;
+}
+
+FText FMaterialToDiff::GetToolTip()
+{
+	if (OldGraph && NewGraph)
+	{
+		if (DiffListSource.Num() > 0)
+		{
+			return LOCTEXT("ContainsDifferences", "Revisions are different");
+		}
+		else
+		{
+			return LOCTEXT("MaterialGraphsIdentical", "Revisions appear to be identical");
+		}
+	}
+	else
+	{
+		UMaterialGraph* GoodGraph = OldGraph ? OldGraph : NewGraph;
+		check(GoodGraph);
+		const FRevisionInfo& Revision = NewGraph ? OldRevision : NewRevision;
+		FText RevisionText = LOCTEXT("CurrentRevision", "Current Revision");
+
+		if (!Revision.Revision.IsEmpty())
+		{
+			RevisionText = FText::Format(LOCTEXT("Revision Number", "Revision {0}"), FText::FromString(Revision.Revision));
+		}
+
+		return FText::Format(LOCTEXT("MissingGraph", "Graph '{0}' missing from {1}"), FText::FromString(GoodGraph->GetName()), RevisionText);
+	}
+}
+
+TSharedRef<SWidget> FMaterialToDiff::GenerateCategoryWidget()
+{
+	const UMaterialGraph* Graph = OldGraph ? OldGraph : NewGraph;
+	check(Graph);
+
+	FLinearColor Color = (OldGraph && NewGraph) ? DiffViewUtils::Identical() : FLinearColor(0.3f, 0.3f, 1.f);
+
+	const bool bHasDiffs = DiffListSource.Num() > 0;
+
+	if (bHasDiffs)
+	{
+		Color = DiffViewUtils::Differs();
+	}
+
+	FText GraphName;
+	if (const UEdGraphSchema* Schema = Graph->GetSchema())
+	{
+		FGraphDisplayInfo DisplayInfo;
+		Schema->GetGraphDisplayInformation(*Graph, DisplayInfo);
+
+		GraphName = DisplayInfo.DisplayName;
+	}
+	else
+	{
+		GraphName = FText::FromName(Graph->GetFName());
+	}
+
+	return SNew(SHorizontalBox) 
+		   +SHorizontalBox::Slot()
+		   [
+		       SNew(STextBlock)
+		       .ColorAndOpacity(Color)
+		       .Text(GraphName)
+		       .ToolTipText(GetToolTip())
+		    ] 
+		    +DiffViewUtils::Box(OldGraph != nullptr, Color) 
+		    +DiffViewUtils::Box(NewGraph != nullptr, Color);
+}
+
+void FMaterialToDiff::BuildDiffSourceArray()
+{
+	FoundDiffs->Empty();
+	FGraphDiffControl::DiffGraphs(OldGraph, NewGraph, *FoundDiffs);
+
+	Algo::SortBy(*FoundDiffs, &FDiffSingleResult::Diff);
+
+	DiffListSource.Empty();
+	for (const FDiffSingleResult& Diff : *FoundDiffs)
+	{
+		TSharedRef<FMaterialDiffResultItem> DiffResultItem = MakeShared<FMaterialDiffResultItem>(Diff);
+
+		if (Diff.Diff == EDiffType::NODE_PROPERTY && Diff.Object1 && Diff.Object2)
+		{
+			for (TFieldIterator<FProperty> PropertyIt(Diff.Object1->GetClass(), EFieldIteratorFlags::IncludeSuper); PropertyIt; ++PropertyIt)
+			{
+				// We use the DisplayString containing the changed property's name
+				if (PropertyIt->GetName().Equals(Diff.DisplayString.ToString()))
+				{
+					DiffResultItem->Property.AddProperty(FPropertyInfo(*PropertyIt));
+
+					// We found the property, we keep it in memory and we can use the correct format for properties of MaterialGraph
+					DiffResultItem->Result.DisplayString = FText::Format(LOCTEXT("Diff_MaterialGraphNodePropertyFormat", "Property Changed: {0} in {1}"), 
+						                                                 FText::FromString(PropertyIt->GetName()),
+						                                                 FText::FromString(Diff.Object1->GetName()));
+
+					break;
+				}
+			}
+		}
+
+		DiffListSource.Add(DiffResultItem);
+	}
+}
+
+#undef LOCTEXT_NAMESPACE
\ No newline at end of file
diff --git a/Engine/Source/Editor/MaterialEditor/Private/DiffTool/MaterialToDiff.h b/Engine/Source/Editor/MaterialEditor/Private/DiffTool/MaterialToDiff.h
new file mode 100644
--- /dev/null
+++ b/Engine/Source/Editor/MaterialEditor/Private/DiffTool/MaterialToDiff.h
@@ -0,0 +1,60 @@
+// Copyright Epic Games, Inc. All Rights Reserved.
+
+#pragma once
+
+#include "CoreMinimal.h"
+#include "DiffControl.h"
+
+class SMaterialDiff;
+class UMaterialGraph;
+
+struct FMaterialDiffResultItem;
+
+enum class EAssetEditorCloseReason : uint8;
+
+/** Category list item for a MaterialGraph **/
+struct FMaterialToDiff : public TSharedFromThis<FMaterialToDiff>, IDiffControl
+{
+	FMaterialToDiff(SMaterialDiff* DiffWidget, UMaterialGraph* OldGraph, UMaterialGraph* NewGraph, const FRevisionInfo& OldRevision, const FRevisionInfo& NewRevision);
+	virtual ~FMaterialToDiff() override;
+
+	/** Add widgets to the differences tree */
+	virtual void GenerateTreeEntries(TArray<TSharedPtr<FBlueprintDifferenceTreeEntry>>& OutTreeEntries, TArray<TSharedPtr<FBlueprintDifferenceTreeEntry>>& OutRealDifferences) override;
+
+	/** Get old(left) graph*/
+	UMaterialGraph* GetOldGraph() const;
+
+	/** Get new(right) graph*/
+	UMaterialGraph* GetNewGraph() const;
+
+	/** Source for list view */
+	TArray<TSharedPtr<FMaterialDiffResultItem>> DiffListSource;
+	TSharedPtr<TArray<FDiffSingleResult>> FoundDiffs;
+
+	/** Index of the first item in RealDifferences that was generated by this graph */
+	int32 RealDifferencesStartIndex = INDEX_NONE;
+
+private:
+	/** Get tooltip for category */
+	FText GetToolTip();
+
+	/** Generate Widget for category list */
+	TSharedRef<SWidget> GenerateCategoryWidget();
+
+	/** Build up the Diff Source Array*/
+	void BuildDiffSourceArray();
+
+	/** Diff widget */
+	SMaterialDiff* DiffWidget;
+
+	/** The old graph(left)*/
+	UMaterialGraph* OldGraph;
+
+	/** The new graph(right)*/
+	UMaterialGraph* NewGraph;
+
+	/** Description of Old and New graph*/
+	FRevisionInfo OldRevision;
+	FRevisionInfo NewRevision;
+};
+
diff --git a/Engine/Source/Editor/MaterialEditor/Private/DiffTool/Widgets/SMaterialDiff.cpp b/Engine/Source/Editor/MaterialEditor/Private/DiffTool/Widgets/SMaterialDiff.cpp
new file mode 100644
--- /dev/null
+++ b/Engine/Source/Editor/MaterialEditor/Private/DiffTool/Widgets/SMaterialDiff.cpp
@@ -0,0 +1,815 @@
+// Copyright Epic Games, Inc. All Rights Reserved.
+
+#include "DiffTool/Widgets/SMaterialDiff.h"
+
+#include "Editor.h"
+#include "Framework/Application/SlateApplication.h"
+#include "Framework/MultiBox/MultiBoxBuilder.h"
+#include "Framework/MultiBox/MultiBoxDefs.h"
+#include "GraphDiffControl.h"
+#include "IDetailsView.h"
+#include "MaterialGraph/MaterialGraph.h"
+#include "Materials/MaterialFunction.h"
+#include "Modules/ModuleManager.h"
+#include "PropertyEditorModule.h"
+#include "SBlueprintDiff.h"
+#include "SDetailsDiff.h"
+#include "SlateOptMacros.h"
+#include "SMaterialEditorViewport.h"
+#include "SMyBlueprint.h"
+#include "SourceControlHelpers.h"
+#include "Styling/AppStyle.h"
+#include "Subsystems/AssetEditorSubsystem.h"
+#include "DiffTool/MaterialToDiff.h"
+#include "Widgets/Docking/SDockTab.h"
+#include "Widgets/Input/SButton.h"
+#include "Widgets/Layout/SSpacer.h"
+#include "Widgets/Layout/SSplitter.h"
+#include "Widgets/Layout/SWidgetSwitcher.h"
+#include "Widgets/SOverlay.h"
+
+#define LOCTEXT_NAMESPACE "SMaterialGraphDiff"
+
+static const FName GraphMode = FName(TEXT("GraphMode"));
+
+BEGIN_SLATE_FUNCTION_BUILD_OPTIMIZATION
+void SMaterialDiff::Construct(const FArguments& InArgs)
+{
+	check(InArgs._OldMaterialGraph || InArgs._NewMaterialGraph);
+
+	PanelOld.MaterialGraph = InArgs._OldMaterialGraph;
+	PanelNew.MaterialGraph = InArgs._NewMaterialGraph;
+	PanelOld.RevisionInfo = InArgs._OldRevision;
+	PanelNew.RevisionInfo = InArgs._NewRevision;
+
+	PanelOld.bShowAssetName = InArgs._ShowAssetNames;
+	PanelNew.bShowAssetName = InArgs._ShowAssetNames;
+
+	UObject* OldAsset = nullptr;
+	if (PanelOld.MaterialGraph && PanelOld.MaterialGraph->MaterialFunction)
+	{
+		OldAsset = PanelOld.MaterialGraph->MaterialFunction;
+	}
+	else if(PanelOld.MaterialGraph)
+	{
+		OldAsset = PanelOld.MaterialGraph->Material;
+	}
+
+	UObject* NewAsset = nullptr;
+	if (PanelNew.MaterialGraph && PanelNew.MaterialGraph->MaterialFunction)
+	{
+		NewAsset = PanelNew.MaterialGraph->MaterialFunction;
+	}
+	else if(PanelNew.MaterialGraph)
+	{
+		NewAsset = PanelNew.MaterialGraph->Material;
+	}
+
+	if (InArgs._ParentWindow.IsValid())
+	{
+		WeakParentWindow = InArgs._ParentWindow;
+
+		AssetEditorCloseDelegate = GEditor->GetEditorSubsystem<UAssetEditorSubsystem>()->OnAssetEditorRequestClose().AddSP(this, &SMaterialDiff::OnCloseAssetEditor);
+	}
+
+	FToolBarBuilder NavToolBarBuilder(TSharedPtr<const FUICommandList>(), FMultiBoxCustomization::None);
+	NavToolBarBuilder.AddToolBarButton(
+		FUIAction(
+			FExecuteAction::CreateSP(this, &SMaterialDiff::PrevDiff),
+			FCanExecuteAction::CreateSP(this, &SMaterialDiff::HasPrevDiff)),
+		NAME_None,
+		LOCTEXT("PrevDiffLabel", "Prev"),
+		LOCTEXT("PrevDiffTooltip", "Go to previous difference"),
+		FSlateIcon(FAppStyle::GetAppStyleSetName(), "BlueprintDif.PrevDiff")
+	);
+	NavToolBarBuilder.AddToolBarButton(
+		FUIAction(
+			FExecuteAction::CreateSP(this, &SMaterialDiff::NextDiff),
+			FCanExecuteAction::CreateSP(this, &SMaterialDiff::HasNextDiff)),
+		NAME_None,
+		LOCTEXT("NextDiffLabel", "Next"),
+		LOCTEXT("NextDiffTooltip", "Go to next difference"),
+		FSlateIcon(FAppStyle::GetAppStyleSetName(), "BlueprintDif.NextDiff")
+	);
+	NavToolBarBuilder.AddToolBarButton(
+		FUIAction(
+			FExecuteAction::CreateSP(this, &SMaterialDiff::ToggleViewport)),
+		NAME_None,
+		LOCTEXT("ToggleViewportLabel", "Toggle Viewport"),
+		LOCTEXT("ToggleViewportTooltip", "Show/Hide a preview Viewport for the Materials diffed"),
+		TAttribute<FSlateIcon>(this, &SMaterialDiff::GetViewportImage)
+	);
+
+	FToolBarBuilder GraphToolbarBuilder(TSharedPtr< const FUICommandList >(), FMultiBoxCustomization::None);
+	GraphToolbarBuilder.AddToolBarButton(
+		FUIAction(FExecuteAction::CreateSP(this, &SMaterialDiff::OnToggleLockView)), 
+		NAME_None,
+		LOCTEXT("LockMaterialGraphsLabel", "Lock/Unlock"),
+		LOCTEXT("LockMaterialGraphsTooltip", "Force all graph views to change together, or allow independent scrolling/zooming"),
+		TAttribute<FSlateIcon>(this, &SMaterialDiff::GetLockViewImage)
+	);
+	GraphToolbarBuilder.AddToolBarButton(
+		FUIAction(FExecuteAction::CreateSP(this, &SMaterialDiff::OnToggleSplitViewMode)),
+		NAME_None,
+		LOCTEXT("SplitMaterialGraphsModeLabel", "Vertical/Horizontal"),
+		LOCTEXT("SplitMaterialGraphsModeLabelTooltip", "Toggles the split view of Material Graphs between vertical and horizontal"), 
+		TAttribute<FSlateIcon>(this, &SMaterialDiff::GetSplitViewModeImage)
+	);
+
+	DifferencesTreeView = DiffTreeView::CreateTreeView(&PrimaryDifferencesList);
+
+	GenerateDifferencesList();
+
+	const auto TextBlock = [](FText Text) -> TSharedRef<SWidget>
+	{
+		return SNew(SBox)
+		.Padding(FMargin(4.0f, 10.0f))
+		.VAlign(VAlign_Center)
+		.HAlign(HAlign_Left)
+		[
+			SNew(STextBlock)
+			.Visibility(EVisibility::HitTestInvisible)
+			.TextStyle(FAppStyle::Get(), "DetailsView.CategoryTextStyle")
+			.Text(Text)
+		];
+	};
+
+	TopRevisionInfoWidget =
+		SNew(SSplitter)
+		.Visibility(EVisibility::HitTestInvisible)
+		+ SSplitter::Slot()
+		.Value(.2f)
+		[
+			SNew(SBox)
+		]
+		+ SSplitter::Slot()
+		.Value(.8f)
+		[
+			SNew(SSplitter)
+			.PhysicalSplitterHandleSize(10.0f)
+			+ SSplitter::Slot()
+			.Value(.5f)
+			[
+				TextBlock(DiffViewUtils::GetPanelLabel(PanelOld.MaterialGraph, PanelOld.RevisionInfo, FText()))
+			]
+			+ SSplitter::Slot()
+			.Value(.5f)
+			[
+				TextBlock(DiffViewUtils::GetPanelLabel(PanelNew.MaterialGraph, PanelNew.RevisionInfo, FText()))
+			]
+		];
+
+	GraphToolBarWidget = 
+		SNew(SSplitter)
+		.Visibility(EVisibility::HitTestInvisible)
+		+ SSplitter::Slot()
+		.Value(.2f)
+		[
+			SNew(SBox)
+		]
+		+ SSplitter::Slot()
+		.Value(.8f)
+		[
+			SNew(SHorizontalBox)
+			+SHorizontalBox::Slot()
+			.AutoWidth()
+			[
+				GraphToolbarBuilder.MakeWidget()
+			]	
+		];
+
+	// Clang-format off
+	ChildSlot
+	[
+		SNew(SBorder)
+		.BorderImage(FAppStyle::GetBrush("Docking.Tab", ".ContentAreaBrush"))
+		[
+			SNew(SOverlay)
+			+SOverlay::Slot()
+			.VAlign(VAlign_Top)
+			[
+				TopRevisionInfoWidget.ToSharedRef()		
+			]
+			+ SOverlay::Slot()
+			.Padding(0.f, 6.f, 0.f, 4.f)
+			[
+				SNew(SVerticalBox)
+				+SVerticalBox::Slot()
+				.Padding(0.f, 2.f, 0.f, 5.f)
+				.AutoHeight()
+				[
+					SNew(SHorizontalBox)
+					+SHorizontalBox::Slot()
+					[
+						SNew(SButton)
+						.HAlign(HAlign_Center)
+						.OnClicked_Lambda([this]()->FReply { SetCurrentWidgetIndex(0); return FReply::Handled(); })
+						[
+							SNew(STextBlock)
+							.Text(FText::FromString("Graph"))
+						]
+					]
+					+SHorizontalBox::Slot()
+					[
+						SNew(SButton)
+						.HAlign(HAlign_Center)
+						.OnClicked_Lambda([this]()->FReply { SetCurrentWidgetIndex(1); return FReply::Handled(); })
+						[
+							SNew(STextBlock)
+							.Text(FText::FromString("Properties"))
+						]
+					]
+				]
+				+SVerticalBox::Slot()
+				.FillHeight(1.f)
+				[
+					SAssignNew(WidgetSwitcher, SWidgetSwitcher)
+					.WidgetIndex(this, &SMaterialDiff::GetCurrentWidgetIndex)
+					+ SWidgetSwitcher::Slot()
+					[
+						SNew(SOverlay)
+						+SOverlay::Slot()
+						.VAlign(VAlign_Top)
+						.Padding(0.0f, 6.0f, 0.0f, 4.0f)
+						[
+							GraphToolBarWidget.ToSharedRef()
+						]
+						+SOverlay::Slot()
+						[
+							SNew(SVerticalBox)
+							+SVerticalBox::Slot()
+							.AutoHeight()
+							.Padding(0.0f, 2.0f, 0.0f, 2.0f)
+							[
+								SNew(SHorizontalBox)
+								+SHorizontalBox::Slot()
+								.Padding(4.f)
+								.AutoWidth()
+								[
+									NavToolBarBuilder.MakeWidget()
+								]
+								+SHorizontalBox::Slot()
+								[
+									SNew(SSpacer)
+								]
+							]
+							+SVerticalBox::Slot()
+							.FillHeight(1.f)
+							[
+								SNew(SSplitter)
+								+SSplitter::Slot()
+								.Value(.2f)
+								[
+									SNew(SBorder)
+									.BorderImage(FAppStyle::GetBrush("ToolPanel.GroupBorder"))
+									[
+										DifferencesTreeView.ToSharedRef()
+									]
+								]
+								+SSplitter::Slot()
+								.Value(.8f)
+								[
+									SAssignNew(ModeContents, SBox)
+								]
+							]
+						]
+					]
+					+ SWidgetSwitcher::Slot()
+					[
+						SNew(SVerticalBox)
+						+ SVerticalBox::Slot()
+						[
+							SNew(SDetailsDiff)
+							.OldAsset(OldAsset)
+							.NewAsset(NewAsset)
+							.OldRevision(PanelOld.RevisionInfo)
+							.NewRevision(PanelNew.RevisionInfo)
+							.ShowAssetNames(true)
+							.ParentWindow(WeakParentWindow.Pin())
+						]
+					]
+				]
+				
+			]
+		]
+	];
+	// clang-format on
+
+	CurrentWidgetIndex = 0;
+
+	// Display the Material Graphs by default
+	HandleGraphChanged(FString{});
+}
+END_SLATE_FUNCTION_BUILD_OPTIMIZATION
+
+void SMaterialDiff::SetCurrentWidgetIndex(int32 Index)
+{
+	CurrentWidgetIndex = Index;
+}
+
+
+int32 SMaterialDiff::GetCurrentWidgetIndex() const
+{
+	return CurrentWidgetIndex;
+}
+
+SMaterialDiff::~SMaterialDiff()
+{
+	if (AssetEditorCloseDelegate.IsValid())
+	{
+		GEditor->GetEditorSubsystem<UAssetEditorSubsystem>()->OnAssetEditorRequestClose().Remove(AssetEditorCloseDelegate);
+	}
+}
+
+void SMaterialDiff::OnGraphChanged(FMaterialToDiff* Diff)
+{
+	if (PanelNew.GraphEditor.IsValid() && PanelNew.GraphEditor.Pin()->GetCurrentGraph() == Diff->GetNewGraph())
+	{
+		FocusOnMaterialGraphRevisions(Diff);
+	}
+}
+
+void SMaterialDiff::OnGraphSelectionChanged(TSharedPtr<FMaterialToDiff> Item, ESelectInfo::Type SelectionType)
+{
+	if (!Item.IsValid())
+	{
+		return;
+	}
+
+	FocusOnMaterialGraphRevisions(Item.Get());
+}
+
+void SMaterialDiff::OnDiffListSelectionChanged(TSharedPtr<FMaterialDiffResultItem> TheDiff)
+{
+	check(!TheDiff->Result.OwningObjectPath.IsEmpty());
+
+	FocusOnMaterialGraphRevisions(FindGraphToDiffEntry(TheDiff->Result.OwningObjectPath));
+	FDiffSingleResult Result = TheDiff->Result;
+
+	const auto SafeClearSelection = [](TWeakPtr<SGraphEditor> GraphEditor) {
+		TSharedPtr<SGraphEditor> GraphEditorPtr = GraphEditor.Pin();
+		if (GraphEditorPtr.IsValid())
+		{
+			GraphEditorPtr->ClearSelectionSet();
+		}
+	};
+
+	SafeClearSelection(PanelNew.GraphEditor);
+	SafeClearSelection(PanelOld.GraphEditor);
+
+	if (Result.Pin1)
+	{
+		GetDiffPanelForNode(*Result.Pin1->GetOwningNode()).FocusDiff(*Result.Pin1);
+		if (Result.Pin2)
+		{
+			GetDiffPanelForNode(*Result.Pin2->GetOwningNode()).FocusDiff(*Result.Pin2);
+		}
+	}
+	else if (Result.Node1)
+	{
+		GetDiffPanelForNode(*Result.Node1).FocusDiff(*Result.Node1);
+		if (Result.Node2)
+		{
+			GetDiffPanelForNode(*Result.Node2).FocusDiff(*Result.Node2);
+		}
+	}
+
+	PanelOld.PropertyToFocus = TheDiff->Property;
+	PanelNew.PropertyToFocus = TheDiff->Property;
+}
+
+TSharedRef<SWidget> SMaterialDiff::DefaultEmptyPanel()
+{
+	return SNew(SHorizontalBox)
+		   +SHorizontalBox::Slot()
+		   .HAlign(HAlign_Center)
+		   .VAlign(VAlign_Center)
+		   [
+			   SNew(STextBlock)
+			   .Text(LOCTEXT("MaterialGraphDiffGraphsToolTip", "Select Material Graph to Diff"))
+		   ];
+}
+
+TSharedPtr<SWindow> SMaterialDiff::CreateDiffWindow(FText WindowTitle, TObjectPtr<UMaterialGraph> OldMaterialGraph, TObjectPtr<UMaterialGraph> NewMaterialGraph, const FRevisionInfo& OldRevision, const FRevisionInfo& NewRevision)
+{
+	// sometimes we're comparing different revisions of one single asset (other
+	// times we're comparing two completely separate assets altogether)
+	bool bIsSingleAsset = !OldMaterialGraph || !NewMaterialGraph || (OldMaterialGraph->GetName() == NewMaterialGraph->GetName());
+
+	TSharedPtr<SWindow> Window = SNew(SWindow)
+									 .Title(WindowTitle)
+									 .ClientSize(FVector2D(1000, 800));
+
+	Window->SetContent(SNew(SMaterialDiff)
+						   .OldMaterialGraph(OldMaterialGraph)
+						   .NewMaterialGraph(NewMaterialGraph)
+						   .OldRevision(OldRevision)
+						   .NewRevision(NewRevision)
+						   .ShowAssetNames(!bIsSingleAsset)
+						   .ParentWindow(Window));
+
+	// Make this window a child of the modal window if we've been spawned while one is active.
+	const TSharedPtr<SWindow> ActiveModal = FSlateApplication::Get().GetActiveModalWindow();
+	if (ActiveModal.IsValid())
+	{
+		FSlateApplication::Get().AddWindowAsNativeChild(Window.ToSharedRef(), ActiveModal.ToSharedRef());
+	}
+	else
+	{
+		FSlateApplication::Get().AddWindow(Window.ToSharedRef());
+	}
+
+	return Window;
+}
+
+TSharedPtr<SWindow> SMaterialDiff::CreateDiffWindow(TObjectPtr<UMaterialGraph> OldMaterialGraph, TObjectPtr<UMaterialGraph> NewMaterialGraph, const FRevisionInfo& OldRevision, const FRevisionInfo& NewRevision, const UClass* ObjectClass)
+{
+	check(OldMaterialGraph || NewMaterialGraph);
+
+	// sometimes we're comparing different revisions of one single asset (other
+	// times we're comparing two completely separate assets altogether)
+	const bool bIsSingleAsset = !OldMaterialGraph ||
+		!NewMaterialGraph ||
+		(NewMaterialGraph->MaterialFunction && OldMaterialGraph->MaterialFunction && NewMaterialGraph->MaterialFunction.GetName() == OldMaterialGraph->MaterialFunction.GetName()) ||
+		(NewMaterialGraph->Material.GetName() == OldMaterialGraph->Material.GetName());
+
+	FText WindowTitle = FText::Format(LOCTEXT("NamelessMaterialGraphDiff", "{0} Diff"), ObjectClass->GetDisplayNameText());
+	// if we're diffing one asset against itself
+	if (bIsSingleAsset)
+	{
+		// identify the assumed single asset in the window's title
+		FString MaterialGraphName{};
+		if (NewMaterialGraph)
+		{
+			MaterialGraphName = NewMaterialGraph->MaterialFunction ? NewMaterialGraph->MaterialFunction.GetName() : NewMaterialGraph->Material.GetName();
+		}
+		else
+		{
+			MaterialGraphName = OldMaterialGraph->MaterialFunction ? OldMaterialGraph->MaterialFunction.GetName() : OldMaterialGraph->Material.GetName();
+		}
+		WindowTitle = FText::Format(LOCTEXT("NamedMaterialGraphDiff", "{0} - {1} Diff"), FText::FromString(MaterialGraphName), ObjectClass->GetDisplayNameText());
+	}
+
+	return CreateDiffWindow(WindowTitle, OldMaterialGraph, NewMaterialGraph, OldRevision, NewRevision);
+}
+
+void SMaterialDiff::NextDiff()
+{
+	DiffTreeView::HighlightNextDifference(DifferencesTreeView.ToSharedRef(), RealDifferences, PrimaryDifferencesList);
+}
+
+void SMaterialDiff::PrevDiff()
+{
+	DiffTreeView::HighlightPrevDifference(DifferencesTreeView.ToSharedRef(), RealDifferences, PrimaryDifferencesList);
+}
+
+bool SMaterialDiff::HasNextDiff() const
+{
+	return DiffTreeView::HasNextDifference(DifferencesTreeView.ToSharedRef(), RealDifferences);
+}
+
+bool SMaterialDiff::HasPrevDiff() const
+{
+	return DiffTreeView::HasPrevDifference(DifferencesTreeView.ToSharedRef(), RealDifferences);
+}
+
+void SMaterialDiff::ToggleViewport()
+{
+	bShowViewport = !bShowViewport;
+
+	PanelOld.SetViewportVisibility(bShowViewport);
+	PanelNew.SetViewportVisibility(bShowViewport);
+}
+
+FSlateIcon SMaterialDiff::GetViewportImage() const
+{
+	return FSlateIcon(FAppStyle::GetAppStyleSetName(), "LevelEditor.Tabs.Viewports");
+}
+
+FMaterialToDiff* SMaterialDiff::FindGraphToDiffEntry(const FString& GraphPath)
+{
+	if (MaterialGraphToDiff)
+	{
+		FString SearchGraphPath = MaterialGraphToDiff->GetOldGraph() ? FGraphDiffControl::GetGraphPath(MaterialGraphToDiff->GetOldGraph()) : FGraphDiffControl::GetGraphPath(MaterialGraphToDiff->GetNewGraph());
+		if (SearchGraphPath.Equals(GraphPath, ESearchCase::CaseSensitive))
+		{
+			return MaterialGraphToDiff.Get();
+		}
+	}
+
+	return nullptr;
+}
+
+void SMaterialDiff::FocusOnMaterialGraphRevisions(FMaterialToDiff* Diff)
+{
+	if(!Diff)
+	{
+		return;
+	}
+
+	UEdGraph* Graph = Diff->GetOldGraph() ? Diff->GetOldGraph() : Diff->GetNewGraph();
+
+	FString GraphPath = FGraphDiffControl::GetGraphPath(Graph);
+
+	HandleGraphChanged(GraphPath);
+
+	ResetGraphEditors();
+}
+
+void SMaterialDiff::OnToggleLockView()
+{
+	bLockViews = !bLockViews;
+	ResetGraphEditors();
+}
+
+void SMaterialDiff::OnToggleSplitViewMode()
+{
+	bVerticalSplitGraphMode = !bVerticalSplitGraphMode;
+
+	if (SSplitter* DiffGraphSplitterPtr = DiffGraphSplitter.Get())
+	{
+		DiffGraphSplitterPtr->SetOrientation(bVerticalSplitGraphMode ? Orient_Horizontal : Orient_Vertical);
+	}
+}
+
+FSlateIcon SMaterialDiff::GetLockViewImage() const
+{
+	return FSlateIcon(FAppStyle::GetAppStyleSetName(), bLockViews ? "Icons.Lock" : "Icons.Unlock");
+}
+
+FSlateIcon SMaterialDiff::GetSplitViewModeImage() const
+{
+	return FSlateIcon(FAppStyle::GetAppStyleSetName(), bVerticalSplitGraphMode ? "BlueprintDif.VerticalDiff" : "BlueprintDif.HorizontalDiff");
+}
+
+void SMaterialDiff::ResetGraphEditors()
+{
+	if (PanelOld.GraphEditor.IsValid() && PanelNew.GraphEditor.IsValid())
+	{
+		if (bLockViews)
+		{
+			PanelOld.GraphEditor.Pin()->LockToGraphEditor(PanelNew.GraphEditor);
+			PanelNew.GraphEditor.Pin()->LockToGraphEditor(PanelOld.GraphEditor);
+		}
+		else
+		{
+			PanelOld.GraphEditor.Pin()->UnlockFromGraphEditor(PanelNew.GraphEditor);
+			PanelNew.GraphEditor.Pin()->UnlockFromGraphEditor(PanelOld.GraphEditor);
+		}
+	}
+}
+
+FMaterialDiffPanel& SMaterialDiff::GetDiffPanelForNode(UEdGraphNode& Node)
+{
+	TSharedPtr<SGraphEditor> OldGraphEditorPtr = PanelOld.GraphEditor.Pin();
+	if (OldGraphEditorPtr.IsValid() && Node.GetGraph() == OldGraphEditorPtr->GetCurrentGraph())
+	{
+		return PanelOld;
+	}
+
+	TSharedPtr<SGraphEditor> NewGraphEditorPtr = PanelNew.GraphEditor.Pin();
+	if (NewGraphEditorPtr.IsValid() && Node.GetGraph() == NewGraphEditorPtr->GetCurrentGraph())
+	{
+		return PanelNew;
+	}
+
+	ensureMsgf(false, TEXT("Looking for node %s but it cannot be found in provided panels"), *Node.GetName());
+	static FMaterialDiffPanel Default {};
+	return Default;
+}
+
+void SMaterialDiff::HandleGraphChanged(const FString& GraphPath)
+{
+	SetCurrentMode(GraphMode);
+
+	TSharedPtr<TArray<FDiffSingleResult>> DiffResults;
+	int32 RealDifferencesStartIndex = INDEX_NONE;
+	
+	UEdGraph* NewGraph = MaterialGraphToDiff->GetNewGraph();
+	UEdGraph* OldGraph = MaterialGraphToDiff->GetOldGraph();
+	DiffResults = MaterialGraphToDiff->FoundDiffs;
+	RealDifferencesStartIndex = MaterialGraphToDiff->RealDifferencesStartIndex;
+
+	const TAttribute<int32> FocusedDiffResult = TAttribute<int32>::CreateLambda([this, RealDifferencesStartIndex]() {
+		int32 FocusedDiffResult = INDEX_NONE;
+		if (RealDifferencesStartIndex != INDEX_NONE)
+		{
+			FocusedDiffResult = DiffTreeView::CurrentDifference(DifferencesTreeView.ToSharedRef(), RealDifferences) - RealDifferencesStartIndex;
+		}
+
+		// find selected index in all the graphs, and subtract the index of the first entry in this graph
+		return FocusedDiffResult;
+	});
+	
+	// we could avoid regenerating the panels if the graph was unchanged ala the blueprint diff tool..
+	if (!PanelOld.GraphEditor.IsValid())
+	{
+		PanelOld.GeneratePanel(OldGraph, DiffResults, FocusedDiffResult);
+	}
+
+	if (!PanelNew.GraphEditor.IsValid())
+	{
+		PanelNew.GeneratePanel(NewGraph, DiffResults, FocusedDiffResult);
+	}
+
+	ResetGraphEditors();
+}
+
+void SMaterialDiff::GenerateDifferencesList()
+{
+	PrimaryDifferencesList.Empty();
+	RealDifferences.Empty();
+
+	const auto CreateDetailsView = []() {
+		FDetailsViewArgs DetailsViewArgs;
+		DetailsViewArgs.bHideSelectionTip = true;
+		DetailsViewArgs.DefaultsOnlyVisibility = EEditDefaultsOnlyNodeVisibility::Hide;
+		DetailsViewArgs.bShowSectionSelector = true;
+		DetailsViewArgs.bAllowFavoriteSystem = true;
+		DetailsViewArgs.bShowScrollBar = true;
+		DetailsViewArgs.NameAreaSettings = FDetailsViewArgs::HideNameArea;
+
+		FPropertyEditorModule& PropertyModule = FModuleManager::LoadModuleChecked<FPropertyEditorModule>("PropertyEditor");
+		return PropertyModule.CreateDetailView(DetailsViewArgs);
+	};
+
+	if (PanelOld.MaterialGraph)
+	{
+		PanelOld.MaterialNodeDetailsView = CreateDetailsView();
+	}
+	if (PanelNew.MaterialGraph)
+	{
+		PanelNew.MaterialNodeDetailsView = CreateDetailsView();
+	}
+	
+	MaterialGraphToDiff = MakeShared<FMaterialToDiff>(this, 
+		                                             PanelOld.MaterialGraph, 
+		                                             PanelNew.MaterialGraph,
+													 PanelOld.RevisionInfo,
+													 PanelNew.RevisionInfo);
+	MaterialGraphToDiff->EnableComments(DifferencesTreeView.ToWeakPtr());
+	MaterialGraphToDiff->GenerateTreeEntries(PrimaryDifferencesList, RealDifferences);
+
+	ModePanels.Add(GraphMode, GenerateMaterialGraphPanel());
+
+	DifferencesTreeView->RebuildList();
+}
+
+void SMaterialDiff::OnCloseAssetEditor(UObject* Asset, EAssetEditorCloseReason CloseReason)
+{
+	if (CloseReason == EAssetEditorCloseReason::CloseAllAssetEditors)
+	{
+		// Tell our window to close and set our selves to collapsed to try and stop it from ticking
+		SetVisibility(EVisibility::Collapsed);
+	
+		if (AssetEditorCloseDelegate.IsValid())
+		{
+			GEditor->GetEditorSubsystem<UAssetEditorSubsystem>()->OnAssetEditorRequestClose().Remove(AssetEditorCloseDelegate);
+		}
+	
+		if (WeakParentWindow.IsValid())
+		{
+			WeakParentWindow.Pin()->RequestDestroyWindow();
+		}
+	}
+}
+
+SMaterialDiff::FMaterialGraphDiffControl SMaterialDiff::GenerateMaterialGraphPanel()
+{
+	FMaterialGraphDiffControl MaterialGraphPanel{};
+
+	PanelOld.SetViewportToDisplay();
+	PanelNew.SetViewportToDisplay();
+
+	MaterialGraphPanel.Widget = SNew(SVerticalBox)
+	+SVerticalBox::Slot()
+	.FillHeight(1.f)
+	[
+		SNew(SHorizontalBox)
+		+SHorizontalBox::Slot()
+		.FillWidth(1.f)
+		[
+			SNew(SSplitter)
+			.Orientation(Orient_Vertical)
+			+SSplitter::Slot()
+			.Value(0.6f)
+			[
+				SAssignNew(DiffGraphSplitter, SSplitter)
+				.PhysicalSplitterHandleSize(10.0f)
+				.Orientation(bVerticalSplitGraphMode ? Orient_Horizontal : Orient_Vertical)
+				+SSplitter::Slot()
+				[
+				   GenerateMaterialGraphWidgetForPanel(PanelOld)
+				]
+				+SSplitter::Slot()
+				[
+				   GenerateMaterialGraphWidgetForPanel(PanelNew)
+				]
+			]
+			+SSplitter::Slot()
+			.Value(0.4f)
+			[
+				SNew(SSplitter)
+				.Orientation(Orient_Horizontal)
+				+ SSplitter::Slot()
+				[
+					PanelOld.GetViewportToDisplay()
+				]
+				+SSplitter::Slot()
+				[
+					PanelOld.GetMaterialNodeDetailsViewWidget()
+				]
+				+SSplitter::Slot()
+				[
+					PanelNew.GetMaterialNodeDetailsViewWidget()
+				]
+				+SSplitter::Slot()
+				[
+					PanelNew.GetViewportToDisplay()
+				]
+			]
+		]
+	];
+
+	return MaterialGraphPanel;
+}
+
+TSharedRef<SOverlay> SMaterialDiff::GenerateMaterialGraphWidgetForPanel(FMaterialDiffPanel& OutDiffPanel) const
+{
+	return SNew(SOverlay) 
+	+ SOverlay::Slot() // Graph slot
+	[
+		SAssignNew(OutDiffPanel.GraphEditorBox, SBox)
+		.HAlign(HAlign_Fill)
+		[
+			DefaultEmptyPanel()
+		]
+	] 
+	+SOverlay::Slot() // Revision info slot
+	.VAlign(VAlign_Bottom)
+	.HAlign(HAlign_Right)
+	.Padding(FMargin(20.0f, 10.0f))
+	[
+		GenerateRevisionInfoWidgetForPanel(OutDiffPanel.OverlayGraphRevisionInfo, DiffViewUtils::GetPanelLabel(OutDiffPanel.MaterialGraph, OutDiffPanel.RevisionInfo, FText()))
+	];
+}
+
+TSharedRef<SBox> SMaterialDiff::GenerateRevisionInfoWidgetForPanel(TSharedPtr<SWidget>& OutGeneratedWidget, const FText& InRevisionText) const
+{
+	return SAssignNew(OutGeneratedWidget,SBox)
+	.Padding(FMargin(4.0f, 10.0f))
+	.VAlign(VAlign_Center)
+	.HAlign(HAlign_Left)
+	[
+		SNew(STextBlock)
+		.TextStyle(FAppStyle::Get(), "DetailsView.CategoryTextStyle")
+		.Text(InRevisionText)
+		.ShadowColorAndOpacity(FColor::Black)
+		.ShadowOffset(FVector2D(1.4f, 1.4f))
+	];
+}
+
+void SMaterialDiff::SetCurrentMode(FName NewMode)
+{
+	CurrentMode = NewMode;
+
+	if (FMaterialGraphDiffControl* FoundControl = ModePanels.Find(NewMode))
+	{
+		ModeContents->SetContent(FoundControl->Widget.ToSharedRef());
+	}
+	else
+	{
+		ensureMsgf(false, TEXT("Diff panel does not support mode %s"), *NewMode.ToString());
+	}
+
+	OnModeChanged(NewMode);
+}
+
+void SMaterialDiff::OnModeChanged(const FName& InNewViewMode) const
+{
+	UpdateTopSectionVisibility(InNewViewMode);
+}
+
+void SMaterialDiff::UpdateTopSectionVisibility(const FName& InNewViewMode) const
+{
+	SSplitter* GraphToolBarPtr = GraphToolBarWidget.Get();
+	SSplitter* TopRevisionInfoWidgetPtr = TopRevisionInfoWidget.Get();
+
+	if (!GraphToolBarPtr || !TopRevisionInfoWidgetPtr)
+	{
+		return;
+	}
+
+	if (InNewViewMode == GraphMode)
+	{
+		GraphToolBarPtr->SetVisibility(EVisibility::Visible);
+		TopRevisionInfoWidgetPtr->SetVisibility(EVisibility::Collapsed);
+	}
+	else
+	{
+		GraphToolBarPtr->SetVisibility(EVisibility::Collapsed);
+		TopRevisionInfoWidgetPtr->SetVisibility(EVisibility::HitTestInvisible);
+	}
+}
+
+#undef LOCTEXT_NAMESPACE
diff --git a/Engine/Source/Editor/MaterialEditor/Private/DiffTool/Widgets/SMaterialInstanceDiff.cpp b/Engine/Source/Editor/MaterialEditor/Private/DiffTool/Widgets/SMaterialInstanceDiff.cpp
new file mode 100644
--- /dev/null
+++ b/Engine/Source/Editor/MaterialEditor/Private/DiffTool/Widgets/SMaterialInstanceDiff.cpp
@@ -0,0 +1,433 @@
+// Copyright Epic Games, Inc. All Rights Reserved.
+
+#include "DiffTool/Widgets/SMaterialInstanceDiff.h"
+
+#include "Editor.h"
+#include "Framework/Application/SlateApplication.h"
+#include "Framework/MultiBox/MultiBoxBuilder.h"
+#include "Framework/MultiBox/MultiBoxDefs.h"
+#include "GraphDiffControl.h"
+#include "IDetailsView.h"
+#include "MaterialGraph/MaterialGraph.h"
+#include "Modules/ModuleManager.h"
+#include "PropertyEditorModule.h"
+#include "SBlueprintDiff.h"
+#include "SDetailsDiff.h"
+#include "SDetailsSplitter.h"
+#include "SlateOptMacros.h"
+#include "SMyBlueprint.h"
+#include "SourceControlHelpers.h"
+#include "Styling/AppStyle.h"
+#include "Subsystems/AssetEditorSubsystem.h"
+#include "DiffTool/MaterialToDiff.h"
+#include "Widgets/Docking/SDockTab.h"
+#include "Widgets/Input/SButton.h"
+#include "Widgets/Layout/SSpacer.h"
+#include "Widgets/Layout/SSplitter.h"
+#include "Widgets/Layout/SWidgetSwitcher.h"
+#include "Widgets/SOverlay.h"
+#include "Materials/MaterialInstance.h"
+
+#define LOCTEXT_NAMESPACE "SMaterialGraphDiff"
+
+static const FName DetailsMode = FName(TEXT("DetailsMode"));
+
+BEGIN_SLATE_FUNCTION_BUILD_OPTIMIZATION
+void SMaterialInstanceDiff::Construct(const FArguments& InArgs)
+{
+	check(InArgs._OldMaterialInstance || InArgs._NewMaterialInstance);
+
+	PanelOld.MaterialInstance = InArgs._OldMaterialInstance;
+	PanelNew.MaterialInstance = InArgs._NewMaterialInstance;
+	PanelOld.RevisionInfo = InArgs._OldRevision;
+	PanelNew.RevisionInfo = InArgs._NewRevision;
+
+	if (InArgs._ParentWindow.IsValid())
+	{
+		WeakParentWindow = InArgs._ParentWindow;
+
+		AssetEditorCloseDelegate = GEditor->GetEditorSubsystem<UAssetEditorSubsystem>()->OnAssetEditorRequestClose().AddSP(this, &SMaterialInstanceDiff::OnCloseAssetEditor);
+	}
+
+	FToolBarBuilder NavToolBarBuilder(TSharedPtr<const FUICommandList>(), FMultiBoxCustomization::None);
+	NavToolBarBuilder.AddToolBarButton(
+		FUIAction(
+			FExecuteAction::CreateSP(this, &SMaterialInstanceDiff::PrevDiff),
+			FCanExecuteAction::CreateSP(this, &SMaterialInstanceDiff::HasPrevDiff)),
+		NAME_None,
+		LOCTEXT("PrevDiffLabel", "Prev"),
+		LOCTEXT("PrevDiffTooltip", "Go to previous difference"),
+		FSlateIcon(FAppStyle::GetAppStyleSetName(), "BlueprintDif.PrevDiff")
+	);
+	NavToolBarBuilder.AddToolBarButton(
+		FUIAction(
+			FExecuteAction::CreateSP(this, &SMaterialInstanceDiff::NextDiff),
+			FCanExecuteAction::CreateSP(this, &SMaterialInstanceDiff::HasNextDiff)),
+		NAME_None,
+		LOCTEXT("NextDiffLabel", "Next"),
+		LOCTEXT("NextDiffTooltip", "Go to next difference"),
+		FSlateIcon(FAppStyle::GetAppStyleSetName(), "BlueprintDif.NextDiff")
+	);
+	NavToolBarBuilder.AddToolBarButton(
+		FUIAction(
+			FExecuteAction::CreateSP(this, &SMaterialInstanceDiff::ToggleViewport)),
+		NAME_None,
+		LOCTEXT("ToggleViewportLabel", "Toggle Viewport"),
+		LOCTEXT("ToggleViewportTooltip", "Show/Hide a preview Viewport for the Materials diffed"),
+		TAttribute<FSlateIcon>(this, &SMaterialInstanceDiff::GetViewportImage)
+	);
+
+	DifferencesTreeView = DiffTreeView::CreateTreeView(&PrimaryDifferencesList);
+
+	GenerateDifferencesList();
+
+	const auto TextBlock = [](FText Text) -> TSharedRef<SWidget>
+	{
+		return SNew(SBox)
+		.Padding(FMargin(4.0f, 10.0f))
+		.VAlign(VAlign_Center)
+		.HAlign(HAlign_Left)
+		[
+			SNew(STextBlock)
+			.Visibility(EVisibility::HitTestInvisible)
+			.TextStyle(FAppStyle::Get(), "DetailsView.CategoryTextStyle")
+			.Text(Text)
+		];
+	};
+
+	TopRevisionInfoWidget =
+		SNew(SSplitter)
+		.Visibility(EVisibility::HitTestInvisible)
+		+ SSplitter::Slot()
+		.Value(.2f)
+		[
+			SNew(SBox)
+		]
+		+ SSplitter::Slot()
+		.Value(.8f)
+		[
+			SNew(SSplitter)
+			.PhysicalSplitterHandleSize(10.0f)
+			+ SSplitter::Slot()
+			.Value(.5f)
+			[
+				TextBlock(DiffViewUtils::GetPanelLabel(PanelOld.MaterialInstance, PanelOld.RevisionInfo, FText()))
+			]
+			+ SSplitter::Slot()
+			.Value(.5f)
+			[
+				TextBlock(DiffViewUtils::GetPanelLabel(PanelNew.MaterialInstance, PanelNew.RevisionInfo, FText()))
+			]
+		];
+
+	// Clang-format off
+	ChildSlot
+	[
+		SNew(SBorder)
+		.BorderImage(FAppStyle::GetBrush("Docking.Tab", ".ContentAreaBrush"))
+		[
+			SNew(SOverlay)
+			+SOverlay::Slot()
+			.VAlign(VAlign_Top)
+			[
+				TopRevisionInfoWidget.ToSharedRef()		
+			]
+			+ SOverlay::Slot()
+			.Padding(0.f, 6.f, 0.f, 4.f)
+			[
+				SNew(SOverlay)
+				+SOverlay::Slot()
+				.VAlign(VAlign_Top)
+				.Padding(0.0f, 6.0f, 0.0f, 4.0f)
+				[
+					SNew(SVerticalBox)
+					+SVerticalBox::Slot()
+					.AutoHeight()
+					.Padding(0.0f, 2.0f, 0.0f, 2.0f)
+					[
+						SNew(SHorizontalBox)
+						+SHorizontalBox::Slot()
+						.Padding(4.f)
+						.AutoWidth()
+						[
+							NavToolBarBuilder.MakeWidget()
+						]
+						+SHorizontalBox::Slot()
+						[
+							SNew(SSpacer)
+						]
+					]
+					+SVerticalBox::Slot()
+					.FillHeight(1.f)
+					[
+						SNew(SSplitter)
+						+SSplitter::Slot()
+						.Value(.2f)
+						[
+							SNew(SBorder)
+							.BorderImage(FAppStyle::GetBrush("ToolPanel.GroupBorder"))
+							[
+								DifferencesTreeView.ToSharedRef()
+							]
+						]
+						+SSplitter::Slot()
+						.Value(.8f)
+						[
+							SAssignNew(ModeContents, SBox)
+						]
+					]
+				]
+				
+			]
+		]
+	];
+	// clang-format on
+
+	SetCurrentMode(DetailsMode);
+}
+END_SLATE_FUNCTION_BUILD_OPTIMIZATION
+
+SMaterialInstanceDiff::~SMaterialInstanceDiff()
+{
+	if (AssetEditorCloseDelegate.IsValid())
+	{
+		GEditor->GetEditorSubsystem<UAssetEditorSubsystem>()->OnAssetEditorRequestClose().Remove(AssetEditorCloseDelegate);
+	}
+}
+
+void SMaterialInstanceDiff::OnDiffListSelectionChanged(TSharedPtr<FMaterialDiffResultItem> TheDiff)
+{
+}
+
+TSharedRef<SWidget> SMaterialInstanceDiff::DefaultEmptyPanel()
+{
+	return SNew(SHorizontalBox)
+		   +SHorizontalBox::Slot()
+		   .HAlign(HAlign_Center)
+		   .VAlign(VAlign_Center)
+		   [
+			   SNew(STextBlock)
+			   .Text(LOCTEXT("MaterialInstanceDiffToolTip", "Select Material Instance to Diff"))
+		   ];
+}
+
+TSharedPtr<SWindow> SMaterialInstanceDiff::CreateDiffWindow(FText WindowTitle, TObjectPtr<UMaterialInstance> OldMaterialInstance, TObjectPtr<UMaterialInstance> NewMaterialInstance, const FRevisionInfo& OldRevision, const FRevisionInfo& NewRevision)
+{
+	TSharedPtr<SWindow> Window = SNew(SWindow)
+									 .Title(WindowTitle)
+									 .ClientSize(FVector2D(1000, 800));
+
+	Window->SetContent(SNew(SMaterialInstanceDiff)
+						   .OldMaterialInstance(OldMaterialInstance)
+						   .NewMaterialInstance(NewMaterialInstance)
+						   .OldRevision(OldRevision)
+						   .NewRevision(NewRevision)
+						   .ParentWindow(Window));
+
+	// Make this window a child of the modal window if we've been spawned while one is active.
+	const TSharedPtr<SWindow> ActiveModal = FSlateApplication::Get().GetActiveModalWindow();
+	if (ActiveModal.IsValid())
+	{
+		FSlateApplication::Get().AddWindowAsNativeChild(Window.ToSharedRef(), ActiveModal.ToSharedRef());
+	}
+	else
+	{
+		FSlateApplication::Get().AddWindow(Window.ToSharedRef());
+	}
+
+	return Window;
+}
+
+TSharedPtr<SWindow> SMaterialInstanceDiff::CreateDiffWindow(TObjectPtr<UMaterialInstance> OldMaterialInstance, TObjectPtr<UMaterialInstance> NewMaterialInstance, const FRevisionInfo& OldRevision, const FRevisionInfo& NewRevision, const UClass* ObjectClass)
+{
+	check(OldMaterialInstance || NewMaterialInstance);
+
+	// sometimes we're comparing different revisions of one single asset (other
+	// times we're comparing two completely separate assets altogether)
+	const bool bIsSingleAsset = !OldMaterialInstance || 
+		                        !NewMaterialInstance || 
+								(NewMaterialInstance->GetName() == OldMaterialInstance->GetName());
+
+	FText WindowTitle = FText::Format(LOCTEXT("NamelessMaterialGraphDiff", "{0} Diff"), ObjectClass->GetDisplayNameText());
+	// if we're diffing one asset against itself
+	if (bIsSingleAsset)
+	{
+		// identify the assumed single asset in the window's title
+		FString MaterialGraphName = OldMaterialInstance ? OldMaterialInstance->GetName() : NewMaterialInstance->GetName();
+		WindowTitle = FText::Format(LOCTEXT("NamedMaterialGraphDiff", "{0} - {1} Diff"), FText::FromString(MaterialGraphName), ObjectClass->GetDisplayNameText());
+	}
+
+	return CreateDiffWindow(WindowTitle, OldMaterialInstance, NewMaterialInstance, OldRevision, NewRevision);
+}
+
+void SMaterialInstanceDiff::NextDiff()
+{
+	DiffTreeView::HighlightNextDifference(DifferencesTreeView.ToSharedRef(), RealDifferences, PrimaryDifferencesList);
+}
+
+void SMaterialInstanceDiff::PrevDiff()
+{
+	DiffTreeView::HighlightPrevDifference(DifferencesTreeView.ToSharedRef(), RealDifferences, PrimaryDifferencesList);
+}
+
+bool SMaterialInstanceDiff::HasNextDiff() const
+{
+	return DiffTreeView::HasNextDifference(DifferencesTreeView.ToSharedRef(), RealDifferences);
+}
+
+bool SMaterialInstanceDiff::HasPrevDiff() const
+{
+	return DiffTreeView::HasPrevDifference(DifferencesTreeView.ToSharedRef(), RealDifferences);
+}
+
+void SMaterialInstanceDiff::ToggleViewport()
+{
+	bShowViewport = !bShowViewport;
+}
+
+FSlateIcon SMaterialInstanceDiff::GetViewportImage() const
+{
+	return FSlateIcon(FAppStyle::GetAppStyleSetName(), "LevelEditor.Tabs.Viewports");
+}
+ 
+void SMaterialInstanceDiff::GenerateDifferencesList()
+{
+	PrimaryDifferencesList.Empty();
+	RealDifferences.Empty();
+
+	const auto CreateDetailsView = []()
+	{
+		FDetailsViewArgs DetailsViewArgs;
+		DetailsViewArgs.bHideSelectionTip = true;
+		DetailsViewArgs.DefaultsOnlyVisibility = EEditDefaultsOnlyNodeVisibility::Hide;
+		DetailsViewArgs.bShowSectionSelector = true;
+		DetailsViewArgs.bAllowFavoriteSystem = true;
+		DetailsViewArgs.bShowScrollBar = true;
+		DetailsViewArgs.NameAreaSettings = FDetailsViewArgs::HideNameArea;
+
+		FPropertyEditorModule& PropertyModule = FModuleManager::LoadModuleChecked<FPropertyEditorModule>("PropertyEditor");
+		return PropertyModule.CreateDetailView(DetailsViewArgs);
+	};
+
+	ModePanels.Add(DetailsMode, GenerateMaterialInstancePanel());
+
+	DifferencesTreeView->RebuildList();
+}
+
+void SMaterialInstanceDiff::OnCloseAssetEditor(UObject* Asset, EAssetEditorCloseReason CloseReason)
+{
+	if (CloseReason == EAssetEditorCloseReason::CloseAllAssetEditors)
+	{
+		// Tell our window to close and set our selves to collapsed to try and stop it from ticking
+		SetVisibility(EVisibility::Collapsed);
+	
+		if (AssetEditorCloseDelegate.IsValid())
+		{
+			GEditor->GetEditorSubsystem<UAssetEditorSubsystem>()->OnAssetEditorRequestClose().Remove(AssetEditorCloseDelegate);
+		}
+	
+		if (WeakParentWindow.IsValid())
+		{
+			WeakParentWindow.Pin()->RequestDestroyWindow();
+		}
+	}
+}
+
+SMaterialInstanceDiff::FMaterialInstanceDiffControl SMaterialInstanceDiff::GenerateMaterialInstancePanel()
+{
+	const TSharedPtr<FDetailsDiffControl> NewDiffControl = MakeShared<FDetailsDiffControl>(PanelOld.MaterialInstance,
+																						   PanelNew.MaterialInstance,
+																						   FOnDiffEntryFocused::CreateRaw(this, &SMaterialInstanceDiff::SetCurrentMode, DetailsMode),
+																						   true);
+	NewDiffControl->EnableComments(DifferencesTreeView.ToWeakPtr());
+	NewDiffControl->GenerateTreeEntries(PrimaryDifferencesList, RealDifferences);
+
+	FMaterialInstanceDiffControl MaterialInstancePanel{};
+
+	PanelOld.SetViewportToDisplay();
+	PanelNew.SetViewportToDisplay();
+
+	DiffDetailSplitter = SNew(SDetailsSplitter);
+
+	if (PanelOld.MaterialInstance)
+	{
+		const UObject* OldMaterialInstance = PanelOld.MaterialInstance;
+		DiffDetailSplitter->AddSlot(
+			SDetailsSplitter::Slot()
+			.Value(0.5f)
+			.DetailsView(NewDiffControl->GetDetailsWidget(PanelOld.MaterialInstance))
+			.DifferencesWithRightPanel(NewDiffControl.ToSharedRef(), &FDetailsDiffControl::GetDifferencesWithRight, OldMaterialInstance)
+		);
+	}
+	if (PanelNew.MaterialInstance)
+	{
+		const UObject* NewMaterialInstance = PanelNew.MaterialInstance;
+		DiffDetailSplitter->AddSlot(
+			SDetailsSplitter::Slot()
+			.Value(0.5f)
+			.DetailsView(NewDiffControl->GetDetailsWidget(PanelNew.MaterialInstance))
+			.DifferencesWithRightPanel(NewDiffControl.ToSharedRef(), &FDetailsDiffControl::GetDifferencesWithRight, NewMaterialInstance)
+		);
+	}
+
+	MaterialInstancePanel.DiffControl = NewDiffControl;
+	MaterialInstancePanel.Widget = SNew(SVerticalBox)
+	+SVerticalBox::Slot()
+	.FillHeight(1.0f)
+	[
+		SNew(SHorizontalBox)
+		+SHorizontalBox::Slot()
+		.FillWidth(1.0f)
+		[
+			SNew(SSplitter)
+			.Orientation(Orient_Vertical)
+			+SSplitter::Slot()
+			.Value(0.6f)
+			[
+				DiffDetailSplitter.ToSharedRef()
+			]
+			+SSplitter::Slot()
+			.Value(0.4f)
+			[
+				SNew(SSplitter)
+				.Orientation(Orient_Horizontal).Visibility_Lambda([this]()
+				{
+					return bShowViewport ? EVisibility::Visible : EVisibility::Collapsed;
+				})
+				+ SSplitter::Slot()
+				[
+					PanelOld.GetViewportToDisplay()
+				]
+				+SSplitter::Slot()
+				[
+					PanelNew.GetViewportToDisplay()
+				]
+			]
+		]
+	];
+
+	return MaterialInstancePanel;
+}
+
+void SMaterialInstanceDiff::SetCurrentMode(FName NewMode)
+{
+	CurrentMode = NewMode;
+
+	if (FMaterialInstanceDiffControl* FoundControl = ModePanels.Find(NewMode))
+	{
+		ModeContents->SetContent(FoundControl->Widget.ToSharedRef());
+	}
+	else
+	{
+		ensureMsgf(false, TEXT("Diff panel does not support mode %s"), *NewMode.ToString());
+	}
+
+	OnModeChanged(NewMode);
+}
+
+void SMaterialInstanceDiff::OnModeChanged(const FName& InNewViewMode) const
+{
+	
+}
+
+#undef LOCTEXT_NAMESPACE
diff --git a/Engine/Source/Editor/MaterialEditor/Private/DiffTool/Widgets/SMaterialRevisionMenu.cpp b/Engine/Source/Editor/MaterialEditor/Private/DiffTool/Widgets/SMaterialRevisionMenu.cpp
new file mode 100644
--- /dev/null
+++ b/Engine/Source/Editor/MaterialEditor/Private/DiffTool/Widgets/SMaterialRevisionMenu.cpp
@@ -0,0 +1,364 @@
+// Copyright Epic Games, Inc. All Rights Reserved.
+
+#include "DiffTool/Widgets/SMaterialRevisionMenu.h"
+
+#include "AssetToolsModule.h"
+#include "DiffUtils.h"
+#include "Framework/Commands/UIAction.h"
+#include "Framework/MultiBox/MultiBoxBuilder.h"
+#include "Input/Reply.h"
+#include "IAssetTypeActions.h"
+#include "ISourceControlModule.h"
+#include "ISourceControlRevision.h"
+#include "ISourceControlState.h"
+#include "Layout/Visibility.h"
+#include "MaterialEditorContext.h"
+#include "MaterialEditor.h"
+#include "Materials/Material.h"
+#include "Materials/MaterialFunction.h"
+#include "Misc/Attribute.h"
+#include "Misc/MessageDialog.h"
+#include "SourceControlHelpers.h"
+#include "SourceControlOperations.h"
+#include "Styling/AppStyle.h"
+#include "Textures/SlateIcon.h"
+#include "ToolMenu.h"
+#include "Widgets/Images/SThrobber.h"
+#include "Widgets/Input/SButton.h"
+#include "Widgets/Layout/SBorder.h"
+#include "Widgets/Text/STextBlock.h"
+
+#define LOCTEXT_NAMESPACE "SMaterialRevisionMenu"
+
+namespace ESourceControlQueryState
+{
+	enum Type
+	{
+		NotQueried,
+		QueryInProgress,
+		Queried,
+	};
+}
+
+SMaterialRevisionMenu::~SMaterialRevisionMenu()
+{
+	// Cancel any operation if this widget is destroyed while in progress
+	if (SourceControlQueryState == ESourceControlQueryState::QueryInProgress)
+	{
+		ISourceControlProvider& SourceControlProvider = ISourceControlModule::Get().GetProvider();
+		if (SourceControlQueryOp.IsValid() && SourceControlProvider.CanCancelOperation(SourceControlQueryOp.ToSharedRef()))
+		{
+			SourceControlProvider.CancelOperation(SourceControlQueryOp.ToSharedRef());
+		}
+	}
+}
+
+void SMaterialRevisionMenu::Construct(const FArguments& InArgs, UObject const* MaterialObj)
+{
+	bIncludeLocalRevision = InArgs._bIncludeLocalRevision;
+	OnRevisionSelected = InArgs._OnRevisionSelected;
+
+	SourceControlQueryState = ESourceControlQueryState::NotQueried;
+
+	ChildSlot
+	[
+		SAssignNew(MenuBox, SVerticalBox) 
+		+SVerticalBox::Slot()
+		[
+			SNew(SBorder)
+			.Visibility(this, &SMaterialRevisionMenu::GetInProgressVisibility)
+			.BorderImage(FAppStyle::GetBrush("Menu.Background"))
+			.Content()
+			[
+				SNew(SHorizontalBox) 
+				+SHorizontalBox::Slot()
+				.AutoWidth()
+				.VAlign(VAlign_Center)
+				[
+					SNew(SThrobber)
+				]
+				+SHorizontalBox::Slot()
+				.AutoWidth()
+				.VAlign(VAlign_Center)
+				.Padding(2.0f, 0.0f, 4.0f, 0.0f)
+				[
+					SNew(STextBlock)
+					.Text(LOCTEXT("MaterialGraphDiffMenuOperationInProgress", "Updating history..."))
+				] 
+				+SHorizontalBox::Slot()
+				.FillWidth(1.0f)
+				.HAlign(HAlign_Right)
+				.VAlign(VAlign_Center)
+				[
+					SNew(SButton)
+					.Visibility(this, &SMaterialRevisionMenu::GetCancelButtonVisibility)
+					.OnClicked(this, &SMaterialRevisionMenu::OnCancelButtonClicked)
+					.VAlign(VAlign_Center)
+					.HAlign(HAlign_Center)
+					.Content()
+					[
+						SNew(STextBlock)
+						.Text(LOCTEXT("MaterialGraphDiffMenuCancelButton", "Cancel"))
+					]
+				]
+			]
+		]
+	];
+
+	if (MaterialObj)
+	{
+		Filename = SourceControlHelpers::PackageFilename(MaterialObj->GetPathName());
+		
+		// make sure the history info is up to date
+		SourceControlQueryOp = ISourceControlOperation::Create<FUpdateStatus>();
+		SourceControlQueryOp->SetUpdateHistory(true);
+		ISourceControlModule::Get().GetProvider().Execute(SourceControlQueryOp.ToSharedRef(), Filename, EConcurrency::Asynchronous, FSourceControlOperationComplete::CreateSP(this, &SMaterialRevisionMenu::OnSourceControlQueryComplete));
+
+		SourceControlQueryState = ESourceControlQueryState::QueryInProgress;
+	}
+}
+
+static void OnDiffRevisionPicked(const FRevisionInfo& RevisionInfo, TWeakObjectPtr<UObject> MaterialObjWeak)
+{
+	UObject* MaterialObj = MaterialObjWeak.Get();
+	if (!MaterialObj)
+	{
+		return;
+	}
+		
+	// Get the SCC state
+	ISourceControlProvider& SourceControlProvider = ISourceControlModule::Get().GetProvider();
+	const FString Filename = SourceControlHelpers::PackageFilename(MaterialObj->GetPathName());
+	FSourceControlStatePtr SourceControlState = SourceControlProvider.GetState(Filename, EStateCacheUsage::Use);
+
+	if (!SourceControlState.IsValid())
+	{
+		return;
+	}
+
+	for (int32 HistoryIndex = 0; HistoryIndex != SourceControlState->GetHistorySize(); ++HistoryIndex)
+	{
+		TSharedPtr<ISourceControlRevision, ESPMode::ThreadSafe> Revision = SourceControlState->GetHistoryItem(HistoryIndex);
+
+		if (!Revision.IsValid() || Revision->GetRevision() != RevisionInfo.Revision)
+		{
+			continue;
+		}
+
+		// Get the revision of this package from source control
+		if (UPackage* PreviousTempPkg = DiffUtils::LoadPackageForDiff(Revision))
+		{
+			FString PreviousAssetName = FPaths::GetBaseFilename(Filename, true);
+			if (UObject* PreviousMaterial = FindObject<UObject>(PreviousTempPkg, *PreviousAssetName))
+			{
+				FAssetToolsModule& AssetToolsModule = FModuleManager::LoadModuleChecked<FAssetToolsModule>(TEXT("AssetTools"));
+				FRevisionInfo OldRevision{ Revision->GetRevision(), Revision->GetCheckInIdentifier(), Revision->GetDate() };
+				FRevisionInfo CurrentRevision{ TEXT(""), Revision->GetCheckInIdentifier(), Revision->GetDate() };
+
+				AssetToolsModule.Get().DiffAssets(PreviousMaterial, MaterialObj, OldRevision, CurrentRevision);
+			}
+		}
+		else
+		{
+			FMessageDialog::Open(EAppMsgType::Ok, NSLOCTEXT("SourceControl.HistoryWindow", "UnableToLoadMaterialGraphs", "Unable to load MaterialGraphs to diff. Please retry."));
+		}
+
+		break;
+	}
+}
+
+void SMaterialRevisionMenu::MakeDiffMenu(UToolMenu* Menu)
+{
+	FToolMenuSection& DiffSection = Menu->AddSection("SourceControl");
+	DiffSection.AddDynamicEntry("SourceControlCommands", FNewToolMenuSectionDelegate::CreateLambda([](FToolMenuSection& InSection)
+		{
+			UMaterialEditorMenuContext* SubMenuContext = InSection.FindContext<UMaterialEditorMenuContext>();
+			if (SubMenuContext && SubMenuContext->MaterialEditor.IsValid())
+			{
+				TWeakPtr<IMaterialEditor> MaterialEditorWeak = SubMenuContext->MaterialEditor;
+				
+				const auto RevisionMenu = [MaterialEditorWeak]() -> TSharedRef<SWidget>
+				{
+					TSharedPtr<IMaterialEditor> MaterialEditorStrong = MaterialEditorWeak.Pin();
+					if (MaterialEditorStrong && ISourceControlModule::Get().IsEnabled() && ISourceControlModule::Get().GetProvider().IsAvailable())
+					{
+						UMaterialInterface* MaterialObj = MaterialEditorStrong->GetMaterialInterface();
+						if( MaterialObj == nullptr || 
+							(MaterialObj->HasAnyFlags(RF_Transient) || MaterialObj->GetPackage() == GetTransientPackage()))
+						{
+							const TArray< UObject* >* ObjectsBeingEdited = MaterialEditorStrong->GetObjectsCurrentlyBeingEdited();
+							MaterialObj = nullptr;
+							if(ObjectsBeingEdited)
+							{
+								for(UObject* Obj: *ObjectsBeingEdited)
+								{
+									UMaterialInterface* PotentialMaterialObj = Cast<UMaterialInterface>(Obj);
+									if(PotentialMaterialObj && 
+										!(PotentialMaterialObj->HasAnyFlags(RF_Transient) || PotentialMaterialObj->GetPackage() == GetTransientPackage()))
+									{
+										MaterialObj = PotentialMaterialObj;
+										break;
+									}
+								}
+							}
+						}
+						if (MaterialObj)
+						{
+							TWeakObjectPtr<UObject> WeakMaterialObj = MaterialObj;
+							return SNew(SMaterialRevisionMenu, MaterialObj).OnRevisionSelected_Static(&OnDiffRevisionPicked, WeakMaterialObj);
+						}
+					}
+
+					FMenuBuilder MenuBuilder(true, nullptr);
+					MenuBuilder.AddMenuEntry(LOCTEXT("SourceControlDisabled", "Revision control is disabled"),
+												FText(),
+												FSlateIcon(),
+												FUIAction());
+
+					return MenuBuilder.MakeWidget();
+				};
+
+	 			FToolMenuEntry& DiffEntry = InSection.AddEntry(FToolMenuEntry::InitComboButton(
+					"Diff",
+					FUIAction(),
+					FOnGetContent::CreateLambda(RevisionMenu),
+					LOCTEXT("Diff", "Diff"),
+					LOCTEXT("MaterialEditorDiffToolTip", "Diff against previous revisions"),
+					FSlateIcon(FAppStyle::GetAppStyleSetName(), "BlueprintDiff.ToolbarIcon")
+				));
+				DiffEntry.StyleNameOverride = "CalloutToolbar";
+			}
+		}));
+}
+
+EVisibility SMaterialRevisionMenu::GetInProgressVisibility() const
+{
+	return (SourceControlQueryState == ESourceControlQueryState::QueryInProgress) ? EVisibility::Visible : EVisibility::Collapsed;
+}
+
+EVisibility SMaterialRevisionMenu::GetCancelButtonVisibility() const
+{
+	ISourceControlProvider& SourceControlProvider = ISourceControlModule::Get().GetProvider();
+	return SourceControlQueryOp.IsValid() && SourceControlProvider.CanCancelOperation(SourceControlQueryOp.ToSharedRef()) ? EVisibility::Visible : EVisibility::Collapsed;
+}
+
+FReply SMaterialRevisionMenu::OnCancelButtonClicked() const
+{
+	if (SourceControlQueryOp.IsValid())
+	{
+		ISourceControlProvider& SourceControlProvider = ISourceControlModule::Get().GetProvider();
+		SourceControlProvider.CancelOperation(SourceControlQueryOp.ToSharedRef());
+	}
+
+	return FReply::Handled();
+}
+
+void SMaterialRevisionMenu::OnSourceControlQueryComplete(const FSourceControlOperationRef& InOperation, ECommandResult::Type InResult)
+{
+	check(SourceControlQueryOp == InOperation);
+
+	// Add pop-out menu for each revision
+	FMenuBuilder MenuBuilder(true, nullptr);
+
+	MenuBuilder.BeginSection("AddDiffRevision", LOCTEXT("Revisions", "Revisions"));
+	if (bIncludeLocalRevision)
+	{
+		FText const ToolTipText = LOCTEXT("LocalRevisionToolTip", "The current copy you have saved to disk (locally)");
+
+		FOnRevisionSelected OnRevisionSelectedDelegate = OnRevisionSelected;
+		auto OnMenuItemSelected = [OnRevisionSelectedDelegate]() {
+			OnRevisionSelectedDelegate.ExecuteIfBound(FRevisionInfo::InvalidRevision());
+		};
+
+		MenuBuilder.AddMenuEntry(LOCTEXT("LocalRevision", "Local"), ToolTipText, FSlateIcon(), FUIAction(FExecuteAction::CreateLambda(OnMenuItemSelected)));
+	}
+
+	if (InResult == ECommandResult::Succeeded)
+	{
+		// get the cached state
+		ISourceControlProvider& SourceControlProvider = ISourceControlModule::Get().GetProvider();
+		FSourceControlStatePtr SourceControlState = SourceControlProvider.GetState(Filename, EStateCacheUsage::Use);
+
+		if (SourceControlState.IsValid() && SourceControlState->GetHistorySize() > 0)
+		{
+			// Figure out the highest revision # (so we can label it "Depot")
+			int32 LatestRevision = 0;
+			for (int32 HistoryIndex = 0; HistoryIndex != SourceControlState->GetHistorySize(); ++HistoryIndex)
+			{
+				TSharedPtr<ISourceControlRevision, ESPMode::ThreadSafe> Revision = SourceControlState->GetHistoryItem(HistoryIndex);
+				if (Revision.IsValid() && Revision->GetRevisionNumber() > LatestRevision)
+				{
+					LatestRevision = Revision->GetRevisionNumber();
+				}
+			}
+
+			for (int32 HistoryIndex = 0; HistoryIndex != SourceControlState->GetHistorySize(); ++HistoryIndex)
+			{
+				TSharedPtr<ISourceControlRevision, ESPMode::ThreadSafe> Revision = SourceControlState->GetHistoryItem(HistoryIndex);
+				if (Revision.IsValid())
+				{
+					FInternationalization& I18N = FInternationalization::Get();
+
+					FText Label = FText::Format(LOCTEXT("RevisionNumber", "Revision {0}"), FText::AsNumber(Revision->GetRevisionNumber(), NULL, I18N.GetInvariantCulture()));
+
+					FFormatNamedArguments Args;
+					Args.Add(TEXT("CheckInNumber"), FText::AsNumber(Revision->GetCheckInIdentifier(), NULL, I18N.GetInvariantCulture()));
+					Args.Add(TEXT("Revision"), FText::FromString(Revision->GetRevision()));
+					Args.Add(TEXT("UserName"), FText::FromString(Revision->GetUserName()));
+					Args.Add(TEXT("DateTime"), FText::AsDate(Revision->GetDate()));
+					Args.Add(TEXT("ChangelistDescription"), FText::FromString(Revision->GetDescription()));
+					FText ToolTipText;
+					if (ISourceControlModule::Get().GetProvider().UsesChangelists())
+					{
+						ToolTipText = FText::Format(LOCTEXT("ChangelistToolTip", "CL #{CheckInNumber} {UserName} \n{DateTime} \n{ChangelistDescription}"), Args);
+					}
+					else
+					{
+						ToolTipText = FText::Format(LOCTEXT("RevisionToolTip", "{Revision} {UserName} \n{DateTime} \n{ChangelistDescription}"), Args);
+					}
+
+					if (LatestRevision == Revision->GetRevisionNumber())
+					{
+						Label = LOCTEXT("Depot", "Depot");
+					}
+
+					FRevisionInfo RevisionInfo{ Revision->GetRevision(), Revision->GetCheckInIdentifier(), Revision->GetDate() };
+					FOnRevisionSelected OnRevisionSelectedDelegate = OnRevisionSelected;
+					auto OnMenuItemSelected = [RevisionInfo, OnRevisionSelectedDelegate]() 
+					{
+						OnRevisionSelectedDelegate.ExecuteIfBound(RevisionInfo);
+					};
+					MenuBuilder.AddMenuEntry(TAttribute<FText>(Label), ToolTipText, FSlateIcon(), FUIAction(FExecuteAction::CreateLambda(OnMenuItemSelected)));
+				}
+			}
+		}
+		else if (!bIncludeLocalRevision)
+		{
+			// Show 'empty' item in toolbar
+			MenuBuilder.AddMenuEntry(LOCTEXT("NoRevisionHistory", "No revisions found"),
+									 FText(),
+									 FSlateIcon(),
+									 FUIAction());
+		}
+	}
+	else if (!bIncludeLocalRevision)
+	{
+		// Show 'empty' item in toolbar
+		MenuBuilder.AddMenuEntry(LOCTEXT("NoRevisionHistory", "No revisions found"),
+								 FText(),
+								 FSlateIcon(),
+								 FUIAction());
+	}
+
+	MenuBuilder.EndSection();
+
+	MenuBox->AddSlot()
+	[
+		MenuBuilder.MakeWidget(nullptr, 500)
+	];
+
+	SourceControlQueryOp.Reset();
+	SourceControlQueryState = ESourceControlQueryState::Queried;
+}
+
+#undef LOCTEXT_NAMESPACE
diff --git a/Engine/Source/Editor/MaterialEditor/Private/DiffTool/Widgets/SMaterialRevisionMenu.h b/Engine/Source/Editor/MaterialEditor/Private/DiffTool/Widgets/SMaterialRevisionMenu.h
new file mode 100644
--- /dev/null
+++ b/Engine/Source/Editor/MaterialEditor/Private/DiffTool/Widgets/SMaterialRevisionMenu.h
@@ -0,0 +1,65 @@
+// Copyright Epic Games, Inc. All Rights Reserved.
+
+#pragma once
+
+#include "CoreMinimal.h"
+#include "Delegates/Delegate.h"
+#include "ISourceControlOperation.h"
+#include "ISourceControlProvider.h"
+#include "Widgets/SCompoundWidget.h"
+
+class FReply;
+class FUpdateStatus;
+class SVerticalBox;
+class UToolMenu;
+
+struct EVisibility;
+struct FRevisionInfo;
+
+class SMaterialRevisionMenu : public SCompoundWidget
+{
+	DECLARE_DELEGATE_OneParam(FOnRevisionSelected, const FRevisionInfo&)
+
+public : 
+SLATE_BEGIN_ARGS(SMaterialRevisionMenu) : _bIncludeLocalRevision(false) {}
+	SLATE_ARGUMENT(bool, bIncludeLocalRevision)
+	SLATE_EVENT(FOnRevisionSelected, OnRevisionSelected)
+SLATE_END_ARGS()
+
+	void Construct(const FArguments& InArgs, UObject const* MaterialObj);
+	virtual ~SMaterialRevisionMenu();
+
+	// helper function to add the diff menu to a UToolMenu:
+	static void MakeDiffMenu(UToolMenu* Menu);
+
+private:
+	/** Delegate used to determine the visibility 'in progress' widgets */
+	EVisibility GetInProgressVisibility() const;
+
+	/** Delegate used to determine the visibility of the cancel button */
+	EVisibility GetCancelButtonVisibility() const;
+
+	/** Delegate used to cancel a source control operation in progress */
+	FReply OnCancelButtonClicked() const;
+
+	/** Callback for when the source control operation is complete */
+	void OnSourceControlQueryComplete(const FSourceControlOperationRef& InOperation, ECommandResult::Type InResult);
+
+	/** Whether or not we should include the local revision to the menu */
+	bool bIncludeLocalRevision;
+
+	/** Called when we select one Revision from the list in the menu */
+	FOnRevisionSelected OnRevisionSelected;
+
+	/** The name of the file we want revision info for */
+	FString Filename;
+
+	/** The box we are using to display our menu */
+	TSharedPtr<SVerticalBox> MenuBox;
+
+	/** The source control operation in progress */
+	TSharedPtr<FUpdateStatus, ESPMode::ThreadSafe> SourceControlQueryOp;
+
+	/** The state of the SCC query */
+	uint32 SourceControlQueryState;
+};
diff --git a/Engine/Source/Editor/MaterialEditor/Private/MaterialEditor.cpp b/Engine/Source/Editor/MaterialEditor/Private/MaterialEditor.cpp
--- a/Engine/Source/Editor/MaterialEditor/Private/MaterialEditor.cpp
+++ b/Engine/Source/Editor/MaterialEditor/Private/MaterialEditor.cpp
@@ -184,6 +184,8 @@
 #include "ToolMenus.h"
 #include "ViewportToolbar/UnrealEdViewportToolbar.h"
 
+#include "DiffTool/Widgets/SMaterialRevisionMenu.h"
+
 #define LOCTEXT_NAMESPACE "MaterialEditor"
 
 DEFINE_LOG_CATEGORY_STATIC(LogMaterialEditor, Log, All);
@@ -2151,12 +2153,11 @@ void FMaterialEditor::RegisterToolBar()
 			FToolMenuSection& Section = ToolBar->AddSection("Stats", TAttribute<FText>(), InsertAfterAssetSection);
 			Section.AddEntry(FToolMenuEntry::InitToolBarButton(FMaterialEditorCommands::Get().ToggleMaterialStats));
 			Section.AddEntry(FToolMenuEntry::InitToolBarButton(FMaterialEditorCommands::Get().TogglePlatformStats));
-			
 		}
 
-		
+		SMaterialRevisionMenu::MakeDiffMenu(ToolBar);
 	}
-};
+}
 
 void FMaterialEditor::AddInheritanceMenuEntry(FToolMenuSection& Section, const FAssetData& AssetData, bool bIsFunctionPreviewMaterial)
 {
diff --git a/Engine/Source/Editor/MaterialEditor/Private/MaterialEditorUtilities.cpp b/Engine/Source/Editor/MaterialEditor/Private/MaterialEditorUtilities.cpp
--- a/Engine/Source/Editor/MaterialEditor/Private/MaterialEditorUtilities.cpp
+++ b/Engine/Source/Editor/MaterialEditor/Private/MaterialEditorUtilities.cpp
@@ -124,6 +124,14 @@ FText FMaterialEditorUtilities::GetOriginalObjectName(const class UEdGraph* Grap
 	{
 		return MaterialEditor->GetOriginalObjectName();
 	}
+
+	// When diffing Material, there might not be any MaterialEditor opened for the Material
+	// So we have to get the MaterialGraph's name another way
+	if (UMaterial* GraphMaterial = Cast<UMaterial>(Graph->GetOuter()))
+	{
+		return FText::FromString(GraphMaterial->GetName());
+	}
+	
 	return FText::GetEmpty();
 }
 
diff --git a/Engine/Source/Editor/MaterialEditor/Private/MaterialInstanceEditor.cpp b/Engine/Source/Editor/MaterialEditor/Private/MaterialInstanceEditor.cpp
--- a/Engine/Source/Editor/MaterialEditor/Private/MaterialInstanceEditor.cpp
+++ b/Engine/Source/Editor/MaterialEditor/Private/MaterialInstanceEditor.cpp
@@ -68,6 +68,8 @@
 #include "SMaterialEditorSubstrateWidget.h"
 #include "ViewportToolbar/UnrealEdViewportToolbar.h"
 
+#include "DiffTool/Widgets/SMaterialRevisionMenu.h"
+
 #define LOCTEXT_NAMESPACE "MaterialInstanceEditor"
 
 DEFINE_LOG_CATEGORY_STATIC(LogMaterialInstanceEditor, Log, All);
@@ -997,6 +999,8 @@ void FMaterialInstanceEditor::RegisterToolBar()
 
 		FToolMenuSection& UISection = ToolBar->AddSection("Stats", TAttribute<FText>(), InsertAfterAssetSection);
 		UISection.AddEntry(FToolMenuEntry::InitToolBarButton(FMaterialEditorCommands::Get().TogglePlatformStats));
+
+		SMaterialRevisionMenu::MakeDiffMenu(ToolBar);
 	}
 }
 
diff --git a/Engine/Source/Editor/MaterialEditor/Private/SMaterialEditorViewport.cpp b/Engine/Source/Editor/MaterialEditor/Private/SMaterialEditorViewport.cpp
--- a/Engine/Source/Editor/MaterialEditor/Private/SMaterialEditorViewport.cpp
+++ b/Engine/Source/Editor/MaterialEditor/Private/SMaterialEditorViewport.cpp
@@ -249,6 +249,8 @@ void FMaterialEditorViewportClient::FocusViewportOnBounds(const FBoxSphereBounds
 void SMaterialEditor3DPreviewViewport::Construct(const FArguments& InArgs)
 {
 	MaterialEditorPtr = InArgs._MaterialEditor;
+	PreviewMaterialPtr = InArgs._PreviewMaterial;
+
 	AdvancedPreviewScene = MakeShareable(new FAdvancedPreviewScene(FPreviewScene::ConstructionValues()));
 
 	// restore last used feature level
@@ -273,13 +275,10 @@ void SMaterialEditor3DPreviewViewport::Construct(const FArguments& InArgs)
 	PreviewMeshComponent = nullptr;
 	PostProcessVolumeActor = nullptr;
 
-	if (MaterialEditorPtr.IsValid())
+	UMaterialInterface* Material = GetPreviewingMaterialInterface();
+	if (Material)
 	{
-		UMaterialInterface* Material = MaterialEditorPtr.Pin()->GetMaterialInterface();
-		if (Material)
-		{
-			SetPreviewMaterial(Material);
-		}
+		SetPreviewMaterial(Material);
 	}
 
 	SetPreviewAsset( GUnrealEd->GetThumbnailManager()->EditorSphere );
@@ -347,7 +346,7 @@ void SMaterialEditor3DPreviewViewport::RefreshViewport()
 
 bool SMaterialEditor3DPreviewViewport::SetPreviewAsset(UObject* InAsset)
 {
-	if (!MaterialEditorPtr.Pin()->ApproveSetPreviewAsset(InAsset))
+	if (MaterialEditorPtr.IsValid() && !MaterialEditorPtr.Pin()->ApproveSetPreviewAsset(InAsset))
 	{
 		return false;
 	}
@@ -724,8 +723,10 @@ void SMaterialEditor3DPreviewViewport::BindCommands()
 
 	const FMaterialEditorCommands& Commands = FMaterialEditorCommands::Get();
 
-	check(MaterialEditorPtr.IsValid());
-	CommandList->Append(MaterialEditorPtr.Pin()->GetToolkitCommands());
+	if(MaterialEditorPtr.IsValid())
+	{
+		CommandList->Append(MaterialEditorPtr.Pin()->GetToolkitCommands());
+	}
 
 	// Add the commands to the toolkit command list so that the toolbar buttons can find them
 	CommandList->MapAction(
@@ -773,6 +774,11 @@ void SMaterialEditor3DPreviewViewport::OnFocusViewportToSelection()
 	}
 }
 
+UMaterialInterface* SMaterialEditor3DPreviewViewport::GetPreviewingMaterialInterface()
+{
+	return MaterialEditorPtr.IsValid() ? MaterialEditorPtr.Pin()->GetMaterialInterface() : PreviewMaterialPtr.Get();
+}
+
 void SMaterialEditor3DPreviewViewport::OnSetPreviewPrimitive(EThumbnailPrimType PrimType, bool bInitialLoad)
 {
 	if (SceneViewport.IsValid())
@@ -791,7 +797,7 @@ void SMaterialEditor3DPreviewViewport::OnSetPreviewPrimitive(EThumbnailPrimType
 			SetPreviewAsset(Primitive);
 			
 			// Clear the thumbnail preview mesh
-			if (UMaterialInterface* MaterialInterface = MaterialEditorPtr.Pin()->GetMaterialInterface())
+			if (UMaterialInterface* MaterialInterface = GetPreviewingMaterialInterface())
 			{
 				MaterialInterface->PreviewMesh = nullptr;
 				FMaterialEditor::UpdateThumbnailInfoPreviewMesh(MaterialInterface);
@@ -816,7 +822,7 @@ void SMaterialEditor3DPreviewViewport::OnSetPreviewMeshFromSelection()
 	bool bFoundPreviewMesh = false;
 	FEditorDelegates::LoadSelectedAssetsIfNeeded.Broadcast();
 
-	UMaterialInterface* MaterialInterface = MaterialEditorPtr.Pin()->GetMaterialInterface();
+	UMaterialInterface* MaterialInterface = GetPreviewingMaterialInterface();
 
 	// Look for a selected asset that can be converted to a mesh component
 	for (FSelectionIterator SelectionIt(*GEditor->GetSelectedObjects()); SelectionIt && !bFoundPreviewMesh; ++SelectionIt)
@@ -1005,7 +1011,7 @@ TSharedPtr<IPreviewProfileController> SMaterialEditor3DPreviewViewport::CreatePr
 void SMaterialEditor3DPreviewViewport::PopulateViewportOverlays(TSharedRef<class SOverlay> Overlay)
 {
 	// only show topology widget for material instances
-	if (Substrate::IsMaterialLayeringSupportEnabled() && Cast<UMaterialInstance>(MaterialEditorPtr.Pin()->GetMaterialInterface()))
+	if (Substrate::IsMaterialLayeringSupportEnabled() && MaterialEditorPtr.IsValid() && Cast<UMaterialInstance>(MaterialEditorPtr.Pin()->GetMaterialInterface()))
 	{
 		Overlay->AddSlot()
 			// clang-format off
diff --git a/Engine/Source/Editor/MaterialEditor/Private/SMaterialEditorViewport.h b/Engine/Source/Editor/MaterialEditor/Private/SMaterialEditorViewport.h
--- a/Engine/Source/Editor/MaterialEditor/Private/SMaterialEditorViewport.h
+++ b/Engine/Source/Editor/MaterialEditor/Private/SMaterialEditorViewport.h
@@ -28,8 +28,10 @@ class UMeshComponent;
 class SMaterialEditor3DPreviewViewport : public SEditorViewport, public FGCObject, public ICommonEditorViewportToolbarInfoProvider
 {
 public:
-	SLATE_BEGIN_ARGS( SMaterialEditor3DPreviewViewport ){}
+	SLATE_BEGIN_ARGS( SMaterialEditor3DPreviewViewport )
+		: _PreviewMaterial(nullptr){}
 		SLATE_ARGUMENT(TWeakPtr<IMaterialEditor>, MaterialEditor)
+		SLATE_ARGUMENT(UMaterialInterface*, PreviewMaterial)
 	SLATE_END_ARGS()
 
 	void Construct(const FArguments& InArgs);
@@ -100,6 +102,8 @@ protected:
 	virtual void BindCommands() override;
 	virtual void OnFocusViewportToSelection() override;
 private:
+	/** If we have a MaterialEditor, it will look up the preview material from that. Otherwise it will fallback to previewing this material */
+	UMaterialInterface* GetPreviewingMaterialInterface();
 
 	void OnPropertyChanged(UObject* ObjectBeingModified, FPropertyChangedEvent& PropertyChangedEvent);
 
@@ -114,6 +118,9 @@ private:
 	/** Pointer back to the material editor tool that owns us */
 	TWeakPtr<IMaterialEditor> MaterialEditorPtr;
 
+	/** The material we're previewing if we're not hosted by a material editor. Used for diffs. */
+	TWeakObjectPtr<UMaterialInterface> PreviewMaterialPtr;
+
 	/** Level viewport client */
 	TSharedPtr<class FMaterialEditorViewportClient> EditorViewportClient;
 
