Restore ability to directly edit values for vector channel masks.

Fix bug where FMathStructCustomization would not create any child widgets if CustomizeHeader was skipped.

#tests Material Editor - Parameters, Layer Parameter Preview tabs. Material Instance Editor - Details, Layer Parameters tabs
#rnx

diff --git a/Engine/Source/Editor/DetailCustomizations/Private/MathStructCustomizations.cpp b/Engine/Source/Editor/DetailCustomizations/Private/MathStructCustomizations.cpp
--- a/Engine/Source/Editor/DetailCustomizations/Private/MathStructCustomizations.cpp
+++ b/Engine/Source/Editor/DetailCustomizations/Private/MathStructCustomizations.cpp
@@ -62,6 +62,12 @@ void FMathStructCustomization::CustomizeHeader(TSharedRef<class IPropertyHandle>
 
 void FMathStructCustomization::CustomizeChildren(TSharedRef<class IPropertyHandle> StructPropertyHandle, class IDetailChildrenBuilder& StructBuilder, IPropertyTypeCustomizationUtils& StructCustomizationUtils)
 {
+	// Grab children if CustomizeHeader was never run (e.g. vector mask params in the material editor)
+	if (SortedChildHandles.Num() == 0)
+	{
+		GetSortedChildren(StructPropertyHandle, SortedChildHandles);
+	}
+
 	for (int32 ChildIndex = 0; ChildIndex < SortedChildHandles.Num(); ++ChildIndex)
 	{
 		TSharedRef<IPropertyHandle> ChildHandle = SortedChildHandles[ChildIndex];
diff --git a/Engine/Source/Editor/MaterialEditor/Private/MaterialPropertyHelpers.cpp b/Engine/Source/Editor/MaterialEditor/Private/MaterialPropertyHelpers.cpp
--- a/Engine/Source/Editor/MaterialEditor/Private/MaterialPropertyHelpers.cpp
+++ b/Engine/Source/Editor/MaterialEditor/Private/MaterialPropertyHelpers.cpp
@@ -575,11 +575,6 @@ bool FMaterialPropertyHelpers::ShouldCreateChildPropertiesForParameter(UDEditorP
 		// Static component mask has a checkbox for each channel
 		return false;
 	}
-	else if (UDEditorVectorParameterValue* VectorParam = Cast<UDEditorVectorParameterValue>(Parameter))
-	{
-		// Vector channel mask ahs a dropdown for each channel
-		return !VectorParam->bIsUsedAsChannelMask;
-	}
 
 	return true;
 }
@@ -748,7 +743,8 @@ void FMaterialPropertyHelpers::SetPropertyRowParameterWidget(IDetailPropertyRow&
 			TSharedPtr<SWidget> ValueWidget;
 			PropertyRow.GetDefaultWidgets(NameWidget, ValueWidget);
 
-			PropertyRow.CustomWidget(true)
+			const bool bShowChildren = true;
+			PropertyRow.CustomWidget(bShowChildren)
 				.NameContent()
 				[
 					ParameterValueProperty->CreatePropertyNameWidget()
@@ -823,7 +819,8 @@ void FMaterialPropertyHelpers::SetPropertyRowVectorChannelMaskParameterWidget(ID
 	FOnPropertyComboBoxValueSelected SetMaskValue = FOnPropertyComboBoxValueSelected::CreateStatic(&FMaterialPropertyHelpers::SetVectorChannelMaskValue, ParameterValueProperty, Parameter, (UObject*)MaterialEditorParameters);
 
 	// Widget replaces color picker with combo box
-	PropertyRow.CustomWidget()
+	const bool bShowChildren = true;
+	PropertyRow.CustomWidget(bShowChildren)
 		.NameContent()
 		[
 			ParameterValueProperty->CreatePropertyNameWidget()
