Substrate - update widget to show pass through behavior instead of error.

#tests PC editor
#rnx
#DisabledValidations nonUnityAndStaticAnalysis

diff --git a/Engine/Source/Editor/MaterialEditor/Private/SMaterialEditorSubstrateWidget.cpp b/Engine/Source/Editor/MaterialEditor/Private/SMaterialEditorSubstrateWidget.cpp
--- a/Engine/Source/Editor/MaterialEditor/Private/SMaterialEditorSubstrateWidget.cpp
+++ b/Engine/Source/Editor/MaterialEditor/Private/SMaterialEditorSubstrateWidget.cpp
@@ -708,7 +708,16 @@ void SMaterialEditorSubstrateWidget::Tick(const FGeometry& AllottedGeometry, con
 				MaterialDescription += CompilationOutput.SubstrateMaterialDescription;
 
 				// Now generate a visual representation of the material from the topology tree of operators.
-				if (CompilationOutput.RootOperatorIndex >= 0)
+				
+				if (CompilationOutput.bMaterialUsesLegacyGBufferDataPassThrough)
+				{
+					auto TreeError = SNew(SErrorText)
+						.ErrorText(LOCTEXT("NoSubstrateTree", "No Substrate Tree - Pass Through Legacy GBuffer Data"))
+						.BackgroundColor(FSlateColor(EStyleColor::AccentGreen));
+					const TSharedRef<SWidget>& TreeErrorAsShared = TreeError->AsShared();
+					MaterialBox->SetContent(TreeErrorAsShared);
+				}
+				else if (CompilationOutput.RootOperatorIndex >= 0)
 				{
 					MaterialBox->SetContent(FSubstrateWidget::ProcessOperator(CompilationOutput));
 				}
diff --git a/Engine/Source/Runtime/Core/Private/UObject/FortniteMainBranchObjectVersion.cpp b/Engine/Source/Runtime/Core/Private/UObject/FortniteMainBranchObjectVersion.cpp
--- a/Engine/Source/Runtime/Core/Private/UObject/FortniteMainBranchObjectVersion.cpp
+++ b/Engine/Source/Runtime/Core/Private/UObject/FortniteMainBranchObjectVersion.cpp
@@ -13,7 +13,7 @@ TMap<FGuid, FGuid> FFortniteMainBranchObjectVersion::GetSystemGuids()
 	SystemGuids.Add(DevGuids.MATERIALSHADERMAP_DERIVEDDATA_VER, FGuid("F051877974B44D409435F321BDB39DBB"));
 	SystemGuids.Add(DevGuids.NIAGARASHADERMAP_DERIVEDDATA_VER, FGuid("6360A977062842A29ED30E3A7ACB0E64"));
 	SystemGuids.Add(DevGuids.Niagara_LatestScriptCompileVersion, FGuid("783FF3700BB641C799B2D43552FBF91E"));
-	SystemGuids.Add(DevGuids.MaterialTranslationDDCVersion, FGuid("287723BF015541FB90E91D2F0AF0AD5E"));
+	SystemGuids.Add(DevGuids.MaterialTranslationDDCVersion, FGuid("6DFF2413E72B44219C54F748928F5280"));
 
 	// These GUIDs are in AppendRestrictedFortniteMainStreamObjectVersionGuids as they should rarely be changed.
 	//SystemGuids.Add(DevGuids.NANITE_DERIVEDDATA_VER, FGuid("ACF98184978742A592537DA223A1E6BE"));
diff --git a/Engine/Source/Runtime/Engine/Private/Materials/HLSLMaterialTranslator.cpp b/Engine/Source/Runtime/Engine/Private/Materials/HLSLMaterialTranslator.cpp
--- a/Engine/Source/Runtime/Engine/Private/Materials/HLSLMaterialTranslator.cpp
+++ b/Engine/Source/Runtime/Engine/Private/Materials/HLSLMaterialTranslator.cpp
@@ -1487,6 +1487,7 @@ void FHLSLMaterialTranslator::TranslateMaterial()
 		FMaterialShadingModelField MaterialShadingModels = GetCompiledShadingModels();
 		MaterialCompilationOutput.SubstrateMaterialCompilationOutput.SubstrateClosureCount = 1;
 		MaterialCompilationOutput.SubstrateMaterialCompilationOutput.SubstrateUintPerPixel = 0;
+		MaterialCompilationOutput.SubstrateMaterialCompilationOutput.bMaterialUsesLegacyGBufferDataPassThrough = 1;
 		if (MaterialShadingModels.HasOnlyShadingModel(MSM_DefaultLit))
 		{
 			MaterialCompilationOutput.SubstrateMaterialCompilationOutput.SubstrateMaterialType = SUBSTRATE_MATERIAL_TYPE_SIMPLE;
diff --git a/Engine/Source/Runtime/Engine/Public/Rendering/SubstrateMaterialShared.cpp b/Engine/Source/Runtime/Engine/Public/Rendering/SubstrateMaterialShared.cpp
--- a/Engine/Source/Runtime/Engine/Public/Rendering/SubstrateMaterialShared.cpp
+++ b/Engine/Source/Runtime/Engine/Public/Rendering/SubstrateMaterialShared.cpp
@@ -22,6 +22,7 @@ FSubstrateMaterialCompilationOutput::FSubstrateMaterialCompilationOutput()
 	, bIsThin(0)
 	, MaterialType(0)
 	, bMaterialOutOfBudgetHasBeenSimplified(0)
+	, bMaterialUsesLegacyGBufferDataPassThrough(0)
 	, RootOperatorIndex(0)
 #endif
 {
diff --git a/Engine/Source/Runtime/Engine/Public/Rendering/SubstrateMaterialShared.h b/Engine/Source/Runtime/Engine/Public/Rendering/SubstrateMaterialShared.h
--- a/Engine/Source/Runtime/Engine/Public/Rendering/SubstrateMaterialShared.h
+++ b/Engine/Source/Runtime/Engine/Public/Rendering/SubstrateMaterialShared.h
@@ -225,6 +225,8 @@ public:
 
 	/** The byte per pixel count supported by the platform the material has been compiled against */
 	LAYOUT_FIELD_EDITORONLY(uint8, bMaterialOutOfBudgetHasBeenSimplified, 1);
+	/** True if the blendable GBuffer is used by Substrate and the Front material pin is not plugged in: in this case we can passthough substrate and use legacy base pass material shader.*/
+	LAYOUT_FIELD_EDITORONLY(uint8, bMaterialUsesLegacyGBufferDataPassThrough, 1);
 
 	LAYOUT_FIELD_EDITORONLY(int8, RootOperatorIndex);
 	LAYOUT_ARRAY_EDITORONLY(FSubstrateOperator, Operators, SUBSTRATE_COMPILATION_OUTPUT_MAX_OPERATOR);
