Added some checks to FMaterialEditor::PostUndo() to get more information out of random Undo crashes occurring.

#rnx

diff --git a/Engine/Source/Editor/MaterialEditor/Private/MaterialEditor.cpp b/Engine/Source/Editor/MaterialEditor/Private/MaterialEditor.cpp
--- a/Engine/Source/Editor/MaterialEditor/Private/MaterialEditor.cpp
+++ b/Engine/Source/Editor/MaterialEditor/Private/MaterialEditor.cpp
@@ -7195,6 +7195,8 @@ void FMaterialEditor::PostUndo(bool bSuccess)
 {
 	if (bSuccess)
 	{	
+		checkf(Material, TEXT("Unknown material during material Undo operation"));
+
 		if (TSharedPtr<SGraphEditor> FocusedGraphEd = FocusedGraphEdPtr.Pin())
 		{
 			FocusedGraphEd->ClearSelectionSet();
@@ -7226,10 +7228,14 @@ void FMaterialEditor::PostUndo(bool bSuccess)
 
 		if (TSharedPtr<SGraphEditor> FocusedGraphEd = FocusedGraphEdPtr.Pin())
 		{
+			checkf(FocusedGraphEd->GetCurrentGraph(), TEXT("Null FocusedGraphEd->GetCurrentGraph() during Undo operation in material '%s'"), *Material->GetName());
+
 			// @TODO Find a more coherent check for this, rather than rely on sage knowledge that the schema won't exist.
 			// If our previous transaction created the current graph, it won't have a schema after undo.
 			if (!FocusedGraphEd->GetCurrentGraph()->GetSchema())
 			{
+				checkf(DocumentManager, TEXT("Null DocumentManager during Undo operation in material '%s'"), *Material->GetName());
+
 				CloseDocumentTab(FocusedGraphEd->GetCurrentGraph());
 				DocumentManager->CleanInvalidTabs();
 				DocumentManager->RefreshAllTabs();
@@ -7237,7 +7243,7 @@ void FMaterialEditor::PostUndo(bool bSuccess)
 			}
 
 			// Active graph can change above, re-acquire ptr
-			FocusedGraphEdPtr.Pin()->NotifyGraphChanged();
+			FocusedGraphEd->NotifyGraphChanged();
 		}
 
 		SetEditorInstanceForWidgets();
