- Added support for specifying an array of valid values to SGraphEnumPin through a new 'ValidEnumValues' attribute.
- MaterialExpressions with enum properties marked as ShowAsInputPin will now pick up the list of permitted values in the ValidEnumValues metadata, if specified. This makes the combo box displayed on the left in the details panel identical to that displayed on the expression node itself and avoids possible material breakages.

[RN] Fixed ShadingModel material expression in-node combobox to display only valid shading model entries.
#tests Node combobox displays the same values as those in the details panel combo box, and the two are in sync.

diff --git a/Engine/Source/Editor/GraphEditor/Private/KismetPins/SGraphPinEnum.cpp b/Engine/Source/Editor/GraphEditor/Private/KismetPins/SGraphPinEnum.cpp
--- a/Engine/Source/Editor/GraphEditor/Private/KismetPins/SGraphPinEnum.cpp
+++ b/Engine/Source/Editor/GraphEditor/Private/KismetPins/SGraphPinEnum.cpp
@@ -142,6 +142,7 @@ TSharedRef<ITableRow> SPinComboBox::OnGenerateComboWidget( TSharedPtr<int32> InC
 
 void SGraphPinEnum::Construct(const FArguments& InArgs, UEdGraphPin* InGraphPinObj)
 {
+	ValidEnumValues = InArgs._ValidEnumValues;
 	SGraphPin::Construct(SGraphPin::FArguments(), InGraphPinObj);
 }
 
@@ -258,12 +259,11 @@ void SGraphPinEnum::GenerateComboBoxIndexes( TArray< TSharedPtr<int32> >& OutCom
 	UEnum* EnumPtr = Cast<UEnum>( GraphPinObj->PinType.PinSubCategoryObject.Get() );
 	if (EnumPtr)
 	{
-
 		//NumEnums() - 1, because the last item in an enum is the _MAX item
 		for (int32 EnumIndex = 0; EnumIndex < EnumPtr->NumEnums() - 1; ++EnumIndex)
 		{
 			// Ignore hidden enum values
-			if( !EnumPtr->HasMetaData(TEXT("Hidden"), EnumIndex ) )
+			if( !EnumPtr->HasMetaData(TEXT("Hidden"), EnumIndex ) && (ValidEnumValues.IsEmpty() || ValidEnumValues.Contains(EnumPtr->GetNameByIndex(EnumIndex))))
 			{
 				TSharedPtr<int32> EnumIdxPtr(new int32(EnumIndex));
 				OutComboBoxIndexes.Add(EnumIdxPtr);
diff --git a/Engine/Source/Editor/GraphEditor/Public/KismetPins/SGraphPinEnum.h b/Engine/Source/Editor/GraphEditor/Public/KismetPins/SGraphPinEnum.h
--- a/Engine/Source/Editor/GraphEditor/Public/KismetPins/SGraphPinEnum.h
+++ b/Engine/Source/Editor/GraphEditor/Public/KismetPins/SGraphPinEnum.h
@@ -20,6 +20,10 @@ class SGraphPinEnum : public SGraphPin
 {
 public:
 	SLATE_BEGIN_ARGS(SGraphPinEnum) {}
+
+	/** The array of values to allow in the combobox. All values permitted if this array is empty. */
+	SLATE_ARGUMENT(TArray<FString>, ValidEnumValues)
+
 	SLATE_END_ARGS()
 
 	UE_API void Construct(const FArguments& InArgs, UEdGraphPin* InGraphPinObj);
@@ -71,7 +75,11 @@ private:
 	 */
 	UE_API FText OnGetTooltip(int32 EnumIndex);
 
-	TSharedPtr<class SPinComboBox>	ComboBox;
+	/** The array of values to allow in the combobox. All values permitted if this array is empty. */
+	TArray<FString> ValidEnumValues;
+
+	/** Shared pointer to the combobox widget. */
+	TSharedPtr<class SPinComboBox> ComboBox;
 };
 
 #undef UE_API
diff --git a/Engine/Source/Editor/MaterialEditor/Private/MaterialEditorGraphPanelPinFactory.cpp b/Engine/Source/Editor/MaterialEditor/Private/MaterialEditorGraphPanelPinFactory.cpp
--- a/Engine/Source/Editor/MaterialEditor/Private/MaterialEditorGraphPanelPinFactory.cpp
+++ b/Engine/Source/Editor/MaterialEditor/Private/MaterialEditorGraphPanelPinFactory.cpp
@@ -84,7 +84,33 @@ TSharedPtr<class SGraphPin> FMaterialEditorGraphPanelPinFactory::CreatePin(class
 				// Check for valid enum object reference
 				if ((InPin->PinType.PinSubCategoryObject != NULL) && (InPin->PinType.PinSubCategoryObject->IsA(UEnum::StaticClass())))
 				{
-					return SNew(SGraphPinEnum, InPin);
+					// The property may have a ValidEnumValues metadata to filter out invalid values. This array contains the
+					// extracted valid values if any, or remains empty if the ValidEnumValues metadata was not specified,
+					// which will cause no values to be filtered out.
+					TArray<FString> ValidEnumValues;
+
+					// We know the property type is an enum so it is safe to cast here.
+					static const FName ValidEnumValuesName("ValidEnumValues");
+
+					// If the property has the ValidEnumValues metadata, it will be in the form of a string of elements separated by
+					// comma (with potentially space in between that needs trimming).
+					if (InProperty && InProperty->HasMetaData(ValidEnumValuesName))
+					{
+						// The property has the metadata. Parse the values into an array of strings.
+						TArray<FString> ValidEnumValuesAsString;
+						InProperty->GetMetaData(ValidEnumValuesName).ParseIntoArray(ValidEnumValuesAsString, TEXT(","));
+
+						// For each value, trim it and then push to ValidEnumValues the equivalent FName.
+						for(FString& Value : ValidEnumValuesAsString)
+						{
+							Value.TrimStartAndEndInline();
+							ValidEnumValues.Add(Value);
+						}
+					}
+
+					// Create the SGraphPinEnum with the allowed list of values (empty if no metadata was found).
+					return SNew(SGraphPinEnum, InPin)
+						.ValidEnumValues(ValidEnumValues);
 				}
 				else
 				{
