Minor single comment improvement to my last change.

#rnx

diff --git a/Engine/Source/Editor/MaterialEditor/Private/MaterialEditorGraphPanelPinFactory.cpp b/Engine/Source/Editor/MaterialEditor/Private/MaterialEditorGraphPanelPinFactory.cpp
--- a/Engine/Source/Editor/MaterialEditor/Private/MaterialEditorGraphPanelPinFactory.cpp
+++ b/Engine/Source/Editor/MaterialEditor/Private/MaterialEditorGraphPanelPinFactory.cpp
@@ -100,7 +100,7 @@ TSharedPtr<class SGraphPin> FMaterialEditorGraphPanelPinFactory::CreatePin(class
 						TArray<FString> ValidEnumValuesAsString;
 						InProperty->GetMetaData(ValidEnumValuesName).ParseIntoArray(ValidEnumValuesAsString, TEXT(","));
 
-						// For each value, trim it and then push to ValidEnumValues the equivalent FName.
+						// For each value, trim it and then push to ValidEnumValues.
 						for(FString& Value : ValidEnumValuesAsString)
 						{
 							Value.TrimStartAndEndInline();
