[Materials] Update cached expressions for preview material when stopping preview mode to fix reference crash on texture reimport.

#rnx

diff --git a/Engine/Source/Editor/MaterialEditor/Private/MaterialEditor.cpp b/Engine/Source/Editor/MaterialEditor/Private/MaterialEditor.cpp
--- a/Engine/Source/Editor/MaterialEditor/Private/MaterialEditor.cpp
+++ b/Engine/Source/Editor/MaterialEditor/Private/MaterialEditor.cpp
@@ -2564,10 +2564,13 @@ void FMaterialEditor::UpdatePreviewMaterial( bool bForce )
 
 	bStatsFromPreviewMaterial = true;
 
-	if( PreviewExpression && ExpressionPreviewMaterial )
+	if (ExpressionPreviewMaterial)
 	{
 		ExpressionPreviewMaterial->UpdateCachedExpressionData();
-		PreviewExpression->ConnectToPreviewMaterial(ExpressionPreviewMaterial,0);
+		if (PreviewExpression)
+		{
+			PreviewExpression->ConnectToPreviewMaterial(ExpressionPreviewMaterial, 0);
+		}
 	}
 
 	if(PreviewExpression)
