PR #10326: Inconsistency fixed on Comment Bubbles on Material Editor

To ensure visual consistency new Comment bubbles created use the same color value as the other node editors, like Blueprints.

#10326
[RN] minor | Rendering - Materials and Shaders

diff --git a/Engine/Source/Editor/MaterialEditor/Private/MaterialEditor.cpp b/Engine/Source/Editor/MaterialEditor/Private/MaterialEditor.cpp
--- a/Engine/Source/Editor/MaterialEditor/Private/MaterialEditor.cpp
+++ b/Engine/Source/Editor/MaterialEditor/Private/MaterialEditor.cpp
@@ -150,6 +150,7 @@
 #include "Engine/Selection.h"
 #include "Materials/Material.h"
 #include "AdvancedPreviewSceneModule.h"
+#include "GraphEditorSettings.h"
 #include "MaterialLayersFunctionsCustomization.h"
 #include "MaterialEditor/MaterialEditorPreviewParameters.h"
 #include "SMaterialLayersFunctionsTree.h"
@@ -5981,6 +5982,7 @@ UMaterialExpressionComment* FMaterialEditor::CreateNewMaterialExpressionComment(
 		}
 
 		NewComment->Text = NSLOCTEXT("K2Node", "CommentBlock_NewEmptyComment", "Comment").ToString();
+		NewComment->CommentColor = GetDefault<UGraphEditorSettings>()->DefaultCommentNodeTitleColor;
 
 		ExpressionGraph->AddComment(NewComment, true);
 
diff --git a/Engine/Source/Editor/UnrealEd/Private/EdGraphNode_Comment.cpp b/Engine/Source/Editor/UnrealEd/Private/EdGraphNode_Comment.cpp
--- a/Engine/Source/Editor/UnrealEd/Private/EdGraphNode_Comment.cpp
+++ b/Engine/Source/Editor/UnrealEd/Private/EdGraphNode_Comment.cpp
@@ -49,7 +49,7 @@ UEdGraphNode_Comment::UEdGraphNode_Comment(const FObjectInitializer& ObjectIniti
 	NodeWidth = 400;
 	NodeHeight = 100;
 	FontSize = 18;
-	CommentColor = FLinearColor::White;
+	CommentColor = GetDefault<UGraphEditorSettings>()->DefaultCommentNodeTitleColor;
 	bColorCommentBubble = false;
 	MoveMode = ECommentBoxMode::GroupMovement;
 
