Added checks and sanitization to MaterialEditor.cpp to prevent null expression from appearing in the material expressions list until we understand what is causing them in the first place.

In an effort to track a bug in the material editor that occasionally causes null expressions to appear in the list of expressions in materials cached expression data, this change introduces sanitization and error reporting to operations performed by FMaterialEditor, in an effort to avoid crashes and give information of what operation introduced this erroneous state.

#rnx

diff --git a/Engine/Source/Editor/MaterialEditor/Private/MaterialEditor.cpp b/Engine/Source/Editor/MaterialEditor/Private/MaterialEditor.cpp
--- a/Engine/Source/Editor/MaterialEditor/Private/MaterialEditor.cpp
+++ b/Engine/Source/Editor/MaterialEditor/Private/MaterialEditor.cpp
@@ -3922,8 +3922,9 @@ void FMaterialEditor::OnPromoteObjects()
 	if (SelectedNodes.Num() > 0)
 	{
 		const FScopedTransaction Transaction(LOCTEXT("MaterialEditorPromote", "Material Editor: Promote"));
+		FScopedMaterialChecker Checker{ this, __FUNCTION__ };
 		ModifyMaterial();
-		Material->MaterialGraph->Modify();
+		
 		TArray<class UEdGraphNode*> NodesToDelete;
 		TArray<class UEdGraphNode*> NodesToSelect;
 
@@ -4014,8 +4015,9 @@ void FMaterialEditor::OnConvertObjects()
 	if (SelectedNodes.Num() > 0)
 	{
 		const FScopedTransaction Transaction( LOCTEXT("MaterialEditorConvert", "Material Editor: Convert") );
+		FScopedMaterialChecker Checker{ this, __FUNCTION__ };
 		ModifyMaterial();
-		Material->MaterialGraph->Modify();
+
 		TArray<class UEdGraphNode*> NodesToDelete;
 		TArray<class UEdGraphNode*> NodesToSelect;
 
@@ -4268,8 +4270,8 @@ void FMaterialEditor::OnConvertTextures()
 	if (SelectedNodes.Num() > 0)
 	{
 		const FScopedTransaction Transaction( LOCTEXT("MaterialEditorConvertTexture", "Material Editor: Convert to Texture") );
+		FScopedMaterialChecker Checker{ this, __FUNCTION__ };
 		ModifyMaterial();
-		Material->MaterialGraph->Modify();
 		TArray<class UEdGraphNode*> NodesToDelete;
 		TArray<class UEdGraphNode*> NodesToSelect;
 
@@ -5836,6 +5838,7 @@ UMaterialExpression* FMaterialEditor::CreateNewMaterialExpression(UClass* NewExp
 	UMaterialExpression* NewExpression = NULL;
 	{
 		const FScopedTransaction Transaction( NSLOCTEXT("UnrealEd", "MaterialEditorNewExpression", "Material Editor: New Expression") );
+		FScopedMaterialChecker Checker{ this, __FUNCTION__ };
 		ModifyMaterial();
 
 		UObject* SelectedAsset = nullptr;
@@ -5888,6 +5891,8 @@ UMaterialExpressionComposite* FMaterialEditor::CreateNewMaterialExpressionCompos
 	UMaterialExpressionComposite* NewComposite = nullptr;
 	{
 		const FScopedTransaction Transaction(NSLOCTEXT("UnrealEd", "MaterialEditorNewComposite", "Material Editor: New Composite"));
+		FScopedMaterialChecker Checker{ this, __FUNCTION__ };
+
 		ModifyMaterial();
 
 		UObject* ExpressionOuter = Material;
@@ -5949,6 +5954,7 @@ UMaterialExpressionComment* FMaterialEditor::CreateNewMaterialExpressionComment(
 	UMaterialExpressionComment* NewComment = NULL;
 	{
 		const FScopedTransaction Transaction(NSLOCTEXT("UnrealEd", "MaterialEditorCreateComment", "Material Editor: Create comment"));
+		FScopedMaterialChecker Checker{ this, __FUNCTION__ };
 		ModifyMaterial();
 
 		UObject* ExpressionOuter = Material;
@@ -6225,6 +6231,7 @@ void FMaterialEditor::DeleteSelectedDuplicatableNodes()
 
 void FMaterialEditor::DeleteNodesInternal(const TArray<class UEdGraphNode*>& NodesToDelete, bool& bHaveExpressionsToDelete, bool& bPreviewExpressionDeleted, bool bForceDeletion)
 {
+	FScopedMaterialChecker Checker{ this, __FUNCTION__ };
 	ModifyMaterial();
 
 	for (int32 Index = 0; Index < NodesToDelete.Num(); ++Index)
@@ -6429,13 +6436,14 @@ void FMaterialEditor::PasteNodesHere(const FVector2D& Location, const class UEdG
 	
 	PasteNodesHereFromBuffer(Location, Graph, TextToImport, nullptr);
 }
+
 void FMaterialEditor::PasteNodesHereFromBuffer(const FVector2D& Location, const class UEdGraph* Graph, const FString& TextToImport, TMap<FGuid, FGuid>* OutOldToNewGuids)
 {
 	// Undo/Redo support
 	const FScopedTransaction Transaction( NSLOCTEXT("UnrealEd", "MaterialEditorPaste", "Material Editor: Paste") );
-	Material->MaterialGraph->Modify();
+	FScopedMaterialChecker Checker{ this, __FUNCTION__ };
 	ModifyMaterial();
-
+	
 	UMaterialGraph* ExpressionGraph = Graph ? ToRawPtr(CastChecked<UMaterialGraph>(const_cast<UEdGraph*>(Graph))) : ToRawPtr(Material->MaterialGraph);
 	ExpressionGraph->Modify();
 
@@ -6965,8 +6973,8 @@ void FMaterialEditor::OnCollapseNodes()
 	if (CollapsableNodes.Num())
 	{
 		const FScopedTransaction Transaction(FGraphEditorCommands::Get().CollapseNodes->GetDescription());
+		FScopedMaterialChecker Checker{ this, __FUNCTION__ };
 		ModifyMaterial();
-		Material->MaterialGraph->Modify();
 
 		CollapseNodes(CollapsableNodes);
 	}
@@ -6994,8 +7002,8 @@ bool FMaterialEditor::CanCollapseNodes() const
 void FMaterialEditor::OnExpandNodes()
 {
 	const FScopedTransaction Transaction(FGraphEditorCommands::Get().ExpandNodes->GetLabel());
+	FScopedMaterialChecker Checker{ this, __FUNCTION__ };
 	ModifyMaterial();
-	Material->MaterialGraph->Modify();
 
 	TSet<UEdGraphNode*> ExpandedNodes;
 	TSharedPtr<SGraphEditor> FocusedGraphEd = FocusedGraphEdPtr.Pin();
@@ -7201,6 +7209,7 @@ void FMaterialEditor::PostUndo(bool bSuccess)
 	if (bSuccess)
 	{	
 		checkf(Material, TEXT("Unknown material during material Undo operation"));
+		FScopedMaterialChecker Checker{ this, __FUNCTION__ };
 
 		if (TSharedPtr<SGraphEditor> FocusedGraphEd = FocusedGraphEdPtr.Pin())
 		{
@@ -7408,6 +7417,16 @@ void FMaterialEditor::RefreshExpressionPreviews(bool bForceRefreshAll /*= false*
 	const FScopedBusyCursor BusyCursor;
 
 	Material->UpdateCachedExpressionData();
+	
+	// @todo-massimo.tristano
+	// We have a bug causing nulls to end up in the expressions array, but we are not yet aware of
+	// a repro to track down what actually causes it. This temporary code not only sanitizes any
+	// null expression in the array (that would crash the editor) but also allows us to know when
+	// null expressions are found.
+	if (Material->GetEditorOnlyData()->ExpressionCollection.Expressions.Remove(nullptr) > 0)
+	{
+		UE_LOG(LogMaterialEditor, Warning, TEXT("Null expressions found in expressions list during expression previews refresh in material '%s'."), *GetEditingObjects()[0]->GetName());
+	}
 
 	if ( bAlwaysRefreshAllPreviews || bForceRefreshAll)
 	{
@@ -7438,7 +7457,7 @@ void FMaterialEditor::RefreshExpressionPreviews(bool bForceRefreshAll /*= false*
 	// Go through all expression previews and create new ones as needed, and maintain a list of previews that are being compiled
 	for (UMaterialExpression* MaterialExpression : Material->GetExpressions())
 	{
-		if (MaterialExpression && !MaterialExpression->IsA(UMaterialExpressionComment::StaticClass()) )
+		if (!MaterialExpression->IsA(UMaterialExpressionComment::StaticClass()))
 		{
 			bool bNewlyCreated;
 			FMatExpressionPreview* Preview = GetExpressionPreview( MaterialExpression, bNewlyCreated );
@@ -7452,12 +7471,14 @@ void FMaterialEditor::RefreshExpressionPreviews(bool bForceRefreshAll /*= false*
 
 void FMaterialEditor::RefreshExpressionPreview(UMaterialExpression* MaterialExpression, bool bRecompile)
 {
+	check(MaterialExpression);
+
 	if ( (MaterialExpression->bRealtimePreview || MaterialExpression->bNeedToUpdatePreview) && !MaterialExpression->bCollapsed )
 	{
 		for( int32 PreviewIndex = 0 ; PreviewIndex < ExpressionPreviews.Num() ; ++PreviewIndex )
 		{
 			FMatExpressionPreview* ExpressionPreview = ExpressionPreviews[PreviewIndex];
-			if( ExpressionPreview->GetExpression() == MaterialExpression )
+			if (ensure(ExpressionPreview) && ExpressionPreview->GetExpression() == MaterialExpression)
 			{
 				// We need to make sure we don't hold any other references before queuing the deferred delete as it may execute immediately (e.g. -onethread, or just very fast RT).
 				// There's no danger in releasing the last reference to FMaterial as this doesn't delete it anyway, it needs to be deleted manually via DeferredDelete.
@@ -7945,8 +7966,8 @@ void FMaterialEditor::CleanUnusedExpressions()
 			// Kill off expressions referenced by the material that aren't reachable.
 			const FScopedTransaction Transaction( NSLOCTEXT("UnrealEd", "MaterialEditorCleanUnusedExpressions", "Material Editor: Clean Unused Expressions") );
 				
+			FScopedMaterialChecker Checker{ this, __FUNCTION__ };
 			ModifyMaterial();
-			Material->MaterialGraph->Modify();
 
 			for (int32 Index = 0; Index < UnusedNodes.Num(); ++Index)
 			{
@@ -8009,6 +8030,7 @@ void FMaterialEditor::FixupMismatchedSamplerExpressions()
 		const FScopedTransaction Transaction(NSLOCTEXT("UnrealEd", "MaterialEditorFixupSamplerExpressions", "Material Editor: Fixup Sampler Expressions"));
 			
 		{
+			FScopedMaterialChecker Checker{ this, __FUNCTION__ };
 			ModifyMaterial();
 
 			for (UMaterialExpressionTextureBase* TextureBase : ExpressionsToFix)
@@ -8532,4 +8554,20 @@ void FMaterialEditor::RebuildInheritanceList()
 	}
 }
 
+FMaterialEditor::FScopedMaterialChecker::FScopedMaterialChecker(FMaterialEditor* InEditor, const char* InDebugEditName)
+	: Editor{ InEditor }
+	, DebugEditName{ InDebugEditName }
+{
+}
+
+FMaterialEditor::FScopedMaterialChecker::~FScopedMaterialChecker()
+{
+	if (Editor->Material->GetEditorOnlyData()->ExpressionCollection.Expressions.Remove(nullptr) > 0)
+	{
+		UE_LOG(LogMaterialEditor, Warning, 
+			TEXT("Material Editor internal error after operation '%hs': null expressions found in material '%s' expressions list."), 
+			DebugEditName, *Editor->Material->GetName());
+	}
+}
+
 #undef LOCTEXT_NAMESPACE
diff --git a/Engine/Source/Editor/MaterialEditor/Private/MaterialEditor.h b/Engine/Source/Editor/MaterialEditor/Private/MaterialEditor.h
--- a/Engine/Source/Editor/MaterialEditor/Private/MaterialEditor.h
+++ b/Engine/Source/Editor/MaterialEditor/Private/MaterialEditor.h
@@ -488,6 +488,25 @@ public:
 
 	/** Factory that spawns graph editors; used to look up all tabs spawned by it. */
 	TWeakPtr<FDocumentTabFactory> GraphEditorTabFactoryPtr;
+	
+protected:
+	/**
+	 * Helper struct used to validate the state of the material after the material editor
+	 * performs an operation, searching for any invalid state.
+	 */
+	struct FScopedMaterialChecker
+	{
+		UE_NONCOPYABLE(FScopedMaterialChecker);
+
+		/** The material editor performing the operation */
+		FMaterialEditor* Editor;
+
+		/** Name of the operation being performed, logged when an issue is reported */
+		const char* DebugEditName;
+		
+		FScopedMaterialChecker(FMaterialEditor* Editor, const char* DebugEditName);
+		~FScopedMaterialChecker();
+	};
 
 protected:
 	//~ FAssetEditorToolkit interface
diff --git a/Engine/Source/Runtime/Engine/Private/Materials/MaterialExpressions.cpp b/Engine/Source/Runtime/Engine/Private/Materials/MaterialExpressions.cpp
--- a/Engine/Source/Runtime/Engine/Private/Materials/MaterialExpressions.cpp
+++ b/Engine/Source/Runtime/Engine/Private/Materials/MaterialExpressions.cpp
@@ -11574,6 +11574,11 @@ void UMaterialExpressionPinBase::PostEditChangeProperty(FPropertyChangedEvent& P
 
 	PreEditRereouteExpresions.Empty();
 	Super::PostEditChangeProperty(PropertyChangedEvent);
+	
+	if (Material->GetExpressionCollection().Expressions.Remove(nullptr))
+	{
+		UE_LOG(LogMaterial, Warning, TEXT("Null expressions detected in material expressions list after execution of UMaterialExpressionPinBase::PostEditChangeProperty."));
+	}
 }
 
 TArray<FExpressionOutput>& UMaterialExpressionPinBase::GetOutputs()
