Add cpu scope markers to hot functions when modifying materials and material instances w/ many material parameters.

- Instruments and covers most of UMaterialEditingLibrary::UpdateMaterialFunction

#rnx

diff --git a/Engine/Source/Editor/MaterialEditor/Private/MaterialEditingLibrary.cpp b/Engine/Source/Editor/MaterialEditor/Private/MaterialEditingLibrary.cpp
--- a/Engine/Source/Editor/MaterialEditor/Private/MaterialEditingLibrary.cpp
+++ b/Engine/Source/Editor/MaterialEditor/Private/MaterialEditingLibrary.cpp
@@ -298,6 +298,8 @@ namespace MaterialEditingLibraryImpl
 
 void UMaterialEditingLibrary::RebuildMaterialInstanceEditors(UMaterial* BaseMaterial)
 {
+	TRACE_CPUPROFILER_EVENT_SCOPE(UMaterialEditingLibrary::RebuildMaterialInstanceEditors)
+
 	UAssetEditorSubsystem* AssetEditorSubsystem = GEditor->GetEditorSubsystem<UAssetEditorSubsystem>();
 	TArray<UObject*> EditedAssets = AssetEditorSubsystem->GetAllEditedAssets();
 
diff --git a/Engine/Source/Editor/MaterialEditor/Private/MaterialInstanceEditor.cpp b/Engine/Source/Editor/MaterialEditor/Private/MaterialInstanceEditor.cpp
--- a/Engine/Source/Editor/MaterialEditor/Private/MaterialInstanceEditor.cpp
+++ b/Engine/Source/Editor/MaterialEditor/Private/MaterialInstanceEditor.cpp
@@ -1408,6 +1408,8 @@ void FMaterialInstanceEditor::RebuildMaterialInstanceEditor()
 {
 	if( MaterialEditorInstance )
 	{
+		TRACE_CPUPROFILER_EVENT_SCOPE(UMaterialEditingLibrary::RebuildMaterialInstanceEditors)
+
 		ReInitMaterialFunctionProxies();
 		MaterialEditorInstance->CopyBasePropertiesFromParent();
 		MaterialEditorInstance->RegenerateArrays();
@@ -1447,6 +1449,8 @@ void FMaterialInstanceEditor::DrawMessages( FViewport* Viewport, FCanvas* Canvas
 
 void FMaterialInstanceEditor::RefreshOnScreenTextureWarnings()
 {
+	TRACE_CPUPROFILER_EVENT_SCOPE(FMaterialInstanceEditor::RefreshOnScreenTextureWarnings)
+
 	OnScreenMessages.Reset();
 
 	const bool bValidateMaterials = CVarEnableTextureValidationInMaterialInstanceEditor.GetValueOnAnyThread();
