27% speedup when applying/saving changes to material functions.

- Savings changes to a Material Function or Material while the Material Instance editor was open could be extremely slow depending on how many parameters it has.
- 34.6 seconds -> 25 seconds for Material/MaterialInstances w/ many parameters.
- We were finding and updating the _same_ MaterialInstanceEditor multiple times.  This is due to the code in `FMaterialInstanceEditor::InitMaterialInstanceEditor` where we add register 2 Material Instances as being edited.  The Material Instance itself, and a copy instance for the parameters.
- Instead of changing the logic that registers multiple objects to edit I just make a Set of FMaterialInstanceEditors to update, and update them once.

```
TArray<UObject*> ObjectsToEdit;
ObjectsToEdit.Add(ObjectToEdit);
ObjectsToEdit.Add(MaterialEditorInstance);
```

todo: cut down on the code duplication between the two implementations of UMaterialEditingLibrary::RebuildMaterialInstanceEditors

[RN] minor | Rendering - Materials and Shaders

diff --git a/Engine/Source/Editor/MaterialEditor/Private/MaterialEditingLibrary.cpp b/Engine/Source/Editor/MaterialEditor/Private/MaterialEditingLibrary.cpp
--- a/Engine/Source/Editor/MaterialEditor/Private/MaterialEditingLibrary.cpp
+++ b/Engine/Source/Editor/MaterialEditor/Private/MaterialEditingLibrary.cpp
@@ -303,6 +303,8 @@ void UMaterialEditingLibrary::RebuildMaterialInstanceEditors(UMaterial* BaseMate
 	UAssetEditorSubsystem* AssetEditorSubsystem = GEditor->GetEditorSubsystem<UAssetEditorSubsystem>();
 	TArray<UObject*> EditedAssets = AssetEditorSubsystem->GetAllEditedAssets();
 
+	TSet<FMaterialInstanceEditor*> InstanceEditorsToUpdate;
+
 	for (int32 AssetIdx = 0; AssetIdx < EditedAssets.Num(); AssetIdx++)
 	{
 		UObject* EditedAsset = EditedAssets[AssetIdx];
@@ -325,18 +327,28 @@ void UMaterialEditingLibrary::RebuildMaterialInstanceEditors(UMaterial* BaseMate
 			{
 				if (FMaterialInstanceEditor* MaterialInstanceEditor = MaterialEditingLibraryImpl::FindMaterialInstanceEditorForAsset(SourceInstance))
 				{
-					MaterialInstanceEditor->RebuildMaterialInstanceEditor();
+					InstanceEditorsToUpdate.Add(MaterialInstanceEditor);
 				}
 			}
 		}
 	}
+
+	// Update the unique list of instance editors.
+	for (FMaterialInstanceEditor* MaterialInstanceEditor : InstanceEditorsToUpdate)
+	{
+		MaterialInstanceEditor->RebuildMaterialInstanceEditor();
+	}
 }
 
 void UMaterialEditingLibrary::RebuildMaterialInstanceEditors(UMaterialFunction* BaseFunction)
 {
+	TRACE_CPUPROFILER_EVENT_SCOPE(UMaterialEditingLibrary::RebuildMaterialInstanceEditors)
+
 	UAssetEditorSubsystem* AssetEditorSubsystem = GEditor->GetEditorSubsystem<UAssetEditorSubsystem>();
 	TArray<UObject*> EditedAssets = AssetEditorSubsystem->GetAllEditedAssets();
 
+	TSet<FMaterialInstanceEditor*> InstanceEditorsToUpdate;
+
 	for (int32 AssetIdx = 0; AssetIdx < EditedAssets.Num(); AssetIdx++)
 	{
 		UObject* EditedAsset = EditedAssets[AssetIdx];
@@ -351,7 +363,7 @@ void UMaterialEditingLibrary::RebuildMaterialInstanceEditors(UMaterialFunction*
 			{
 				if (FMaterialInstanceEditor* MaterialInstanceEditor = MaterialEditingLibraryImpl::FindMaterialInstanceEditorForAsset(EditedAsset))
 				{
-					MaterialInstanceEditor->RebuildMaterialInstanceEditor();
+					InstanceEditorsToUpdate.Add(MaterialInstanceEditor);
 				}
 			}
 		}
@@ -377,12 +389,18 @@ void UMaterialEditingLibrary::RebuildMaterialInstanceEditors(UMaterialFunction*
 				{
 					if (FMaterialInstanceEditor* MaterialInstanceEditor = MaterialEditingLibraryImpl::FindMaterialInstanceEditorForAsset(EditedAsset))
 					{
-						MaterialInstanceEditor->RebuildMaterialInstanceEditor();
+						InstanceEditorsToUpdate.Add(MaterialInstanceEditor);
 					}
 				}
 			}
 		}
 	}
+
+	// Update the unique list of instance editors.
+	for (FMaterialInstanceEditor* MaterialInstanceEditor : InstanceEditorsToUpdate)
+	{
+		MaterialInstanceEditor->RebuildMaterialInstanceEditor();
+	}
 }
 
 int32 UMaterialEditingLibrary::GetNumMaterialExpressions(const UMaterial* Material)
diff --git a/Engine/Source/Editor/MaterialEditor/Private/SMaterialLayersFunctionsTree.cpp b/Engine/Source/Editor/MaterialEditor/Private/SMaterialLayersFunctionsTree.cpp
--- a/Engine/Source/Editor/MaterialEditor/Private/SMaterialLayersFunctionsTree.cpp
+++ b/Engine/Source/Editor/MaterialEditor/Private/SMaterialLayersFunctionsTree.cpp
@@ -1436,6 +1436,8 @@ void SMaterialLayersFunctionsInstanceTree::ShowSubParameters(TSharedPtr<FSortedP
 
 void SMaterialLayersFunctionsInstanceWrapper::Refresh()
 {
+	TRACE_CPUPROFILER_EVENT_SCOPE(SMaterialLayersFunctionsInstanceWrapper::Refresh)
+
 	LayerParameter.Reset();
 	TSharedPtr<SHorizontalBox> HeaderBox;
 	NestedTree->CreateGroupsWidget();
