44% speedup to UMaterialEditingLibrary::UpdateMaterialFunction() when applying/saving changes to material functions.

- 25 seconds -> 14 seconds for Material/MaterialInstances w/ many parameters.
- We were updating the material layers parameter panels (including rebuilding the parameter list redundantly, multiple times) even if they weren't visible and even if the material didn't have layers!
- Early out of SMaterialLayersFunctionsInstanceWrapper::Refresh if the parent material has no layers

[RN] minor | Rendering - Materials and Shaders

diff --git a/Engine/Source/Editor/MaterialEditor/Private/SMaterialLayersFunctionsTree.cpp b/Engine/Source/Editor/MaterialEditor/Private/SMaterialLayersFunctionsTree.cpp
--- a/Engine/Source/Editor/MaterialEditor/Private/SMaterialLayersFunctionsTree.cpp
+++ b/Engine/Source/Editor/MaterialEditor/Private/SMaterialLayersFunctionsTree.cpp
@@ -1439,121 +1439,130 @@ void SMaterialLayersFunctionsInstanceWrapper::Refresh()
 	TRACE_CPUPROFILER_EVENT_SCOPE(SMaterialLayersFunctionsInstanceWrapper::Refresh)
 
 	LayerParameter.Reset();
-	TSharedPtr<SHorizontalBox> HeaderBox;
-	NestedTree->CreateGroupsWidget();
-	LayerParameter = NestedTree->GetFunctionParameter();
-	FOnClicked 	OnChildButtonClicked = FOnClicked::CreateStatic(&FMaterialPropertyHelpers::OnClickedSaveNewMaterialInstance, MaterialEditorInstance->GetMaterialInterface(), ImplicitConv<UObject*>(MaterialEditorInstance));
-	FOnClicked	OnSiblingButtonClicked = FOnClicked::CreateStatic(&FMaterialPropertyHelpers::OnClickedSaveNewMaterialInstance, MaterialEditorInstance->GetParentMaterialInterface(), ImplicitConv<UObject*>(MaterialEditorInstance));
 
-	constexpr float ThumbnailSize = 48.0f;
-	TSharedRef<SWidget> ThumbnailWidget = SNullWidget::NullWidget;
-	TSharedPtr<SBox> ThumbnailBox = nullptr;
-	FText MaterialName = FText::GetEmpty();
-	
-	if (LayerParameter != nullptr)
+	UMaterialInterface* Material = MaterialEditorInstance->GetMaterialInterface();
+	UMaterialInterface* ParentMaterial = MaterialEditorInstance->GetParentMaterialInterface();
+
+	FMaterialLayersFunctions Layers;
+	const bool bHasLayers = ParentMaterial ? ParentMaterial->GetMaterialLayers(Layers) : Material->GetMaterialLayers(Layers);
+
+	if (bHasLayers)
 	{
-		
-		FOnClicked OnRelinkToParent = FOnClicked::CreateSPLambda(this, [this]()
-			{
-				return NestedTree->RelinkLayersToParent();
-			});
-		
-		HeaderBox = SNew(SHorizontalBox)
-			+ SHorizontalBox::Slot()
-			.Padding(FMargin(4.0f, 0.0f))
-			.HAlign(HAlign_Left)
-			.AutoWidth()
-			.VAlign(VAlign_Center)
-			[
-				SNew(STextBlock)
-				.Text(LOCTEXT("MaterialLayers", "Material Layers"))
-			];
+		TSharedPtr<SHorizontalBox> HeaderBox;
+		NestedTree->CreateGroupsWidget();
+		LayerParameter = NestedTree->GetFunctionParameter();
+		FOnClicked 	OnChildButtonClicked = FOnClicked::CreateStatic(&FMaterialPropertyHelpers::OnClickedSaveNewMaterialInstance, MaterialEditorInstance->GetMaterialInterface(), ImplicitConv<UObject*>(MaterialEditorInstance));
+		FOnClicked	OnSiblingButtonClicked = FOnClicked::CreateStatic(&FMaterialPropertyHelpers::OnClickedSaveNewMaterialInstance, MaterialEditorInstance->GetParentMaterialInterface(), ImplicitConv<UObject*>(MaterialEditorInstance));
 
-		this->ChildSlot
-			[
-				SNew(SVerticalBox)
-				+ SVerticalBox::Slot()
-				.Padding(FMargin(0.0f, 0.0f, 0.0f, 10.0f))
-				.AutoHeight()
-				[
-					HeaderBox.ToSharedRef()
-				]
-				+ SVerticalBox::Slot()
-				.Padding(FMargin(2.0f))
-				[
-					NestedTree.ToSharedRef()
-				]
-		];
-		if (FMaterialPropertyHelpers::IsOverriddenExpression(NestedTree->GetFunctionParameter()))
+		constexpr float ThumbnailSize = 48.0f;
+		TSharedRef<SWidget> ThumbnailWidget = SNullWidget::NullWidget;
+		TSharedPtr<SBox> ThumbnailBox = nullptr;
+		FText MaterialName = FText::GetEmpty();
+
+		if (LayerParameter != nullptr)
 		{
-			HeaderBox->AddSlot()
+
+			FOnClicked OnRelinkToParent = FOnClicked::CreateSPLambda(this, [this]()
+				{
+					return NestedTree->RelinkLayersToParent();
+				});
+
+			HeaderBox = SNew(SHorizontalBox)
+				+ SHorizontalBox::Slot()
+				.Padding(FMargin(4.0f, 0.0f))
 				.HAlign(HAlign_Left)
 				.AutoWidth()
 				.VAlign(VAlign_Center)
 				[
-				PropertyCustomizationHelpers::MakeAddButton(FSimpleDelegate::CreateSPLambda(this, [this]()
-					{
-						NestedTree->AddRootNodeLayer();
-					}))
+					SNew(STextBlock)
+						.Text(LOCTEXT("MaterialLayers", "Material Layers"))
 				];
-		}
 
-		HeaderBox->AddSlot()
-			.FillWidth(1.0f)
-			[
-				SNullWidget::NullWidget
-			];
-		HeaderBox->AddSlot()
-			.AutoWidth()
-			.Padding(2.0f)
-			[
-				SNew(SButton)
-				.Text(LOCTEXT("Relink", "Relink"))
-				.HAlign(HAlign_Center)
-				.OnClicked(OnRelinkToParent)
-				.ToolTipText(LOCTEXT("RelinkToParentLayers", "Relink to Parent Layers and Blends"))
-				.Visibility_Lambda([this]()
-				{
-					return NestedTree->GetRelinkLayersToParentVisibility();
-				})
-			];
-		HeaderBox->AddSlot()
+			this->ChildSlot
+				[
+					SNew(SVerticalBox)
+						+ SVerticalBox::Slot()
+						.Padding(FMargin(0.0f, 0.0f, 0.0f, 10.0f))
+						.AutoHeight()
+						[
+							HeaderBox.ToSharedRef()
+						]
+						+ SVerticalBox::Slot()
+						.Padding(FMargin(2.0f))
+						[
+							NestedTree.ToSharedRef()
+						]
+				];
+			if (FMaterialPropertyHelpers::IsOverriddenExpression(NestedTree->GetFunctionParameter()))
+			{
+				HeaderBox->AddSlot()
+					.HAlign(HAlign_Left)
+					.AutoWidth()
+					.VAlign(VAlign_Center)
+					[
+						PropertyCustomizationHelpers::MakeAddButton(FSimpleDelegate::CreateSPLambda(this, [this]()
+							{
+								NestedTree->AddRootNodeLayer();
+							}))
+					];
+			}
+
+			HeaderBox->AddSlot()
+				.FillWidth(1.0f)
+				[
+					SNullWidget::NullWidget
+				];
+			HeaderBox->AddSlot()
 				.AutoWidth()
 				.Padding(2.0f)
 				[
 					SNew(SButton)
-					.Text(LOCTEXT("SaveSibling", "Save Sibling"))
-					.HAlign(HAlign_Center)
-					.OnClicked(OnSiblingButtonClicked)
-					.ToolTipText(LOCTEXT("SaveToSiblingInstance", "Save to Sibling Instance"))
+						.Text(LOCTEXT("Relink", "Relink"))
+						.HAlign(HAlign_Center)
+						.OnClicked(OnRelinkToParent)
+						.ToolTipText(LOCTEXT("RelinkToParentLayers", "Relink to Parent Layers and Blends"))
+						.Visibility_Lambda([this]()
+							{
+								return NestedTree->GetRelinkLayersToParentVisibility();
+							})
 				];
-		HeaderBox->AddSlot()
-			.AutoWidth()
-			.Padding(2.0f)
-			[
-				SNew(SButton)
-				.Text(LOCTEXT("SaveChild", "Save Child"))
-				.HAlign(HAlign_Center)
-				.OnClicked(OnChildButtonClicked)
-				.ToolTipText(LOCTEXT("SaveToChildInstance", "Save To Child Instance"))
-			];
+			HeaderBox->AddSlot()
+				.AutoWidth()
+				.Padding(2.0f)
+				[
+					SNew(SButton)
+						.Text(LOCTEXT("SaveSibling", "Save Sibling"))
+						.HAlign(HAlign_Center)
+						.OnClicked(OnSiblingButtonClicked)
+						.ToolTipText(LOCTEXT("SaveToSiblingInstance", "Save to Sibling Instance"))
+				];
+			HeaderBox->AddSlot()
+				.AutoWidth()
+				.Padding(2.0f)
+				[
+					SNew(SButton)
+						.Text(LOCTEXT("SaveChild", "Save Child"))
+						.HAlign(HAlign_Center)
+						.OnClicked(OnChildButtonClicked)
+						.ToolTipText(LOCTEXT("SaveToChildInstance", "Save To Child Instance"))
+				];
+		}
 	}
 	else
 	{
 		this->ChildSlot
 			[
 				SNew(SBox)
-				.Padding(FMargin(10.0f))
-				.HAlign(HAlign_Center)
-				[
-					SNew(STextBlock)
-					.Text(LOCTEXT("AddLayerParameterPrompt", "Add a Material Attribute Layers parameter to see it here."))
-				]
+					.Padding(FMargin(10.0f))
+					.HAlign(HAlign_Center)
+					[
+						SNew(STextBlock)
+							.Text(LOCTEXT("AddLayerParameterPrompt", "Add a Material Attribute Layers parameter to see it here."))
+					]
 			];
 	}
 }
 
-
 void SMaterialLayersFunctionsInstanceWrapper::Construct(const FArguments& InArgs)
 {
 	NestedTree = SNew(SMaterialLayersFunctionsInstanceTree)
