Fix Specular Profile override on material instance hierarchy.

diff --git a/Engine/Source/Editor/MaterialEditor/Private/MaterialEditorInstanceDetailCustomization.cpp b/Engine/Source/Editor/MaterialEditor/Private/MaterialEditorInstanceDetailCustomization.cpp
--- a/Engine/Source/Editor/MaterialEditor/Private/MaterialEditorInstanceDetailCustomization.cpp
+++ b/Engine/Source/Editor/MaterialEditor/Private/MaterialEditorInstanceDetailCustomization.cpp
@@ -668,7 +668,7 @@ EVisibility FMaterialInstanceParameterDetails::ShouldShowSubsurfaceProfile() con
 
 EVisibility FMaterialInstanceParameterDetails::ShouldShowSpecularProfile() const
 {
-	const bool bHasProfiles = Substrate::IsSubstrateEnabled() && MaterialEditorInstance->Parent && MaterialEditorInstance->Parent->SpecularProfiles.Num() > 0;
+	const bool bHasProfiles = Substrate::IsSubstrateEnabled() && (MaterialEditorInstance->SourceInstance->NumSpecularProfile_Internal() > 0 || MaterialEditorInstance->SourceInstance->GetSpecularProfileOverride_Internal() != nullptr);
 	return bHasProfiles ? EVisibility::Visible : EVisibility::Collapsed;
 }
 
diff --git a/Engine/Source/Runtime/Engine/Private/Materials/MaterialInstance.cpp b/Engine/Source/Runtime/Engine/Private/Materials/MaterialInstance.cpp
--- a/Engine/Source/Runtime/Engine/Private/Materials/MaterialInstance.cpp
+++ b/Engine/Source/Runtime/Engine/Private/Materials/MaterialInstance.cpp
@@ -5260,7 +5260,17 @@ USpecularProfile* UMaterialInstance::GetSpecularProfile_Internal(uint32 Index) c
 
 USpecularProfile* UMaterialInstance::GetSpecularProfileOverride_Internal() const
 {
-	return bOverrideSpecularProfile ? SpecularProfileOverride : nullptr;
+	// Recursively look for override among the ancestors
+	USpecularProfile* Out = nullptr;
+	if (bOverrideSpecularProfile && SpecularProfileOverride)
+	{
+		Out = SpecularProfileOverride;
+	}
+	else if (Parent)
+	{
+		Out = Parent->GetSpecularProfileOverride_Internal();
+	}
+	return Out;
 }
 
 bool UMaterialInstance::CastsRayTracedShadows() const
