Substrate - fixed wrong maerial instance blend mode disaply when not opening a material before.

#tests PC editor
#rnx

diff --git a/Engine/Source/Editor/MaterialEditor/Private/MaterialEditorInstanceDetailCustomization.cpp b/Engine/Source/Editor/MaterialEditor/Private/MaterialEditorInstanceDetailCustomization.cpp
--- a/Engine/Source/Editor/MaterialEditor/Private/MaterialEditorInstanceDetailCustomization.cpp
+++ b/Engine/Source/Editor/MaterialEditor/Private/MaterialEditorInstanceDetailCustomization.cpp
@@ -1006,6 +1006,9 @@ void FMaterialInstanceParameterDetails::CreateBasePropertyOverrideWidgets(IDetai
 	const FText ParameterDisabledToolTipString = FText::FromString(TEXT("This material instance parent restricts the creation of new shader permutations. Overriding this parameter would result in the generation of additional shader permutations."));
 	const bool bStaticParametersOverrideDisabled = MaterialEditorInstance->SourceInstance->bDisallowStaticParameterPermutations;
 
+	// Override blend mode appearance based on Substrate being enabled or not.
+	GetBlendModeEnum();
+
 	auto CreateBaseOverrideRow = [&] (
 		const char* PropertyName,
 		auto&& OverrideBoolEnabledMemberFn,
