Add explicit error logging information to FMaterialEditor::CreateDerivedMaterialInstancesPreviews due to crashing non-deterministically but currently unable to track down a repro. The aim is to try and get more info back via the crash reporter for further investigation.

diff --git a/Engine/Source/Editor/MaterialEditor/Private/MaterialEditor.cpp b/Engine/Source/Editor/MaterialEditor/Private/MaterialEditor.cpp
--- a/Engine/Source/Editor/MaterialEditor/Private/MaterialEditor.cpp
+++ b/Engine/Source/Editor/MaterialEditor/Private/MaterialEditor.cpp
@@ -6626,7 +6626,28 @@ void FMaterialEditor::CreateDerivedMaterialInstancesPreviews()
 	OriginalDerivedMaterialInstances.Empty();
 
 	TArray<FString> DisplayNames;
-	DisplayNames.Push(OriginalMaterial->GetName());
+	if (OriginalMaterial)
+	{
+		DisplayNames.Push(OriginalMaterial->GetName());
+	}
+	else
+	{
+		//UE-296291, currently unable to track down why the editor occasionally crashes from OriginalMaterial object being a nullptr at this point.
+		//Our earliest report we are aware of is around UE5.5.3 but there's no obvious criteria as to why, and currently unable to repro deterministically, so this debug path is temporary to try and pull more info via the CrashReporter.
+		//If you see this message and know a method to reproduce please let Epic know!
+		const FString MissingEntry = "NULL";
+		UE_LOG(LogMaterialEditor, Warning, TEXT("MaterialEditor missing valid OriginalMaterial."));
+		UE_LOG(LogMaterialEditor, Warning, TEXT("OriginalMaterialObject = %s"), *(OriginalMaterialObject ? OriginalMaterialObject->GetPathName() : MissingEntry));
+		UE_LOG(LogMaterialEditor, Warning, TEXT("Material = %s"), *(Material ? Material->GetPathName() : MissingEntry));
+		UE_LOG(LogMaterialEditor, Warning, TEXT("ExpressionPreviewMaterial = %s"), *(ExpressionPreviewMaterial ? ExpressionPreviewMaterial->GetPathName() : MissingEntry));
+		UE_LOG(LogMaterialEditor, Warning, TEXT("PreviewExpression = %s"), *(PreviewExpression ? PreviewExpression->GetPathName() : MissingEntry));
+		UE_LOG(LogMaterialEditor, Warning, TEXT("MaterialFunction = %s"), *(MaterialFunction ? MaterialFunction->GetPathName() : MissingEntry));
+		UE_LOG(LogMaterialEditor, Warning, TEXT("bMaterialDirty = %s"), *LexToString(bMaterialDirty));
+		UE_LOG(LogMaterialEditor, Warning, TEXT("bStatsFromPreviewMaterial = %s"), *LexToString(bStatsFromPreviewMaterial));
+		UE_LOG(LogMaterialEditor, Warning, TEXT("bDestructing = %s"), *LexToString(bDestructing));
+
+		UE_LOG(LogMaterialEditor, Fatal, TEXT("OriginalMaterial is invalid in the material editor, with no deterministic repro of the problem. This message is to highlight the crash/callstack via the crash reporter with further Log info for Epic to investigate."));
+	}
 
 	const int32 MaxCount = CVarMaterialEdMaxDerivedMaterialInstances.GetValueOnGameThread();
 	if (MaterialStatsManager->GetProvideDerivedMIFlag() && MaxCount != 0)
