Make material editor log errors hyperlinks.

This is to bring visibility to the fact that it is possible to click on the log line to jump to the Node in question when errors are presented.

It also brings the Material Editor in line to how the Blueprint Editor works where by any nodes present in a log line are hyperlinks that can be clicked.

[RN] minor | Rendering - Materials and Shaders

diff --git a/Engine/Source/Editor/MaterialEditor/Private/MaterialEditor.cpp b/Engine/Source/Editor/MaterialEditor/Private/MaterialEditor.cpp
--- a/Engine/Source/Editor/MaterialEditor/Private/MaterialEditor.cpp
+++ b/Engine/Source/Editor/MaterialEditor/Private/MaterialEditor.cpp
@@ -3143,14 +3143,37 @@ void FMaterialEditor::UpdateMaterialinfoList_Old()
 		GetFeatureLevelName(FeatureLevel,FeatureLevelName);
 		for(int32 ErrorIndex = 0; ErrorIndex < CompileErrors.Num(); ErrorIndex++)
 		{
-			FString ErrorString = FString::Printf(TEXT("[%s] %s"), *FeatureLevelName, *CompileErrors[ErrorIndex]);
+			const FString ErrorString = FString::Printf(TEXT("[%s] %s"), *FeatureLevelName, *CompileErrors[ErrorIndex]);
 			TempMaterialInfoList.Add(MakeShareable(new FMaterialInfo(ErrorString, FLinearColor::Red)));
 			TSharedRef<FTokenizedMessage> Line = FTokenizedMessage::Create( EMessageSeverity::Error );
+			bool bHyperlinkAdded = false;
+
 			if(FailingExpression.Num() && ensure(FailingExpression.Num() == CompileErrors.Num()) && FailingExpression[ErrorIndex])
 			{
 				Line->SetMessageLink(FUObjectToken::Create(FailingExpression[ErrorIndex]));
+
+				const FString& CompilerError = CompileErrors[ErrorIndex];
+				int32 NodeNameEndPos = INDEX_NONE;
+
+				// If we have a failing expression, first matching parantheses will be for its name
+				if (CompilerError.FindChar(L')', NodeNameEndPos))
+				{
+					bHyperlinkAdded = true;
+					Line->AddToken(FTextToken::Create(FText::FromString(FeatureLevelName)));
+
+					const FString NodeName = CompilerError.Left(NodeNameEndPos + 1);
+					const FString ErrorBody = CompilerError.RightChop(NodeNameEndPos + 2);
+
+					Line->AddToken(FUObjectToken::Create(FailingExpression[ErrorIndex], FText::FromString(NodeName)));
+					Line->AddToken(FTextToken::Create(FText::FromString(ErrorBody)));
+				}
+			}
+
+			if (!bHyperlinkAdded)
+			{
+				Line->AddToken(FTextToken::Create(FText::FromString(ErrorString)));
 			}
-			Line->AddToken( FTextToken::Create( FText::FromString( ErrorString ) ) );
+
 			Messages.Add(Line);
 			bForceDisplay = true;
 		}
