Performance optimization on Material Instance being opened.
Reduce obsolete calls to RegenerateArrays() triggered by materials layers when opening up Material Instance Editor(even if a material does not have layers).
CreateGroups(), Refresh() UpdatePropertyWindow, SetEditorInstance() were calling RegenerateArrays() every time, a method that rebuilds cached data, all parameters and their visibility.

This is unnecessary when the Material Instance Editor is just being created, as it comes from inside FMaterialInstanceEditor::InitMaterialInstanceEditor(), where SetSourceInstance() is called just before, which also always calls RegenerateArrays().

RegenerateArrays was being called 2 times on a standard Material Instance(even if it did not have any layers) and 4 times on a Material Instance for a Layered Material.

Also, remove duplicate generation of widgets triggered by [REDACTED_PERSON] and Refresh().

Average 20% performance improvement.

[RN] minor | Rendering Architecture

diff --git a/Engine/Source/Editor/MaterialEditor/Private/MaterialEditor.h b/Engine/Source/Editor/MaterialEditor/Private/MaterialEditor.h
--- a/Engine/Source/Editor/MaterialEditor/Private/MaterialEditor.h
+++ b/Engine/Source/Editor/MaterialEditor/Private/MaterialEditor.h
@@ -1062,7 +1062,7 @@ private:
 	TSharedPtr<class SMaterialCustomPrimitiveDataPanel> MaterialCustomPrimitiveDataWidget;
 
 	/** Layer Properties View */
-	TSharedPtr<class SMaterialLayersFunctionsWrapper> MaterialLayersFunctionsInstance;
+	TSharedPtr<class SMaterialLayersFunctionsMaterialWrapper> MaterialLayersFunctionsInstance;
 	
 	/** The current transaction. */
 	FScopedTransaction* ScopedTransaction;
diff --git a/Engine/Source/Editor/MaterialEditor/Private/MaterialInstanceEditor.cpp b/Engine/Source/Editor/MaterialEditor/Private/MaterialInstanceEditor.cpp
--- a/Engine/Source/Editor/MaterialEditor/Private/MaterialInstanceEditor.cpp
+++ b/Engine/Source/Editor/MaterialEditor/Private/MaterialInstanceEditor.cpp
@@ -530,7 +530,10 @@ void FMaterialInstanceEditor::InitMaterialInstanceEditor( const EToolkitMode::Ty
 		}
 	}
 
-	Refresh();
+	PreviewVC->RefreshViewport();
+	RebuildInheritanceList();
+	MaterialInstanceDetails->SetObject(MaterialEditorInstance, true);
+	RefreshOnScreenTextureWarnings();
 
 	// Notify other editors if this material editor has a post process named output, which may affect their preview
 	NotifyUserSceneTextureLoadOrUnload();
@@ -1052,8 +1055,6 @@ TSharedRef<SDockTab> FMaterialInstanceEditor::SpawnTab_Properties( const FSpawnT
 			]
 		];
 
-	UpdatePropertyWindow();
-
 	AddToSpawnedToolPanels( Args.GetTabId().TabType, SpawnedTab );
 	return SpawnedTab;
 }
diff --git a/Engine/Source/Editor/MaterialEditor/Private/SMaterialLayersFunctionsTree.cpp b/Engine/Source/Editor/MaterialEditor/Private/SMaterialLayersFunctionsTree.cpp
--- a/Engine/Source/Editor/MaterialEditor/Private/SMaterialLayersFunctionsTree.cpp
+++ b/Engine/Source/Editor/MaterialEditor/Private/SMaterialLayersFunctionsTree.cpp
@@ -908,6 +908,7 @@ void SMaterialLayersFunctionsInstanceTree::ResetAssetToDefault(TSharedPtr<FSorte
 {
 	FMaterialPropertyHelpers::ResetLayerAssetToDefault(InData->Parameter, InData->ParameterInfo.Association, InData->ParameterInfo.Index, MaterialEditorInstance);
 	UpdateThumbnailMaterial(InData->ParameterInfo.Association, InData->ParameterInfo.Index, false);
+	MaterialEditorInstance->RegenerateArrays();
 	CreateGroupsWidget();
 	RequestTreeRefresh();
 }
@@ -918,6 +919,7 @@ void SMaterialLayersFunctionsInstanceTree::AddLayer()
 	FunctionInstanceHandle->NotifyPreChange();
 	FunctionInstance->AppendBlendedLayer();
 	FunctionInstanceHandle->NotifyPostChange(EPropertyChangeType::ArrayAdd);
+	MaterialEditorInstance->RegenerateArrays();
 	CreateGroupsWidget();
 	RequestTreeRefresh();
 }
@@ -929,6 +931,7 @@ void SMaterialLayersFunctionsInstanceTree::RemoveLayer(int32 Index)
 	FunctionInstance->RemoveBlendedLayerAt(Index);
 	MaterialEditorInstance->SourceInstance->RemoveLayerParameterIndex(Index);
 	FunctionInstanceHandle->NotifyPostChange(EPropertyChangeType::ArrayRemove);
+	MaterialEditorInstance->RegenerateArrays();
 	CreateGroupsWidget();
 	RequestTreeRefresh();
 }
@@ -939,7 +942,9 @@ FReply SMaterialLayersFunctionsInstanceTree::UnlinkLayer(int32 Index)
 	FunctionInstanceHandle->NotifyPreChange();
 	FunctionInstance->UnlinkLayerFromParent(Index);
 	FunctionInstanceHandle->NotifyPostChange(EPropertyChangeType::ValueSet);
+	MaterialEditorInstance->RegenerateArrays();
 	CreateGroupsWidget();
+
 	return FReply::Handled();
 }
 
@@ -951,6 +956,7 @@ FReply SMaterialLayersFunctionsInstanceTree::RelinkLayersToParent()
 	FunctionInstanceHandle->NotifyPostChange(EPropertyChangeType::ValueSet);
 	MaterialEditorInstance->RegenerateArrays();
 	CreateGroupsWidget();
+
 	return FReply::Handled();
 }
 
@@ -980,9 +986,6 @@ FReply SMaterialLayersFunctionsInstanceTree::ToggleLayerVisibility(int32 Index)
 		const FScopedTransaction Transaction(LOCTEXT("ToggleLayerAndBlendVisibility", "Toggles visibility for a blended layer"));
 		FunctionInstanceHandle->NotifyPreChange();
 		FunctionInstance->ToggleBlendedLayerVisibility(Index);
-		FunctionInstanceHandle->NotifyPostChange(EPropertyChangeType::ValueSet);
-		CreateGroupsWidget();
-		return FReply::Handled();
 	}
 	else
 	{
@@ -1003,11 +1006,13 @@ FReply SMaterialLayersFunctionsInstanceTree::ToggleLayerVisibility(int32 Index)
 		}
 
 		bLayerIsolated = !bLayerIsolated;
-		FunctionInstanceHandle->NotifyPostChange(EPropertyChangeType::ValueSet);
-		CreateGroupsWidget();
-		return FReply::Handled();
 	}
 
+	FunctionInstanceHandle->NotifyPostChange(EPropertyChangeType::ValueSet);
+	MaterialEditorInstance->RegenerateArrays();
+	CreateGroupsWidget();
+
+	return FReply::Handled();
 }
 
 TSharedPtr<class FAssetThumbnailPool> SMaterialLayersFunctionsInstanceTree::GetTreeThumbnailPool()
@@ -1039,7 +1044,7 @@ TSharedPtr<IDetailTreeNode> FindParameterGroupsNode(TSharedPtr<IPropertyRowGener
 void SMaterialLayersFunctionsInstanceTree::CreateGroupsWidget()
 {
 	check(MaterialEditorInstance);
-	MaterialEditorInstance->RegenerateArrays();
+
 	NonLayerProperties.Empty();
 	LayerProperties.Empty();
 	FunctionParameter = nullptr;
@@ -1440,11 +1445,7 @@ void SMaterialLayersFunctionsInstanceWrapper::Refresh()
 
 	LayerParameter.Reset();
 
-	UMaterialInterface* Material = MaterialEditorInstance->GetMaterialInterface();
-	UMaterialInterface* ParentMaterial = MaterialEditorInstance->GetParentMaterialInterface();
-
-	FMaterialLayersFunctions Layers;
-	const bool bHasLayers = ParentMaterial ? ParentMaterial->GetMaterialLayers(Layers) : Material->GetMaterialLayers(Layers);
+	const bool bHasLayers = HasLayers();
 
 	if (bHasLayers)
 	{
@@ -1563,6 +1564,17 @@ void SMaterialLayersFunctionsInstanceWrapper::Refresh()
 	}
 }
 
+bool SMaterialLayersFunctionsInstanceWrapper::HasLayers() const
+{
+	UMaterialInterface* Material = MaterialEditorInstance->GetMaterialInterface();
+	UMaterialInterface* ParentMaterial = MaterialEditorInstance->GetParentMaterialInterface();
+
+	FMaterialLayersFunctions Layers;
+	const bool bHasLayers = ParentMaterial ? ParentMaterial->GetMaterialLayers(Layers) : Material->GetMaterialLayers(Layers);
+
+	return bHasLayers;
+}
+
 void SMaterialLayersFunctionsInstanceWrapper::Construct(const FArguments& InArgs)
 {
 	NestedTree = SNew(SMaterialLayersFunctionsInstanceTree)
@@ -1574,12 +1586,19 @@ void SMaterialLayersFunctionsInstanceWrapper::Construct(const FArguments& InArgs
 	LayerParameter = NestedTree->GetFunctionParameter();
 
 	MaterialEditorInstance = InArgs._InMaterialEditorInstance;
+
 	Refresh();
 }
 
 void SMaterialLayersFunctionsInstanceWrapper::SetEditorInstance(UMaterialEditorParameters* InMaterialEditorInstance)
 {
 	NestedTree->SetMaterialEditorInstance(InMaterialEditorInstance);
+
+	const bool bHasLayers = HasLayers();
+	if (bHasLayers)
+	{
+		MaterialEditorInstance->RegenerateArrays();
+	}
 	Refresh();
 }
 
diff --git a/Engine/Source/Editor/MaterialEditor/Private/SMaterialLayersFunctionsTree.h b/Engine/Source/Editor/MaterialEditor/Private/SMaterialLayersFunctionsTree.h
--- a/Engine/Source/Editor/MaterialEditor/Private/SMaterialLayersFunctionsTree.h
+++ b/Engine/Source/Editor/MaterialEditor/Private/SMaterialLayersFunctionsTree.h
@@ -124,7 +124,8 @@ public:
 	void Construct(const FArguments& InArgs);
 	virtual void SetEditorInstance(UMaterialEditorParameters* InMaterialEditorInstance) override;
 	virtual void Refresh() override;
-	
+	bool HasLayers() const;
+
 	TSharedPtr<class IPropertyRowGenerator> GetGenerator();
 	TAttribute<ECheckBoxState> IsParamChecked;
 	TWeakObjectPtr<class UDEditorParameterValue> LayerParameter;
