MaterialEditingLibrary: Add GetNumShaderTypes() BP function.
#rnx

diff --git a/Engine/Source/Editor/MaterialEditor/Private/MaterialEditingLibrary.cpp b/Engine/Source/Editor/MaterialEditor/Private/MaterialEditingLibrary.cpp
--- a/Engine/Source/Editor/MaterialEditor/Private/MaterialEditingLibrary.cpp
+++ b/Engine/Source/Editor/MaterialEditor/Private/MaterialEditingLibrary.cpp
@@ -1543,6 +1543,28 @@ bool UMaterialEditingLibrary::GetStaticSwitchParameterSource(UMaterialInterface*
 	return false;
 }
 
+int32 UMaterialEditingLibrary::GetNumShaderTypes(UMaterialInterface* Material)
+{
+	int NumShaders = 0;
+
+	if (Material != nullptr)
+	{
+		TArray<FDebugShaderTypeInfo> ShaderInfos;
+		Material->GetShaderTypes(GMaxRHIShaderPlatform, nullptr, ShaderInfos);
+
+		for (FDebugShaderTypeInfo const& ShaderInfo : ShaderInfos)
+		{
+			NumShaders += ShaderInfo.ShaderTypes.Num();
+			for (FDebugShaderPipelineInfo const& PipelineInfo : ShaderInfo.Pipelines)
+			{
+				NumShaders += PipelineInfo.ShaderTypes.Num();
+			}
+		}
+	}
+
+	return NumShaders;
+}
+
 FMaterialStatistics UMaterialEditingLibrary::GetStatistics(class UMaterialInterface* Material)
 {
 	FMaterialStatistics Result;
diff --git a/Engine/Source/Editor/MaterialEditor/Public/MaterialEditingLibrary.h b/Engine/Source/Editor/MaterialEditor/Public/MaterialEditingLibrary.h
--- a/Engine/Source/Editor/MaterialEditor/Public/MaterialEditingLibrary.h
+++ b/Engine/Source/Editor/MaterialEditor/Public/MaterialEditingLibrary.h
@@ -451,6 +451,10 @@ public:
 	*/
 	UFUNCTION(BlueprintCallable, Category = "MaterialEditing")
 	static UE_API bool GetStaticSwitchParameterSource(UMaterialInterface* Material, const FName ParameterName, FSoftObjectPath& ParameterSource);
+
+	/** Returns the number of shader types (including pipeline shader types) for the given material. */
+	UFUNCTION(BlueprintCallable, Category = "MaterialEditing")
+	static UE_API int32 GetNumShaderTypes(UMaterialInterface* Material);
 	
 	/** Returns statistics about the given material */
 	UFUNCTION(BlueprintCallable, Category = "MaterialEditing")
