Substrate - clarified/simplified static blend mode enum update function.

#tests PC editor
#rnx

diff --git a/Engine/Source/Editor/MaterialEditor/Private/MaterialEditorDetailCustomization.cpp b/Engine/Source/Editor/MaterialEditor/Private/MaterialEditorDetailCustomization.cpp
--- a/Engine/Source/Editor/MaterialEditor/Private/MaterialEditorDetailCustomization.cpp
+++ b/Engine/Source/Editor/MaterialEditor/Private/MaterialEditorDetailCustomization.cpp
@@ -55,7 +55,7 @@
 #define LOCTEXT_NAMESPACE "MaterialEditor"
 
 // Update the blend mode names based on what is supported in legacy mode or Substrate mode
-UEnum* GetBlendModeEnum()
+void UpdateBlendModeEnum()
 {
 	UEnum* BlendModeEnum = StaticEnum<EBlendMode>();
 	if (Substrate::IsSubstrateEnabled())
@@ -74,7 +74,6 @@ UEnum* GetBlendModeEnum()
 		// BLEND_TranslucentColoredTransmittance is not supported in legacy mode
 		BlendModeEnum->SetMetaData(TEXT("Hidden"), TEXT("True"), BLEND_TranslucentColoredTransmittance);
 	}
-	return BlendModeEnum;
 }
 
 TSharedRef<IDetailCustomization> FMaterialExpressionParameterDetails::MakeInstance(FOnCollectParameterGroups InCollectGroupsDelegate)
@@ -757,7 +756,7 @@ void FMaterialExpressionCollectionParameterDetails::PopulateParameters(const FSt
 		UObject* CollectionObject = nullptr;
 		FPropertyAccess::Result Result = CollectionPropertyHandle->GetValue(CollectionObject);
 
-		// Return an empty set of parameters if multiple paramter collection nodes are selected.
+		// Return an empty set of parameters if multiple parameter collection nodes are selected.
 		if (Result == FPropertyAccess::MultipleValues)
 		{
 			return;
@@ -892,8 +891,8 @@ void FMaterialDetailCustomization::CustomizeDetails( IDetailLayoutBuilder& Detai
 			// Patch blend mode displayed names
 			if (PropertyName == GET_MEMBER_NAME_CHECKED(UMaterial, BlendMode))
 			{
-				FByteProperty* BlendModeProperty = (FByteProperty*)Property;
-				BlendModeProperty->Enum = GetBlendModeEnum();
+				// Override blend mode appearance based on Substrate being enabled or not.
+				UpdateBlendModeEnum();
 
 				PropertyHandle->SetOnPropertyValueChanged(FSimpleDelegate::CreateLambda([&DetailLayout]()
 					{
@@ -933,7 +932,7 @@ void FMaterialDetailCustomization::CustomizeDetails( IDetailLayoutBuilder& Detai
 			{
 				if (bIsAlphaHoldout)
 				{
-					static FText RestrictReason = NSLOCTEXT("PropertyEditor", "TranslucencyPassRestriction", "Seperate translucency pass should not be used with Alpha Holdout blend mode.");
+					static FText RestrictReason = NSLOCTEXT("PropertyEditor", "TranslucencyPassRestriction", "Separate translucency pass should not be used with Alpha Holdout blend mode.");
 					TSharedPtr<FPropertyRestriction> EnumRestriction = MakeShared<FPropertyRestriction>(RestrictReason);
 
 					FByteProperty* TranslucencyPassProperty = (FByteProperty*)Property;
diff --git a/Engine/Source/Editor/MaterialEditor/Private/MaterialEditorInstanceDetailCustomization.cpp b/Engine/Source/Editor/MaterialEditor/Private/MaterialEditorInstanceDetailCustomization.cpp
--- a/Engine/Source/Editor/MaterialEditor/Private/MaterialEditorInstanceDetailCustomization.cpp
+++ b/Engine/Source/Editor/MaterialEditor/Private/MaterialEditorInstanceDetailCustomization.cpp
@@ -1020,7 +1020,7 @@ void FMaterialInstanceParameterDetails::CreatePostProcessOverrideWidgets(IDetail
 	}
 }
 
-UEnum* GetBlendModeEnum();
+void UpdateBlendModeEnum();
 
 void FMaterialInstanceParameterDetails::CreateBasePropertyOverrideWidgets(IDetailLayoutBuilder& DetailLayout, IDetailGroup& MaterialPropertyOverrideGroup)
 {
@@ -1032,7 +1032,7 @@ void FMaterialInstanceParameterDetails::CreateBasePropertyOverrideWidgets(IDetai
 	const bool bStaticParametersOverrideDisabled = MaterialEditorInstance->SourceInstance->bDisallowStaticParameterPermutations;
 
 	// Override blend mode appearance based on Substrate being enabled or not.
-	GetBlendModeEnum();
+	UpdateBlendModeEnum();
 
 	auto CreateBaseOverrideRow = [&] (
 		const char* PropertyName,
