Fix stale object pointer access in material instance editor UI.
#rnx

diff --git a/Engine/Source/Editor/MaterialEditor/Private/MaterialEditorInstanceDetailCustomization.cpp b/Engine/Source/Editor/MaterialEditor/Private/MaterialEditorInstanceDetailCustomization.cpp
--- a/Engine/Source/Editor/MaterialEditor/Private/MaterialEditorInstanceDetailCustomization.cpp
+++ b/Engine/Source/Editor/MaterialEditor/Private/MaterialEditorInstanceDetailCustomization.cpp
@@ -86,17 +86,22 @@ void FMaterialInstanceParameterDetails::CustomizeDetails(IDetailLayoutBuilder& D
 	FName GroupsCategoryName = TEXT("ParameterGroups");
 	IDetailCategoryBuilder& GroupsCategory = DetailLayout.EditCategory(GroupsCategoryName, LOCTEXT("MICParamGroupsTitle", "Parameter Groups"));
 	
+	TWeakObjectPtr<UMaterialInterface> MaterialWeakPtr(MaterialEditorInstance->GetMaterialInterface());
+
 	// Insert optional validation warning row at the top of the category.
 	GroupsCategory.AddCustomRow(LOCTEXT("ValidationWarningRow", "Validation Warning"))
 	.WholeRowContent()
 	[
 		SNew(SWarningOrErrorBox)
-		.Visibility_Lambda([this]() -> EVisibility
+		.Visibility_Lambda([this, MaterialWeakPtr]() -> EVisibility
 		{
 			bool bWarn = false;
-			if (UMaterialEditorValidationSubsystem* Subsystem = GEditor->GetEditorSubsystem<UMaterialEditorValidationSubsystem>())
+			if (MaterialWeakPtr.IsValid())
 			{
-				bWarn = !Subsystem->Validate(MaterialEditorInstance->GetMaterialInterface(), *ValidationContext.Get());
+				if (UMaterialEditorValidationSubsystem* Subsystem = GEditor->GetEditorSubsystem<UMaterialEditorValidationSubsystem>())
+				{
+					bWarn = !Subsystem->Validate(MaterialWeakPtr.Get(), *ValidationContext.Get());
+				}
 			}
 			return bWarn ? EVisibility::Visible : EVisibility::Collapsed;
 		})
