Add Select Focused Nodes feature to the Hide Unrelated Nodes mode for the Material Editor.
This will add an entry which allos selection to happen on all the nodes that are focused based on the options selected for "Hide Unrelated Nodes".
This feature can only be enabled when "Focus whole chain" is also enabled, as it will have that effect either way(becase selection is the source for the focus point) so the user needs to be aware of it.

[RN] minor | Rendering Architecture

diff --git a/Engine/Source/Editor/MaterialEditor/Private/MaterialEditor.cpp b/Engine/Source/Editor/MaterialEditor/Private/MaterialEditor.cpp
--- a/Engine/Source/Editor/MaterialEditor/Private/MaterialEditor.cpp
+++ b/Engine/Source/Editor/MaterialEditor/Private/MaterialEditor.cpp
@@ -956,6 +956,8 @@ FMaterialEditor::FMaterialEditor()
 	, bHideUnrelatedNodes(false)
 	, bLockNodeFadeState(false)
 	, bFocusWholeChain(false)
+	, bSelectFocusedNodes(false)
+	, bFocusSelectionInProgress(false)
 	, bSelectRegularNode(false)
 	, MenuExtensibilityManager(new FExtensibilityManager)
 	, ToolBarExtensibilityManager(new FExtensibilityManager)
@@ -3724,6 +3726,7 @@ void FMaterialEditor::ToggleHideUnrelatedNodes()
 	{
 		bLockNodeFadeState = false;
 		bFocusWholeChain = false;
+		bSelectFocusedNodes = false;
 	}
 }
 
@@ -3732,6 +3735,11 @@ bool FMaterialEditor::IsToggleHideUnrelatedNodesChecked() const
 	return bHideUnrelatedNodes == true;
 }
 
+bool FMaterialEditor::IsSelectFocusedNodesEnabled() const
+{
+	return bFocusWholeChain;
+}
+
 void FMaterialEditor::CollectDownstreamNodes(UMaterialGraphNode* CurrentNode, TArray<UMaterialGraphNode*>& CollectedNodes)
 {
 	for (UEdGraphPin* OutputPin : CurrentNode->Pins)
@@ -3783,21 +3791,7 @@ void FMaterialEditor::HideUnrelatedNodes()
 		return;
 	}
 
-	TArray<UMaterialGraphNode*> NodesToShow;
-
-	const FGraphPanelSelectionSet SelectedNodes = GetSelectedNodes();
-
-	for (FGraphPanelSelectionSet::TConstIterator NodeIt(SelectedNodes); NodeIt; ++NodeIt)
-	{
-		UMaterialGraphNode* SelectedNode = Cast<UMaterialGraphNode>(*NodeIt);
-
-		if (SelectedNode)
-		{
-			NodesToShow.Add(SelectedNode);
-			CollectDownstreamNodes( SelectedNode, NodesToShow );
-			CollectUpstreamNodes( SelectedNode, NodesToShow );
-		}
-	}
+	const TArray<UMaterialGraphNode*> NodesToShow = GetSelectionRelatedNodes();
 
 	TArray<class UEdGraphNode*> AllNodes = FocusedGraphEd->GetCurrentGraph()->Nodes;
 
@@ -3826,6 +3820,28 @@ void FMaterialEditor::HideUnrelatedNodes()
 	}
 
 	FocusedGraphEd->FocusCommentNodes(CommentNodes, RelatedNodes);
+
+	if (bFocusWholeChain && bSelectFocusedNodes)
+	{
+		SelectMaterialGraphNodes(*FocusedGraphEd, NodesToShow);
+	}
+}
+
+void FMaterialEditor::SelectMaterialGraphNodes(SGraphEditor& InGraphEd, const TArray<UMaterialGraphNode*>& InNodes)
+{
+	TArray<UObject*> SelectedObjects;
+	SelectedObjects.Reserve(InNodes.Num());
+
+	bFocusSelectionInProgress = true;
+	for (UMaterialGraphNode* Node : InNodes)
+	{
+		SelectedObjects.Add(Node->MaterialExpression);
+		InGraphEd.SetNodeSelection(Node, true);
+	}
+	bFocusSelectionInProgress = false;
+
+	GetDetailView()->SetObjects(SelectedObjects, true);
+	FocusDetailsPanel();
 }
 
 void FMaterialEditor::MakeHideUnrelatedNodesOptionsMenu(UToolMenu* Menu)
@@ -3872,9 +3888,34 @@ void FMaterialEditor::MakeHideUnrelatedNodesOptionsMenu(UToolMenu* Menu)
 				]
 		];
 
+	TSharedRef<SWidget> SelectFocusedNodesCheckBox = SNew(SBox)
+		.IsEnabled(this, &FMaterialEditor::IsSelectFocusedNodesEnabled)
+		[
+			SNew(SCheckBox)
+				.IsChecked(bSelectFocusedNodes ? ECheckBoxState::Checked : ECheckBoxState::Unchecked)
+				.OnCheckStateChanged(this, &FMaterialEditor::OnSelectFocusedNodesCheckStateChanged)
+				.Style(FAppStyle::Get(), "Menu.CheckBox")
+				.ToolTipText(
+					LOCTEXT("SelectFocusedNodesTooltip", "Select focused nodes in the chain. \n" 
+					"This only works if Focus Whole Chain is also enabled."))
+				.Content()
+				[
+					SNew(SHorizontalBox)
+
+						+ SHorizontalBox::Slot()
+						.Padding(2.0f, 0.0f, 0.0f, 0.0f)
+						[
+							SNew(STextBlock)
+								.Text(LOCTEXT("SelectFocusedNodes", "Select Focused Nodes"))
+								
+						]
+				]
+		];
+
 	FToolMenuSection& OptionsSection = Menu->AddSection("OptionsSection", LOCTEXT("HideUnrelatedOptions", "Hide Unrelated Options"));
 	OptionsSection.AddEntry(FToolMenuEntry::InitMenuEntry("LockNodeStateCheckBox", FUIAction(), LockNodeStateCheckBox));
 	OptionsSection.AddEntry(FToolMenuEntry::InitMenuEntry("FocusWholeChainCheckBox", FUIAction(), FocusWholeChainCheckBox));
+	OptionsSection.AddEntry(FToolMenuEntry::InitMenuEntry("SelectFocusedNodesCheckBox", FUIAction(), SelectFocusedNodesCheckBox));
 }
 
 void FMaterialEditor::OnLockNodeStateCheckStateChanged(ECheckBoxState NewCheckedState)
@@ -3900,6 +3941,22 @@ void FMaterialEditor::OnFocusWholeChainCheckStateChanged(ECheckBoxState NewCheck
 	}
 }
 
+void FMaterialEditor::OnSelectFocusedNodesCheckStateChanged(ECheckBoxState NewCheckedState)
+{
+	TSharedPtr<SGraphEditor> FocusedGraphEd = FocusedGraphEdPtr.Pin();
+	if (!FocusedGraphEd)
+	{
+		return;
+	}
+
+	bSelectFocusedNodes = NewCheckedState == ECheckBoxState::Checked;
+
+	if (bHideUnrelatedNodes && bSelectRegularNode)
+	{
+		const TArray<UMaterialGraphNode*> NodesToSelect = GetSelectionRelatedNodes();
+		SelectMaterialGraphNodes(*FocusedGraphEd, NodesToSelect);
+	}
+}
 
 void FMaterialEditor::OnUseCurrentTexture()
 {
@@ -7939,6 +7996,26 @@ FGraphAppearanceInfo FMaterialEditor::GetGraphAppearance() const
 	return AppearanceInfo;
 }
 
+TArray<UMaterialGraphNode*> FMaterialEditor::GetSelectionRelatedNodes()
+{
+	TArray<UMaterialGraphNode*> RelatedNodes;
+	const FGraphPanelSelectionSet SelectedNodes = GetSelectedNodes();
+
+	for (FGraphPanelSelectionSet::TConstIterator NodeIt(SelectedNodes); NodeIt; ++NodeIt)
+	{
+		UMaterialGraphNode* SelectedNode = Cast<UMaterialGraphNode>(*NodeIt);
+
+		if (SelectedNode)
+		{
+			RelatedNodes.Add(SelectedNode);
+			CollectDownstreamNodes(SelectedNode, RelatedNodes);
+			CollectUpstreamNodes(SelectedNode, RelatedNodes);
+		}
+	}
+
+	return RelatedNodes;
+}
+
 void FMaterialEditor::DeepCopyExpressions(UMaterialGraph* CopyGraph, UMaterialExpression* NewSubgraphExpression)
 {
 	if (!CopyGraph || !NewSubgraphExpression)
@@ -8312,6 +8389,11 @@ void FMaterialEditor::UpdateDetailView()
 
 void FMaterialEditor::OnSelectedNodesChanged(const TSet<class UObject*>& NewSelection)
 {
+	if (bFocusSelectionInProgress)
+	{
+		return;
+	}
+
 	TArray<UObject*> SelectedObjects;
 
 	UObject* EditObject = Material;
@@ -8348,16 +8430,28 @@ void FMaterialEditor::OnSelectedNodesChanged(const TSet<class UObject*>& NewSele
 	GetDetailView()->SetObjects( SelectedObjects, true );
 	FocusDetailsPanel();
 	
-	if (bHideUnrelatedNodes && !bLockNodeFadeState)
+	if (bHideUnrelatedNodes)
 	{
-		if (TSharedPtr<SGraphEditor> FocusedGraphEd = FocusedGraphEdPtr.Pin())
+		if (!bLockNodeFadeState)
 		{
-			FocusedGraphEd->ResetAllNodesUnrelatedStates();
-		}
+			if (TSharedPtr<SGraphEditor> FocusedGraphEd = FocusedGraphEdPtr.Pin())
+			{
+				FocusedGraphEd->ResetAllNodesUnrelatedStates();
+			}
 
-		if (bSelectRegularNode)
+			if (bSelectRegularNode)
+			{
+				HideUnrelatedNodes();
+			}
+		}
+		else if(bSelectFocusedNodes)
 		{
-			HideUnrelatedNodes();
+			// We still want to select nodes
+			if (TSharedPtr<SGraphEditor> FocusedGraphEd = FocusedGraphEdPtr.Pin())
+			{
+				const TArray<UMaterialGraphNode*> NodesToSelect = GetSelectionRelatedNodes();
+				SelectMaterialGraphNodes(*FocusedGraphEd, NodesToSelect);
+			}
 		}
 	}
 }
diff --git a/Engine/Source/Editor/MaterialEditor/Private/MaterialEditor.h b/Engine/Source/Editor/MaterialEditor/Private/MaterialEditor.h
--- a/Engine/Source/Editor/MaterialEditor/Private/MaterialEditor.h
+++ b/Engine/Source/Editor/MaterialEditor/Private/MaterialEditor.h
@@ -784,19 +784,21 @@ private:
 	/** Command for the Refresh all previews button */
 	void OnAlwaysRefreshAllPreviews();
 	bool IsOnAlwaysRefreshAllPreviews() const;
-	/** Command for the stats button */
 
 	/** Make nodes which are unrelated to the selected nodes fade out */
 	void ToggleHideUnrelatedNodes();
 	bool IsToggleHideUnrelatedNodesChecked() const;
+	bool IsSelectFocusedNodesEnabled() const;
 	void CollectDownstreamNodes(UMaterialGraphNode* CurrentNode, TArray<UMaterialGraphNode*>& CollectedNodes);
 	void CollectUpstreamNodes(UMaterialGraphNode* CurrentNode, TArray<UMaterialGraphNode*>& CollectedNodes);
 	void HideUnrelatedNodes();
+	void SelectMaterialGraphNodes(SGraphEditor& InGraphEd, const TArray<UMaterialGraphNode*>& InNodes);
 
 	/** Make a drop down menu to control the opacity of unrelated nodes */
 	void MakeHideUnrelatedNodesOptionsMenu(class UToolMenu* Menu);
 	void OnLockNodeStateCheckStateChanged(ECheckBoxState NewCheckedState);
 	void OnFocusWholeChainCheckStateChanged(ECheckBoxState NewCheckedState);
+	void OnSelectFocusedNodesCheckStateChanged(ECheckBoxState NewCheckedState);
 
 	/** Command for using currently selected texture */
 	void OnUseCurrentTexture();
@@ -956,6 +958,9 @@ private:
 	/** Gets the current Material Graph's appearance */
 	FGraphAppearanceInfo GetGraphAppearance() const;
 
+	/* Gets the upstream and downstream nodes related to the current selection. **/
+	TArray<UMaterialGraphNode*> GetSelectionRelatedNodes();
+
 	/**
 	 * Deletes any disconnected material expressions.
 	 */
@@ -1100,6 +1105,12 @@ private:
 	/** Focus all nodes in the same output chain  */
 	bool bFocusWholeChain;
 
+	/** Select focused nodes. */
+	bool bSelectFocusedNodes;
+
+	/** Used to avoid selection operations while automated selection takes place. */
+	bool bFocusSelectionInProgress;
+
 	/** If a regular node (not a comment node, not the output node) has been selected */
 	bool bSelectRegularNode;
 
