Adds UMaterialEditingLibrary::ListShaders which will return an array of all shaders that material will compile.

- This returns a list of FDebugShaderInfo, which contains the Vertex Factory Type Name and the Shader Type Name for each shader in the material.
- Re-implement the `GetNumShaderTypes` to simply just return the length of that array returned by ListShaders.
- Handle null shader types, or if a VF type is null (non mesh-material shaders)

[RN] minor | Rendering - Materials and Shaders

diff --git a/Engine/Source/Editor/MaterialEditor/Private/MaterialEditingLibrary.cpp b/Engine/Source/Editor/MaterialEditor/Private/MaterialEditingLibrary.cpp
--- a/Engine/Source/Editor/MaterialEditor/Private/MaterialEditingLibrary.cpp
+++ b/Engine/Source/Editor/MaterialEditor/Private/MaterialEditingLibrary.cpp
@@ -1545,7 +1545,12 @@ bool UMaterialEditingLibrary::GetStaticSwitchParameterSource(UMaterialInterface*
 
 int32 UMaterialEditingLibrary::GetNumShaderTypes(UMaterialInterface* Material)
 {
-	int NumShaders = 0;
+	return UMaterialEditingLibrary::ListShaders(Material).Num();
+}
+
+TArray<FDebugShaderInfo> UMaterialEditingLibrary::ListShaders(UMaterialInterface* Material)
+{
+	TArray<FDebugShaderInfo> OutShaderInfos;
 
 	if (Material != nullptr)
 	{
@@ -1554,15 +1559,34 @@ int32 UMaterialEditingLibrary::GetNumShaderTypes(UMaterialInterface* Material)
 
 		for (FDebugShaderTypeInfo const& ShaderInfo : ShaderInfos)
 		{
-			NumShaders += ShaderInfo.ShaderTypes.Num();
+			for (FShaderType* ShaderType : ShaderInfo.ShaderTypes)
+			{
+				if (ShaderType)
+				{
+					OutShaderInfos.Add({
+						ShaderInfo.VFType ? ShaderInfo.VFType->GetFName() : NAME_None,
+						ShaderType->GetFName()
+					});
+				}
+			}
+
 			for (FDebugShaderPipelineInfo const& PipelineInfo : ShaderInfo.Pipelines)
 			{
-				NumShaders += PipelineInfo.ShaderTypes.Num();
+				for (FShaderType* ShaderType : PipelineInfo.ShaderTypes)
+				{
+					if (ShaderType)
+					{
+						OutShaderInfos.Add({
+							ShaderInfo.VFType ? ShaderInfo.VFType->GetFName() : NAME_None,
+							ShaderType->GetFName()
+						});
+					}
+				}
 			}
 		}
 	}
 
-	return NumShaders;
+	return OutShaderInfos;
 }
 
 FMaterialStatistics UMaterialEditingLibrary::GetStatistics(class UMaterialInterface* Material)
diff --git a/Engine/Source/Editor/MaterialEditor/Public/MaterialEditingLibrary.h b/Engine/Source/Editor/MaterialEditor/Public/MaterialEditingLibrary.h
--- a/Engine/Source/Editor/MaterialEditor/Public/MaterialEditingLibrary.h
+++ b/Engine/Source/Editor/MaterialEditor/Public/MaterialEditingLibrary.h
@@ -55,6 +55,21 @@ struct FMaterialStatistics
 	int32 NumInterpolatorScalars = 0;
 };
 
+/** This struct represents a single shader in a material. */
+USTRUCT(BlueprintType)
+struct FDebugShaderInfo
+{
+	GENERATED_BODY()
+
+	/** Name of the vertex factory that this shader belongs to.  Could be NAME_None for material-only shaders. */
+	UPROPERTY(EditAnywhere, BlueprintReadWrite, Category=Statistics)
+	FName VertexFactoryName;
+
+	/** Name of the shader type for this shader */
+	UPROPERTY(EditAnywhere, BlueprintReadWrite, Category=Statistics)
+	FName ShaderTypeName;
+};
+
 /** Blueprint library for creating/editing Materials */
 UCLASS(MinimalAPI)
 class UMaterialEditingLibrary : public UBlueprintFunctionLibrary
@@ -455,6 +470,10 @@ public:
 	/** Returns the number of shader types (including pipeline shader types) for the given material. */
 	UFUNCTION(BlueprintCallable, Category = "MaterialEditing")
 	static UE_API int32 GetNumShaderTypes(UMaterialInterface* Material);
+
+	/** Returns a list of all the shaders contained in this material. */
+	UFUNCTION(BlueprintCallable, Category = "MaterialEditing")
+	static UE_API TArray<FDebugShaderInfo> ListShaders(UMaterialInterface* Material);
 	
 	/** Returns statistics about the given material */
 	UFUNCTION(BlueprintCallable, Category = "MaterialEditing")
