Added material usage editor blueprint  functions.
SetBaseMaterialUsage() supports both setting and unsetting usage flags on a UMaterial.
This should be preferred over SetMaterialUsage() which may be deprecated in future.
SetMaterialUsageOverride() supports setting/unsetting usage flag overrides on a UMaterialInstanceConstant.
Also HasMaterialUsage() now takes a UMaterialInterface instead of just a UMaterial.

diff --git a/Engine/Source/Editor/MaterialEditor/Private/MaterialEditingLibrary.cpp b/Engine/Source/Editor/MaterialEditor/Private/MaterialEditingLibrary.cpp
--- a/Engine/Source/Editor/MaterialEditor/Private/MaterialEditingLibrary.cpp
+++ b/Engine/Source/Editor/MaterialEditor/Private/MaterialEditingLibrary.cpp
@@ -661,7 +661,7 @@ bool UMaterialEditingLibrary::SetMaterialUsage(UMaterial* Material, EMaterialUsa
 	return bResult;
 }
 
-bool UMaterialEditingLibrary::HasMaterialUsage(UMaterial* Material, EMaterialUsage Usage)
+bool UMaterialEditingLibrary::HasMaterialUsage(UMaterialInterface* Material, EMaterialUsage Usage)
 {
 	bool bResult = false;
 	if (Material)
@@ -671,6 +671,84 @@ bool UMaterialEditingLibrary::HasMaterialUsage(UMaterial* Material, EMaterialUsa
 	return bResult;
 }
 
+void UMaterialEditingLibrary::SetBaseMaterialUsage(UMaterial* Material, EMaterialUsage Usage, bool bValue)
+{
+	if (Material == nullptr)
+	{
+		return;
+	}
+
+	if (bValue == Material->GetUsageByFlag(Usage))
+	{
+		return;
+	}
+
+	{
+		FMaterialUpdateContext UpdateContext;
+		Material->SetUsageByFlag(Usage, bValue);
+		RecompileMaterialInternal(UpdateContext, Material);
+	}
+
+	if (IMaterialEditor* MaterialEditor = MaterialEditingLibraryImpl::FindMaterialEditorForAsset(Material))
+	{
+		MaterialEditor->NotifyExternalMaterialChange();
+	}
+
+	RebuildMaterialInstanceEditors(Material);
+}
+
+void UMaterialEditingLibrary::SetMaterialUsageOverride(UMaterialInstanceConstant* MaterialInstance, EMaterialUsage Usage, bool bOverride, bool bValue)
+{
+	if (MaterialInstance == nullptr)
+	{
+		return;
+	}
+		
+	FMaterialInstanceBasePropertyOverrides BasePropertyOverrides = MaterialInstance->BasePropertyOverrides;
+		
+	if (bOverride)
+	{
+		BasePropertyOverrides.bOverride_UsageFlags |= 1u << Usage;
+	}
+	else
+	{
+		BasePropertyOverrides.bOverride_UsageFlags &= ~(1u << Usage);
+
+		// Get value from parent.		
+		if (UMaterialInterface* Parent = MaterialInstance->Parent)
+		{
+			bValue = Parent->GetUsageByFlag(Usage);
+		}
+	}
+
+	if (bValue)
+	{
+		BasePropertyOverrides.UsageFlags |= 1u << Usage;
+	}
+	else
+	{
+		BasePropertyOverrides.UsageFlags &= ~(1u << Usage);
+	}
+
+	if (BasePropertyOverrides.bOverride_UsageFlags == MaterialInstance->BasePropertyOverrides.bOverride_UsageFlags &&
+		BasePropertyOverrides.UsageFlags == MaterialInstance->BasePropertyOverrides.UsageFlags)
+	{
+		return;
+	}
+		
+	{
+		FMaterialInstanceParameterUpdateContext UpdateContext(MaterialInstance);
+		UpdateContext.SetBasePropertyOverrides(BasePropertyOverrides);
+	}
+			
+	MaterialInstance->MarkPackageDirty();
+		
+	if (UMaterial* Material = MaterialInstance->GetMaterial())
+	{
+		RebuildMaterialInstanceEditors(Material);
+	}
+}
+
 bool UMaterialEditingLibrary::ConnectMaterialProperty(UMaterialExpression* FromExpression, FString FromOutputName, EMaterialProperty Property)
 {
 	bool bResult = false;
diff --git a/Engine/Source/Editor/MaterialEditor/Public/MaterialEditingLibrary.h b/Engine/Source/Editor/MaterialEditor/Public/MaterialEditingLibrary.h
--- a/Engine/Source/Editor/MaterialEditor/Public/MaterialEditingLibrary.h
+++ b/Engine/Source/Editor/MaterialEditor/Public/MaterialEditingLibrary.h
@@ -155,7 +155,26 @@ public:
 	 *	@param	Usage				Usage type to check for this material
 	 */
 	UFUNCTION(BlueprintPure, Category = "MaterialEditing")
-	static UE_API bool HasMaterialUsage(UMaterial* Material, EMaterialUsage Usage);
+	static UE_API bool HasMaterialUsage(UMaterialInterface* Material, EMaterialUsage Usage);
+
+	/** 
+	 *	Enable or disable a particular usage for the supplied material
+	 *	@param	Material			Material to change usage for
+	 *	@param	Usage				Usage type
+	 *	@param	bValue				Whether to enable the usage
+	 */
+	UFUNCTION(BlueprintCallable, Category = "MaterialEditing", meta = (HidePin = "bNeedsRecompile"))
+	static UE_API void SetBaseMaterialUsage(UMaterial* Material, EMaterialUsage Usage, bool bValue);
+
+	/** 
+	 *	Set the material usage override for the supplied material instance
+	 *	@param	MaterialInstance	Material instance to change usage for
+	 *	@param	Usage				Usage type
+	 *	@param	bOverride			Whether to enable the override (set to false to default to Parent setting)
+	 *	@param	bValue				Whether to enable the usage (only if bOverride is true)
+	 */
+	UFUNCTION(BlueprintCallable, Category = "MaterialEditing")
+	static UE_API void SetMaterialUsageOverride(UMaterialInstanceConstant* MaterialInstance, EMaterialUsage Usage, bool bOverride, bool bValue);
 
 	/** 
 	 *	Connect a material expression output to one of the material property inputs (e.g. diffuse color, opacity etc)
