Fix open material details panel not updating after BP function enables/disables usage flags.
#rnx

diff --git a/Engine/Source/Editor/MaterialEditor/Private/MaterialEditingLibrary.cpp b/Engine/Source/Editor/MaterialEditor/Private/MaterialEditingLibrary.cpp
--- a/Engine/Source/Editor/MaterialEditor/Private/MaterialEditingLibrary.cpp
+++ b/Engine/Source/Editor/MaterialEditor/Private/MaterialEditingLibrary.cpp
@@ -689,10 +689,7 @@ void UMaterialEditingLibrary::SetBaseMaterialUsage(UMaterial* Material, EMateria
 		RecompileMaterialInternal(UpdateContext, Material);
 	}
 
-	if (IMaterialEditor* MaterialEditor = MaterialEditingLibraryImpl::FindMaterialEditorForAsset(Material))
-	{
-		MaterialEditor->NotifyExternalMaterialChange();
-	}
+	FEditorSupportDelegates::MaterialUsageFlagsChanged.Broadcast(Material, Usage);
 
 	RebuildMaterialInstanceEditors(Material);
 }
diff --git a/Engine/Source/Editor/MaterialEditor/Private/MaterialEditor.cpp b/Engine/Source/Editor/MaterialEditor/Private/MaterialEditor.cpp
--- a/Engine/Source/Editor/MaterialEditor/Private/MaterialEditor.cpp
+++ b/Engine/Source/Editor/MaterialEditor/Private/MaterialEditor.cpp
@@ -5478,8 +5478,8 @@ void FMaterialEditor::OnMaterialUsageFlagsChanged(UMaterial* MaterialThatChanged
 	EMaterialUsage Flag = static_cast<EMaterialUsage>(FlagThatChanged);
 	if(MaterialThatChanged == OriginalMaterial)
 	{
-		Material->SetMaterialUsage(Flag);
-		UpdateStatsMaterials();
+		Material->SetUsageByFlag(Flag, OriginalMaterial->GetUsageByFlag(Flag));
+		UpdatePreviewMaterial();
 	}
 }
 
