Deprecate SetMaterialUsage BP function.
Use SetBaseMaterialUsage instead.
#rnx

diff --git a/Engine/Source/Editor/MaterialEditor/Public/MaterialEditingLibrary.h b/Engine/Source/Editor/MaterialEditor/Public/MaterialEditingLibrary.h
--- a/Engine/Source/Editor/MaterialEditor/Public/MaterialEditingLibrary.h
+++ b/Engine/Source/Editor/MaterialEditor/Public/MaterialEditingLibrary.h
@@ -146,7 +146,7 @@ public:
 	 *	@param	Usage				New usage type to enable for this material
 	 *	@param	bNeedsRecompile		Returned to indicate if material needs recompiling after this change
 	 */
-	UFUNCTION(BlueprintCallable, Category = "MaterialEditing")
+	UFUNCTION(BlueprintCallable, Category = "MaterialEditing", meta = (DeprecatedFunction, DeprecationMessage = "Use SetBaseMaterialUsage instead"))
 	static UE_API bool SetMaterialUsage(UMaterial* Material, EMaterialUsage Usage, bool& bNeedsRecompile);
 
 	/**
@@ -164,7 +164,7 @@ public:
 	 *	@param	bValue				Whether to enable the usage
 	 */
 	UFUNCTION(BlueprintCallable, Category = "MaterialEditing", meta = (HidePin = "bNeedsRecompile"))
-	static UE_API void SetBaseMaterialUsage(UMaterial* Material, EMaterialUsage Usage, bool bValue);
+	static UE_API void SetBaseMaterialUsage(UMaterial* Material, EMaterialUsage Usage, bool bValue = true);
 
 	/** 
 	 *	Set the material usage override for the supplied material instance
@@ -174,7 +174,7 @@ public:
 	 *	@param	bValue				Whether to enable the usage (only if bOverride is true)
 	 */
 	UFUNCTION(BlueprintCallable, Category = "MaterialEditing")
-	static UE_API void SetMaterialUsageOverride(UMaterialInstanceConstant* MaterialInstance, EMaterialUsage Usage, bool bOverride, bool bValue);
+	static UE_API void SetMaterialUsageOverride(UMaterialInstanceConstant* MaterialInstance, EMaterialUsage Usage, bool bOverride = true, bool bValue = true);
 
 	/** 
 	 *	Connect a material expression output to one of the material property inputs (e.g. diffuse color, opacity etc)
