Fix bad metadata on UFUNCTION
#rnx

diff --git a/Engine/Source/Editor/MaterialEditor/Public/MaterialEditingLibrary.h b/Engine/Source/Editor/MaterialEditor/Public/MaterialEditingLibrary.h
--- a/Engine/Source/Editor/MaterialEditor/Public/MaterialEditingLibrary.h
+++ b/Engine/Source/Editor/MaterialEditor/Public/MaterialEditingLibrary.h
@@ -163,7 +163,7 @@ public:
 	 *	@param	Usage				Usage type
 	 *	@param	bValue				Whether to enable the usage
 	 */
-	UFUNCTION(BlueprintCallable, Category = "MaterialEditing", meta = (HidePin = "bNeedsRecompile"))
+	UFUNCTION(BlueprintCallable, Category = "MaterialEditing")
 	static UE_API void SetBaseMaterialUsage(UMaterial* Material, EMaterialUsage Usage, bool bValue = true);
 
 	/** 
