Add reset to default to MIC usage override UI
#rnx

diff --git a/Engine/Source/Editor/MaterialEditor/Private/MaterialEditorInstanceDetailCustomization.cpp b/Engine/Source/Editor/MaterialEditor/Private/MaterialEditorInstanceDetailCustomization.cpp
--- a/Engine/Source/Editor/MaterialEditor/Private/MaterialEditorInstanceDetailCustomization.cpp
+++ b/Engine/Source/Editor/MaterialEditor/Private/MaterialEditorInstanceDetailCustomization.cpp
@@ -1264,6 +1264,7 @@ void FMaterialInstanceParameterDetails::CreateUsageFlagOverrideWidgets(IDetailLa
 		const FText UsageToolTip = FText::FromString(UsageEnum->GetMetaData(TEXT("ToolTip"), UsageIndex));
 
 		Group.AddWidgetRow()
+		.PropertyHandleList({UsageProperty})
 		.NameContent()
 		[
 			SNew(STextBlock)
@@ -1318,6 +1319,27 @@ void FMaterialInstanceParameterDetails::CreateUsageFlagOverrideWidgets(IDetailLa
 				}
 				OverrideProperty->SetValue(Value);
 			})
+		)
+		.OverrideResetToDefault(FResetToDefaultOverride::Create(
+			FIsResetToDefaultVisible::CreateLambda([this, UsageIndex](TSharedPtr<IPropertyHandle>)
+			{
+				TObjectPtr<UMaterialInterface> Material = MaterialEditorInstance->GetMaterialInterface();
+				TObjectPtr<UMaterialInterface> ParentMaterial = MaterialEditorInstance->GetParentMaterialInterface();
+				if (!Material || !ParentMaterial)
+				{
+					return false;
+				}
+				return Material->GetUsageByFlag((EMaterialUsage)UsageIndex) != ParentMaterial->GetUsageByFlag((EMaterialUsage)UsageIndex);
+			}), 
+			FResetToDefaultHandler::CreateLambda([UsageProperty, UsageIndex](TSharedPtr<IPropertyHandle>)
+			{
+				// Simply toggling the bit is enough here since we can call this after the
+				// IsResetToDefaultVisible above returns that we are different to the parent
+				uint32 Value = 0;
+				UsageProperty->GetValue(Value);
+				Value ^= (1u << UsageIndex);
+				UsageProperty->SetValue(Value);
+			}))
 		);
 	}
 }
