Add HasMaterialUsageOverride BP function to MaterialEditorLibrary

diff --git a/Engine/Source/Editor/MaterialEditor/Private/MaterialEditingLibrary.cpp b/Engine/Source/Editor/MaterialEditor/Private/MaterialEditingLibrary.cpp
--- a/Engine/Source/Editor/MaterialEditor/Private/MaterialEditingLibrary.cpp
+++ b/Engine/Source/Editor/MaterialEditor/Private/MaterialEditingLibrary.cpp
@@ -663,12 +663,7 @@ bool UMaterialEditingLibrary::SetMaterialUsage(UMaterial* Material, EMaterialUsa
 
 bool UMaterialEditingLibrary::HasMaterialUsage(UMaterialInterface* Material, EMaterialUsage Usage)
 {
-	bool bResult = false;
-	if (Material)
-	{
-		bResult = Material->GetUsageByFlag(Usage);
-	}
-	return bResult;
+	return Material ? Material->GetUsageByFlag(Usage) : false;
 }
 
 void UMaterialEditingLibrary::SetBaseMaterialUsage(UMaterial* Material, EMaterialUsage Usage, bool bValue)
@@ -694,6 +689,11 @@ void UMaterialEditingLibrary::SetBaseMaterialUsage(UMaterial* Material, EMateria
 	RebuildMaterialInstanceEditors(Material);
 }
 
+bool UMaterialEditingLibrary::HasMaterialUsageOverride(UMaterialInstanceConstant* MaterialInstance, EMaterialUsage Usage)
+{
+	return MaterialInstance ? (MaterialInstance->BasePropertyOverrides.bOverride_UsageFlags & (1u << Usage)) != 0 : false;
+}
+
 void UMaterialEditingLibrary::SetMaterialUsageOverride(UMaterialInstanceConstant* MaterialInstance, EMaterialUsage Usage, bool bOverride, bool bValue)
 {
 	if (MaterialInstance == nullptr)
diff --git a/Engine/Source/Editor/MaterialEditor/Public/MaterialEditingLibrary.h b/Engine/Source/Editor/MaterialEditor/Public/MaterialEditingLibrary.h
--- a/Engine/Source/Editor/MaterialEditor/Public/MaterialEditingLibrary.h
+++ b/Engine/Source/Editor/MaterialEditor/Public/MaterialEditingLibrary.h
@@ -166,6 +166,14 @@ public:
 	UFUNCTION(BlueprintCallable, Category = "MaterialEditing")
 	static UE_API void SetBaseMaterialUsage(UMaterial* Material, EMaterialUsage Usage, bool bValue = true);
 
+	/**
+	 *	Check if a if the supplied material instance has an override for the particular usage
+	 *	@param	MaterialInstance	Material instance to check override for
+	 *	@param	Usage				Usage type to check for this material
+	 */
+	UFUNCTION(BlueprintPure, Category = "MaterialEditing")
+	static bool HasMaterialUsageOverride(UMaterialInstanceConstant* MaterialInstance, EMaterialUsage Usage);
+
 	/** 
 	 *	Set the material usage override for the supplied material instance
 	 *	@param	MaterialInstance	Material instance to change usage for
