Remove unwanted _MAX entry in material scalar param enum combo boxes.
#rnx

diff --git a/Engine/Source/Editor/MaterialEditor/Private/MaterialPropertyHelpers.cpp b/Engine/Source/Editor/MaterialEditor/Private/MaterialPropertyHelpers.cpp
--- a/Engine/Source/Editor/MaterialEditor/Private/MaterialPropertyHelpers.cpp
+++ b/Engine/Source/Editor/MaterialEditor/Private/MaterialPropertyHelpers.cpp
@@ -58,6 +58,7 @@
 #include "Curves/CurveLinearColorAtlas.h"
 #include "Curves/CurveLinearColor.h"
 #include "MaterialCachedData.h"
+#include "Engine/UserDefinedEnum.h"
 
 #include UE_INLINE_GENERATED_CPP_BY_NAME(MaterialPropertyHelpers)
 
@@ -1608,9 +1609,16 @@ TSharedRef<SWidget> FMaterialPropertyHelpers::CreateScalarEnumerationParameterVa
 		UObject* Enumeration = WeakEnumeration.Get();
 		if (UEnum* Enum = Cast<UEnum>(Enumeration))
 		{
-			for (int32 EnumIndex = 0; EnumIndex < Enum->NumEnums(); ++EnumIndex)
+			// UUserDefinedEnum always add a _MAX entry to the enum. We want to ignore that for populating the combo box.
+			// C++ UEnum usually use UMETA(Hidden) to hide any _MAX entry (and other entries too).
+			const bool bUserDefinedEnum = Enum->IsA<UUserDefinedEnum>();
+			const int32 NumEnums = bUserDefinedEnum ? Enum->NumEnums() - 1 : Enum->NumEnums();
+			for (int32 EnumIndex = 0; EnumIndex < NumEnums; ++EnumIndex)
 			{
-				OutStrings.Emplace(MakeShared<FString>(Enum->GetDisplayNameTextByIndex(EnumIndex).ToString()));
+				if (!Enum->HasMetaData(TEXT("Hidden"), EnumIndex))
+				{
+					OutStrings.Emplace(MakeShared<FString>(Enum->GetDisplayNameTextByIndex(EnumIndex).ToString()));
+				}
 			}
 		}
 		else if (IMaterialEnumerationProvider* EnumProvider = Cast<IMaterialEnumerationProvider>(Enumeration))
