[Substrate] Remove 'Beta' naming from Material Editor warning.

#rnx
#tests View Substrate tab in Material Editor with Substrate disabled

diff --git a/Engine/Source/Editor/MaterialEditor/Private/SMaterialEditorSubstrateWidget.cpp b/Engine/Source/Editor/MaterialEditor/Private/SMaterialEditorSubstrateWidget.cpp
--- a/Engine/Source/Editor/MaterialEditor/Private/SMaterialEditorSubstrateWidget.cpp
+++ b/Engine/Source/Editor/MaterialEditor/Private/SMaterialEditorSubstrateWidget.cpp
@@ -480,7 +480,7 @@ void SMaterialEditorSubstrateWidget::Construct(const FArguments& InArgs, TWeakPt
 					.ColorAndOpacity(FLinearColor::Yellow)
 					.ShadowColorAndOpacity(FLinearColor::Black)
 					.ShadowOffset(FVector2D::UnitVector)
-					.Text(LOCTEXT("SubstrateWidgetNotEnable", "Details cannot be shown: Substrate (Beta) is not enabled for this project (See the project settings window, rendering settings section)."))
+					.Text(LOCTEXT("SubstrateWidgetNotEnable", "Details cannot be shown: Substrate is not enabled for this project (See the project settings window, rendering settings section)."))
 				]
 			]
 		];
diff --git a/Engine/Source/Editor/MaterialEditor/Private/SMaterialEditorTopologyWidget.cpp b/Engine/Source/Editor/MaterialEditor/Private/SMaterialEditorTopologyWidget.cpp
--- a/Engine/Source/Editor/MaterialEditor/Private/SMaterialEditorTopologyWidget.cpp
+++ b/Engine/Source/Editor/MaterialEditor/Private/SMaterialEditorTopologyWidget.cpp
@@ -86,7 +86,7 @@ void SMaterialEditorTopologyWidget::Construct(const FArguments& InArgs, TWeakPtr
 					.ColorAndOpacity(FLinearColor::Yellow)
 					.ShadowColorAndOpacity(FLinearColor::Black)
 					.ShadowOffset(FVector2D::UnitVector)
-					.Text(LOCTEXT("SubstrateWidgetNotEnable", "Details cannot be shown: Substrate (Beta) is not enabled for this project (See the project settings window, rendering settings section)."))
+					.Text(LOCTEXT("SubstrateWidgetNotEnable", "Details cannot be shown: Substrate is not enabled for this project (See the project settings window, rendering settings section)."))
 				]
 			]
 		];
