Improve Material Instance Editor startup time by 20% by optimizing N^2 operation to NlogN
RefreshOnScreenTextureWarnings was doing a N^2 operation by going through all Material Expressions of a Material recursively(through MFs as well), for each parameter and comparing their Guids.
Optimize that by building a TMap Guid->MaterialExpression which is used for all lookups.

[RN] minor | Rendering Architecture

diff --git a/Engine/Source/Editor/MaterialEditor/Private/MaterialInstanceEditor.cpp b/Engine/Source/Editor/MaterialEditor/Private/MaterialInstanceEditor.cpp
--- a/Engine/Source/Editor/MaterialEditor/Private/MaterialInstanceEditor.cpp
+++ b/Engine/Source/Editor/MaterialEditor/Private/MaterialInstanceEditor.cpp
@@ -1479,6 +1479,9 @@ void FMaterialInstanceEditor::RefreshOnScreenTextureWarnings()
 			UEnum* MaterialTypeEnum = StaticEnum<ERuntimeVirtualTextureMaterialType>();
 			check(MaterialTypeEnum);
 
+			TMap<FGuid, TObjectPtr<UMaterialExpression>> MaterialGuids;
+			BaseMaterial->BuildExpressionParametersGuidsMap(MaterialGuids);
+
 			const int32 GroupCount = MaterialEditorInstance->ParameterGroups.Num();
 			for (int32 GroupIndex = 0; GroupIndex < GroupCount; ++GroupIndex)
 			{
@@ -1494,7 +1497,13 @@ void FMaterialInstanceEditor::RefreshOnScreenTextureWarnings()
 						if (Texture)
 						{
 							EMaterialSamplerType SamplerType = UMaterialExpressionTextureBase::GetSamplerTypeForTexture(Texture);
-							UMaterialExpressionTextureSampleParameter* Expression = BaseMaterial->FindExpressionByGUID<UMaterialExpressionTextureSampleParameter>(TextureParameterValue->ExpressionId);
+							UMaterialExpressionTextureSampleParameter* Expression = nullptr;
+
+							TObjectPtr<UMaterialExpression>* MapExpressions = MaterialGuids.Find(TextureParameterValue->ExpressionId);
+							if (MapExpressions)
+							{
+								Expression = Cast<UMaterialExpressionTextureSampleParameter>(*MapExpressions);
+							}
 
 							FString ErrorMessage;
 							if (Expression && !Expression->TextureIsValid(Texture, ErrorMessage))
@@ -1536,7 +1545,14 @@ void FMaterialInstanceEditor::RefreshOnScreenTextureWarnings()
 						MaterialEditorInstance->SourceInstance->GetRuntimeVirtualTextureParameterValue(RuntimeVirtualTextureParameterValue->ParameterInfo, RuntimeVirtualTexture);
 						if (RuntimeVirtualTexture)
 						{
-							UMaterialExpressionRuntimeVirtualTextureSampleParameter* Expression = BaseMaterial->FindExpressionByGUID<UMaterialExpressionRuntimeVirtualTextureSampleParameter>(RuntimeVirtualTextureParameterValue->ExpressionId);
+							UMaterialExpressionRuntimeVirtualTextureSampleParameter* Expression = nullptr;
+
+							TObjectPtr<UMaterialExpression>* MapExpression = MaterialGuids.Find(RuntimeVirtualTextureParameterValue->ExpressionId);
+							if (MapExpression)
+							{
+								Expression = Cast<UMaterialExpressionRuntimeVirtualTextureSampleParameter>(*MapExpression);
+							}
+
 							if (!Expression)
 							{
 								const FText ExpressionNameText = FText::Format(LOCTEXT("MissingRVTExpression", "Warning: Runtime Virtual Texture Expression {0} not found."), FText::FromName(RuntimeVirtualTextureParameterValue->ParameterInfo.Name));
@@ -1582,7 +1598,14 @@ void FMaterialInstanceEditor::RefreshOnScreenTextureWarnings()
 						MaterialEditorInstance->SourceInstance->GetSparseVolumeTextureParameterValue(SparseVolumeTextureParameterValue->ParameterInfo, SparseVolumeTexture);
 						if (SparseVolumeTexture)
 						{
-							UMaterialExpressionSparseVolumeTextureSampleParameter* Expression = BaseMaterial->FindExpressionByGUID<UMaterialExpressionSparseVolumeTextureSampleParameter>(SparseVolumeTextureParameterValue->ExpressionId);
+							UMaterialExpressionSparseVolumeTextureSampleParameter* Expression = nullptr;
+
+							TObjectPtr<UMaterialExpression>* MapExpression = MaterialGuids.Find(SparseVolumeTextureParameterValue->ExpressionId);
+							if (MapExpression)
+							{
+								Expression = Cast<UMaterialExpressionSparseVolumeTextureSampleParameter>(*MapExpression);
+							}
+
 							if (!Expression)
 							{
 								const FText ExpressionNameText = FText::Format(LOCTEXT("MissingSVTExpression", "Warning: Sparse Volume Texture Expression {0} not found."), FText::FromName(SparseVolumeTextureParameterValue->ParameterInfo.Name));
diff --git a/Engine/Source/Runtime/Engine/Private/Materials/Material.cpp b/Engine/Source/Runtime/Engine/Private/Materials/Material.cpp
--- a/Engine/Source/Runtime/Engine/Private/Materials/Material.cpp
+++ b/Engine/Source/Runtime/Engine/Private/Materials/Material.cpp
@@ -1946,6 +1946,51 @@ void UMaterial::GetDependentFunctions(TArray<UMaterialFunctionInterface*>& Depen
 		return true;
 	});
 }
+
+void BuildExpressionParametersGuidsMapRecursive(TConstArrayView<TObjectPtr<UMaterialExpression>> InMaterialExpression, TMap<FGuid, TObjectPtr<UMaterialExpression>>& OutExpressionMap)
+{
+	const int32 NumExpressions = InMaterialExpression.Num();
+	for (int32 ExpressionIndex = 0; ExpressionIndex < NumExpressions; ++ExpressionIndex)
+	{
+		UMaterialExpression* ExpressionPtr = InMaterialExpression[ExpressionIndex];
+		if (ExpressionPtr)
+		{
+			const FGuid ExpressionGuid = ExpressionPtr->GetParameterExpressionId();
+			if (ExpressionPtr->bIsParameterExpression && !OutExpressionMap.Contains(ExpressionGuid))
+			{
+				OutExpressionMap.Add(ExpressionGuid, TObjectPtr<UMaterialExpression>(ExpressionPtr));
+			}
+		}
+		else if (UMaterialFunctionInterface* MaterialFunction = UMaterial::GetExpressionFunctionPointer(ExpressionPtr))
+		{
+			BuildExpressionParametersGuidsMapRecursive(MaterialFunction->GetExpressions(), OutExpressionMap);
+		}
+		else if (TOptional<UMaterial::FLayersInterfaces> LayersInterfaces = UMaterial::GetExpressionLayers(ExpressionPtr))
+		{
+			for (const UMaterialFunctionInterface* Layer : LayersInterfaces->Layers)
+			{
+				if (Layer)
+				{
+					BuildExpressionParametersGuidsMapRecursive(Layer->GetExpressions(), OutExpressionMap);
+				}
+			}
+
+			for (const UMaterialFunctionInterface* Blend : LayersInterfaces->Blends)
+			{
+				if (Blend)
+				{
+					BuildExpressionParametersGuidsMapRecursive(Blend->GetExpressions(), OutExpressionMap);
+				}
+			}
+		}
+	}
+}
+
+void UMaterial::BuildExpressionParametersGuidsMap(TMap<FGuid, TObjectPtr<UMaterialExpression>>& OutExpressionMap)
+{
+	BuildExpressionParametersGuidsMapRecursive(GetExpressions(), OutExpressionMap);
+}
+
 #endif // WITH_EDITORONLY_DATA
 
 extern FPostProcessMaterialNode* FindExistingBlendablePostProcessNode(const FFinalPostProcessSettings& Dest, const UMaterialInterface* Material, const UMaterial* Base);
diff --git a/Engine/Source/Runtime/Engine/Public/Materials/Material.h b/Engine/Source/Runtime/Engine/Public/Materials/Material.h
--- a/Engine/Source/Runtime/Engine/Public/Materials/Material.h
+++ b/Engine/Source/Runtime/Engine/Public/Materials/Material.h
@@ -1552,6 +1552,9 @@ public:
 		return nullptr;
 	}
 
+	/** Builds a map with Guids as keys and Expressions as Values */
+	ENGINE_API void BuildExpressionParametersGuidsMap(TMap<FGuid, TObjectPtr<UMaterialExpression>>& OutExpressionMap);
+
 	/* Get all expressions of the requested type */
 	template<typename ExpressionType>
 	void GetAllExpressionsOfType(TArray<const ExpressionType*>& OutExpressions) const
