Material permutation validation: add validation for material hierarchies touched by changed material functions.
#rnx
#Synced-CL 48340728

diff --git a/Engine/Plugins/Editor/MaterialValidation/Source/Private/MaterialValidators.cpp b/Engine/Plugins/Editor/MaterialValidation/Source/Private/MaterialValidators.cpp
--- a/Engine/Plugins/Editor/MaterialValidation/Source/Private/MaterialValidators.cpp
+++ b/Engine/Plugins/Editor/MaterialValidation/Source/Private/MaterialValidators.cpp
@@ -93,7 +93,17 @@ EDataValidationResult UEditorValidator_MaterialPermutation::ValidateLoadedAsset_
 
 	if (MaterialFunction)
 	{
-		//todo: Add UMaterial refs to BaseMaterials. Note this will require loading assets, but is probably OK.
+		// Add the base materials that use the material function for later valiation.
+		// Note this will require loading UMaterial assets which could be slow for very commonly used material functions.
+		TArray<FAssetData> MaterialAssetDatas;
+		UMaterialEditingLibrary::GetMaterialsReferencingFunction(MaterialFunction, MaterialAssetDatas);
+		for (FAssetData const& AssetData : MaterialAssetDatas)
+		{
+			if (UMaterial* BaseMaterial = Cast<UMaterial>(AssetData.GetAsset()))
+			{
+				BaseMaterials.AddUnique(BaseMaterial);
+			}
+		}
 	}
 
 	if (GetValidationResult() != EDataValidationResult::Invalid)
diff --git a/Engine/Source/Editor/MaterialEditor/Private/MaterialEditingLibrary.cpp b/Engine/Source/Editor/MaterialEditor/Private/MaterialEditingLibrary.cpp
--- a/Engine/Source/Editor/MaterialEditor/Private/MaterialEditingLibrary.cpp
+++ b/Engine/Source/Editor/MaterialEditor/Private/MaterialEditingLibrary.cpp
@@ -1406,6 +1406,63 @@ void UMaterialEditingLibrary::GetChildInstances(UMaterialInterface* Parent, TArr
 	}
 }
 
+void UMaterialEditingLibrary::GetMaterialsReferencingFunction(UMaterialFunction* InFunction, TArray<FAssetData>& OutMaterials)
+{
+	if (InFunction == nullptr)
+	{
+		return;
+	}
+
+	FAssetRegistryModule& AssetRegistryModule = FModuleManager::LoadModuleChecked<FAssetRegistryModule>(TEXT("AssetRegistry"));
+	IAssetRegistry& AssetRegistry = AssetRegistryModule.Get();
+
+	const FAssetData FunctionAssetData = AssetRegistry.GetAssetByObjectPath(InFunction->GetPathName());
+	if (!FunctionAssetData.IsValid())
+	{
+		return;
+	}
+
+	TArray<FName> Visited;
+	TArray<FName> ToVisit;
+	ToVisit.Push(FunctionAssetData.PackageName);
+
+	while (!ToVisit.IsEmpty())
+	{
+		const FName FunctionPackageName = ToVisit.Pop(EAllowShrinking::No);
+
+		TArray<FName> ReferencerPackageNames;
+		AssetRegistry.GetReferencers(FunctionPackageName, ReferencerPackageNames);
+
+		for (const FName& PackageName : ReferencerPackageNames)
+		{
+			if (Visited.Contains(PackageName))
+			{
+				continue;
+			}
+			Visited.Add(PackageName);
+
+			TArray<FAssetData> AssetDatas;
+			AssetRegistry.GetAssetsByPackageName(PackageName, AssetDatas, true);
+
+			for (FAssetData const& AssetData : AssetDatas)
+			{
+				if (!AssetData.IsValid())
+				{
+					continue;
+				}
+				else if (AssetData.IsInstanceOf(UMaterial::StaticClass()))
+				{
+					OutMaterials.Add(AssetData);
+				}
+				else if (AssetData.IsInstanceOf(UMaterialFunction::StaticClass()))
+				{
+					ToVisit.Push(PackageName);
+				}
+			}
+		}
+	}
+}
+
 void UMaterialEditingLibrary::GetScalarParameterNames(UMaterialInterface* Material, TArray<FName>& ParameterNames)
 {
 	ParameterNames.Empty();
diff --git a/Engine/Source/Editor/MaterialEditor/Public/MaterialEditingLibrary.h b/Engine/Source/Editor/MaterialEditor/Public/MaterialEditingLibrary.h
--- a/Engine/Source/Editor/MaterialEditor/Public/MaterialEditingLibrary.h
+++ b/Engine/Source/Editor/MaterialEditor/Public/MaterialEditingLibrary.h
@@ -394,6 +394,10 @@ public:
 	UFUNCTION(BlueprintCallable, Category = "MaterialEditing")
 	static UE_API void GetChildInstances(UMaterialInterface* Parent, TArray<FAssetData>& ChildInstances);
 
+	/** Gets all materials that use a material function */
+	UFUNCTION(BlueprintCallable, Category = "MaterialEditing")
+	static UE_API void GetMaterialsReferencingFunction(UMaterialFunction* InFunction, TArray<FAssetData>& OutMaterials);
+
 	/** Gets all scalar parameter names */
 	UFUNCTION(BlueprintCallable, Category = "MaterialEditing")
 	static UE_API void GetScalarParameterNames(UMaterialInterface* Material, TArray<FName>& ParameterNames);
