[MaterialEditor] Remove stats from the stats panel, only display errors, all the stats are available in the PlatformStats

#Synced-CL 47464230

diff --git a/Engine/Source/Editor/MaterialEditor/Private/MaterialEditor.cpp b/Engine/Source/Editor/MaterialEditor/Private/MaterialEditor.cpp
--- a/Engine/Source/Editor/MaterialEditor/Private/MaterialEditor.cpp
+++ b/Engine/Source/Editor/MaterialEditor/Private/MaterialEditor.cpp
@@ -1851,8 +1851,8 @@ FLinearColor FMaterialEditor::GetWorldCentricTabColorScale() const
 
 void FMaterialEditor::Tick( float InDeltaTime )
 {
-	UpdateMaterialInfoList();
-	UpdateMaterialinfoList_Old();
+	UpdateMaterialPlatformStats();
+	UpdateMaterialInfoLog();
 	UpdateGraphNodeStates();
 }
 
@@ -2919,7 +2919,7 @@ void FMaterialEditor::OnMessageLogLinkActivated(const class TSharedRef<IMessageT
 	}
 }
 
-void FMaterialEditor::UpdateMaterialinfoList_Old()
+void FMaterialEditor::UpdateMaterialInfoLog()
 {
 	bool bForceDisplay = false;
 	TArray< TSharedRef<class FTokenizedMessage> > Messages;
@@ -2972,171 +2972,6 @@ void FMaterialEditor::UpdateMaterialinfoList_Old()
 					}
 				}
 			}
-
-			for (int32 InstructionIndex = 0; InstructionIndex < Results.Num(); InstructionIndex++)
-			{
-				FString InstructionCountString = FString::Printf(TEXT("%s: %u instructions\nStats: %s"), *Results[InstructionIndex].ShaderDescription, Results[InstructionIndex].InstructionCount, *Results[InstructionIndex].ShaderStatisticsString);
-				if (bBuiltinStats)
-				{
-					InstructionCountString += FString::Printf(TEXT(" - Built-in instructions: %u\nStats: %s"), EmptyMaterialResults[InstructionIndex].InstructionCount, *EmptyMaterialResults[InstructionIndex].ShaderStatisticsString);
-				}
-				TempMaterialInfoList.Add(MakeShareable(new FMaterialInfo(InstructionCountString, FLinearColor::Yellow)));
-				TSharedRef<FTokenizedMessage> Line = FTokenizedMessage::Create(EMessageSeverity::Info);
-				Line->AddToken(FTextToken::Create(FText::FromString(InstructionCountString)));
-				Messages.Add(Line);
-			}
-
-			// Display the number of samplers used by the material.
-			const int32 SamplersUsed = MaterialResource->GetSamplerUsage();
-
-			if (SamplersUsed >= 0)
-			{
-				int32 MaxSamplers = GetExpectedFeatureLevelMaxTextureSamplers(MaterialResource->GetFeatureLevel());
-				FString SamplersString = FString::Printf(TEXT("%s samplers: %u/%u"), FeatureLevel <= ERHIFeatureLevel::ES3_1 ? TEXT("Mobile texture") : TEXT("Texture"), SamplersUsed, MaxSamplers);
-				TempMaterialInfoList.Add(MakeShareable(new FMaterialInfo(SamplersString, FLinearColor::Yellow)));
-				TSharedRef<FTokenizedMessage> Line = FTokenizedMessage::Create( EMessageSeverity::Info );
-				Line->AddToken( FTextToken::Create( FText::FromString( SamplersString ) ) );
-				Messages.Add(Line);
-			}
-				
-			// Display estimated texture look-up/sample counts
-			uint32 NumVSTextureSamples = 0, NumPSTextureSamples = 0;
-			MaterialResource->GetEstimatedNumTextureSamples(NumVSTextureSamples, NumPSTextureSamples);
-
-			if (NumVSTextureSamples > 0 || NumPSTextureSamples > 0)
-			{
-				FString SamplesString = FString::Printf(TEXT("Texture Lookups (Est.): VS(%u), PS(%u)"), NumVSTextureSamples, NumPSTextureSamples);
-
-				TempMaterialInfoList.Add(MakeShareable(new FMaterialInfo(SamplesString, FLinearColor::Yellow)));
-				TSharedRef<FTokenizedMessage> Line = FTokenizedMessage::Create(EMessageSeverity::Info);
-				Line->AddToken(FTextToken::Create(FText::FromString(SamplesString)));
-				Messages.Add(Line);
-			}
-
-			// Display estimated virtual texture look-up counts
-			uint32 NumVirtualTextureLookups = MaterialResource->GetEstimatedNumVirtualTextureLookups();
-			if (NumVirtualTextureLookups > 0)
-			{
-				FString LookupsString = FString::Printf(TEXT("Virtual Texture Lookups (Est.): %u"), NumVirtualTextureLookups);
-
-				TempMaterialInfoList.Add(MakeShareable(new FMaterialInfo(LookupsString, FLinearColor::Yellow)));
-				TSharedRef<FTokenizedMessage> Line = FTokenizedMessage::Create(EMessageSeverity::Info);
-				Line->AddToken(FTextToken::Create(FText::FromString(LookupsString)));
-				Messages.Add(Line);
-			}
-
-			const uint32 NumVirtualTextureStacks = MaterialResource->GetNumVirtualTextureStacks();
-			if (NumVirtualTextureStacks > 0u)
-			{
-				FString VTString = FString::Printf(TEXT("Virtual Texture Stacks: %u"), NumVirtualTextureStacks);
-
-				TempMaterialInfoList.Add(MakeShareable(new FMaterialInfo(VTString, FLinearColor::Yellow)));
-				TSharedRef<FTokenizedMessage> Line = FTokenizedMessage::Create(EMessageSeverity::Info);
-				Line->AddToken(FTextToken::Create(FText::FromString(VTString)));
-				Messages.Add(Line);
-			}
-
-			// Display the number of custom/user interpolators used by the material.
-			uint32 UVScalarsUsed, CustomInterpolatorScalarsUsed;
-			MaterialResource->GetUserInterpolatorUsage(UVScalarsUsed, CustomInterpolatorScalarsUsed);
-
-			if (UVScalarsUsed > 0 || CustomInterpolatorScalarsUsed > 0)
-			{
-				uint32 TotalScalars = UVScalarsUsed + CustomInterpolatorScalarsUsed;
-				uint32 MaxScalars = FMath::DivideAndRoundUp(TotalScalars, 4u) * 4;
-
-				FString InterpolatorsString = FString::Printf(TEXT("User interpolators: %u/%u Scalars (%u/4 Vectors) (TexCoords: %i, Custom: %i)"),
-															  TotalScalars, MaxScalars, MaxScalars / 4, UVScalarsUsed, CustomInterpolatorScalarsUsed);
-
-				TempMaterialInfoList.Add(MakeShareable(new FMaterialInfo(InterpolatorsString, FLinearColor::Yellow)));
-				TSharedRef<FTokenizedMessage> Line = FTokenizedMessage::Create( EMessageSeverity::Info );
-				Line->AddToken(FTextToken::Create(FText::FromString(InterpolatorsString)));
-				Messages.Add(Line);
-			}
-
-			FMaterialResource::FLWCUsagesArray LWCFuncUsagesVS;
-			FMaterialResource::FLWCUsagesArray LWCFuncUsagesPS;
-			FMaterialResource::FLWCUsagesArray LWCFuncUsagesCS;
-			MaterialResource->GetEstimatedLWCFuncUsages(LWCFuncUsagesVS, LWCFuncUsagesPS, LWCFuncUsagesCS);
-			for (int KindIndex = 0; KindIndex < (int)ELWCFunctionKind::Max; ++KindIndex)
-			{
-				int Usages = LWCFuncUsagesVS[KindIndex] + LWCFuncUsagesPS[KindIndex] + LWCFuncUsagesCS[KindIndex];
-				if (Usages > 0)
-				{
-					FString Message = FString::Printf(
-						TEXT("LWC %s usages (Est.): %u (VS), %u (PS), %u (CS)"), 
-						*UEnum::GetDisplayValueAsText((ELWCFunctionKind)KindIndex).ToString(), 
-						LWCFuncUsagesVS[KindIndex],
-						LWCFuncUsagesPS[KindIndex],
-						LWCFuncUsagesCS[KindIndex]);
-						
-					TempMaterialInfoList.Add(MakeShareable(new FMaterialInfo(Message, FLinearColor::Yellow)));
-					TSharedRef<FTokenizedMessage> Line = FTokenizedMessage::Create(EMessageSeverity::Info);
-					Line->AddToken(FTextToken::Create(FText::FromString(Message)));
-					Messages.Add(Line);
-				}
-			}
-
-			if (FMaterialShaderMap* ShaderMap = MaterialResource->GetGameThreadShaderMap())
-			{
-				// Add shader count
-				FString ShaderCountString = FString::Printf(TEXT("Shader Count: %u"), ShaderMap->GetShaderNum());
-				TSharedRef<FTokenizedMessage> ShaderCountLine = FTokenizedMessage::Create(EMessageSeverity::Info);
-				ShaderCountLine->AddToken(FTextToken::Create(FText::FromString(ShaderCountString)));
-				Messages.Add(ShaderCountLine);
-
-				// Add number of preshaders and stats
-				uint32 TotalParams, TotalOps;
-				MaterialResource->GetPreshaderStats(TotalParams, TotalOps);
-				FString PreshaderCountString = FString::Printf(TEXT("Preshaders: %u  (%u param fetches, %u ops)"), ShaderMap->GetNumPreshaders(), TotalParams, TotalOps);
-				TSharedRef<FTokenizedMessage> PreshaderCountLine = FTokenizedMessage::Create(EMessageSeverity::Info);
-				PreshaderCountLine->AddToken(FTextToken::Create(FText::FromString(PreshaderCountString)));
-				Messages.Add(PreshaderCountLine);
-
-				const EMaterialDomain Domain = MaterialResource->GetMaterialDomain();
-				if (Domain == MD_LightFunction || Domain == MD_PostProcess)
-				{
-					if(Domain == MD_LightFunction)
-					{
-						const bool bIsCompatibleWithLightFunctionAtlas = ShaderMap->IsLightFunctionAtlasCompatible();
-						FString LightFunctionAtlasStr = FString::Printf(TEXT("Light function material%s compatible with the light function atlas for fast batched deferred light shading."), bIsCompatibleWithLightFunctionAtlas ? TEXT(" IS") : TEXT(" IS NOT"));
-						TSharedRef<FTokenizedMessage> LightFunctionAtlasCountLine = FTokenizedMessage::Create(EMessageSeverity::Info);
-						LightFunctionAtlasCountLine->AddToken(FTextToken::Create(FText::FromString(LightFunctionAtlasStr)));
-						Messages.Add(LightFunctionAtlasCountLine);
-					}
-
-					for (UMaterialExpression* MaterialExpression : Material->GetExpressions())
-					{
-						if (MaterialExpression->IsA(UMaterialExpressionObjectPositionWS::StaticClass()))
-						{
-							FString PrependString;
-							switch (Domain)
-							{
-								case MD_LightFunction:
-								{
-									PrependString += "Although Light Functions are compatible with the ObjectPosition node, Light atlases";
-									break;
-								}
-								case MD_PostProcess:
-								{
-									PrependString += "Post Process Materials";
-									break;
-								}
-								default:break;
-							}
-							if (PrependString.IsEmpty())
-							{
-								break;
-							}
-							FString ObjectPositionWarning = FString::Printf(TEXT("Note: %s cannot resolve position from the ObjectPosition node, will always return 0."), *PrependString);
-							TSharedRef<FTokenizedMessage> LightFunctionAtlasCountLine = FTokenizedMessage::Create(EMessageSeverity::Info);
-							LightFunctionAtlasCountLine->AddToken(FTextToken::Create(FText::FromString(ObjectPositionWarning)));
-							Messages.Add(LightFunctionAtlasCountLine);
-							break;
-						}
-					}
-				}
-			}
 		}
 
 		FString FeatureLevelName;
@@ -3215,7 +3050,7 @@ void FMaterialEditor::UpdateMaterialinfoList_Old()
 	}
 }
 
-void FMaterialEditor::UpdateMaterialInfoList()
+void FMaterialEditor::UpdateMaterialPlatformStats()
 {
 	// setup stats widget visibility
 	if (MaterialFunction)
@@ -3259,29 +3094,29 @@ void FMaterialEditor::UpdateMaterialInfoList()
 	{
 		TSharedPtr<SWidget> TitleBar = FocusedGraphEd->GetTitleBar();
 		TSharedPtr<SMaterialEditorTitleBar> MaterialTitleBar = StaticCastSharedPtr<SMaterialEditorTitleBar>(TitleBar);
-	if (NewErrorHash != MaterialErrorHash)
-	{
-		MaterialErrorHash = NewErrorHash;
-		MaterialInfoList.Reset();
-
-		TArray< TSharedRef<class FTokenizedMessage> > Messages;
-		FString FeatureLevelName;
-		GetFeatureLevelName(GMaxRHIFeatureLevel, FeatureLevelName);
-		for (int32 ErrorIndex = 0; ErrorIndex < CompileErrors.Num(); ErrorIndex++)
+		if (NewErrorHash != MaterialErrorHash)
 		{
-			FString ErrorString = FString::Printf(TEXT("[%s] %s"), *FeatureLevelName, *CompileErrors[ErrorIndex]);
-			TSharedRef<FTokenizedMessage> Line = FTokenizedMessage::Create(EMessageSeverity::Error);
-			Line->AddToken(FTextToken::Create(FText::FromString(ErrorString)));
-			Messages.Add(Line);
+			MaterialErrorHash = NewErrorHash;
+			MaterialInfoList.Reset();
 
-			MaterialInfoList.Add(MakeShareable(new FMaterialInfo(ErrorString, FLinearColor::Red)));
-		}
+			TArray< TSharedRef<class FTokenizedMessage> > Messages;
+			FString FeatureLevelName;
+			GetFeatureLevelName(GMaxRHIFeatureLevel, FeatureLevelName);
+			for (int32 ErrorIndex = 0; ErrorIndex < CompileErrors.Num(); ErrorIndex++)
+			{
+				FString ErrorString = FString::Printf(TEXT("[%s] %s"), *FeatureLevelName, *CompileErrors[ErrorIndex]);
+				TSharedRef<FTokenizedMessage> Line = FTokenizedMessage::Create(EMessageSeverity::Error);
+				Line->AddToken(FTextToken::Create(FText::FromString(ErrorString)));
+				Messages.Add(Line);
 
-		StatsListing->ClearMessages();
-		StatsListing->AddMessages(Messages);
+				MaterialInfoList.Add(MakeShareable(new FMaterialInfo(ErrorString, FLinearColor::Red)));
+			}
 
-			MaterialTitleBar->RequestRefresh();
-	}
+			StatsListing->ClearMessages();
+			StatsListing->AddMessages(Messages);
+
+				MaterialTitleBar->RequestRefresh();
+		}
 
 		if (MaterialTitleBar->MaterialInfoList)
 		{
diff --git a/Engine/Source/Editor/MaterialEditor/Private/MaterialEditor.h b/Engine/Source/Editor/MaterialEditor/Private/MaterialEditor.h
--- a/Engine/Source/Editor/MaterialEditor/Private/MaterialEditor.h
+++ b/Engine/Source/Editor/MaterialEditor/Private/MaterialEditor.h
@@ -314,8 +314,8 @@ public:
 	/**
 	 * Updates list of Material Info used to show stats
 	 */
-	void UpdateMaterialInfoList();
-	void UpdateMaterialinfoList_Old();
+	void UpdateMaterialInfoLog();
+	void UpdateMaterialPlatformStats();
 
 	/**
 	 * Updates flags on the Material Nodes to avoid expensive look up calls when rendering
