Display Total Shaders as a SNotificationBackground as an overlay in the bottom right hand corner of the Material Editor and Material Instance Editor Viewport.

- The goal here is prominently display the number of shaders that this material has.  That way as a user modifies the material they can see live the impact of their changes when changing usage flags, or plugging in WPO.
- SNotificationBackground is an overlay in the bottom right corner of the material viewport.
- UpdateMaterialInfoList is called every tick, and we quickly query the material resource to estimate the number of shaders the material would have.  This number matches the PlatformStats panel.
- Since the Material Stats owns the calculation of the shader count this code can be leveraged in the Material Instance Editor.
- In the Material Editor we use the preview quality level for querying the shader count.
- Don't display this widget if we are editing a Material Function.

#Synced-CL 48372222

diff --git a/Engine/Source/Editor/MaterialEditor/Private/MaterialEditor.cpp b/Engine/Source/Editor/MaterialEditor/Private/MaterialEditor.cpp
--- a/Engine/Source/Editor/MaterialEditor/Private/MaterialEditor.cpp
+++ b/Engine/Source/Editor/MaterialEditor/Private/MaterialEditor.cpp
@@ -16,6 +16,7 @@
 #include "Widgets/Input/SButton.h"
 #include "Widgets/Input/SNumericEntryBox.h"
 #include "Widgets/Input/SCheckBox.h"
+#include "Widgets/Notifications/SNotificationBackground.h"
 #include "Styling/AppStyle.h"
 #include "EdGraph/EdGraph.h"
 #include "WorkflowOrientedApp/WorkflowUObjectDocuments.h"
@@ -1569,6 +1570,11 @@ bool FMaterialEditor::IsQualityPreviewChecked(EMaterialQualityLevel::Type TestQu
 	return NodeQualityLevel == TestQuality;
 }
 
+EMaterialQualityLevel::Type FMaterialEditor::GetQualityPreviewLevel() const
+{
+	return NodeQualityLevel == EMaterialQualityLevel::Num ? GetCachedScalabilityCVars().MaterialQualityLevel : NodeQualityLevel;
+}
+
 void FMaterialEditor::SetFeaturePreview(ERHIFeatureLevel::Type NewFeatureLevel)
 {
 	NodeFeatureLevel = NewFeatureLevel;
@@ -3133,6 +3139,12 @@ void FMaterialEditor::UpdateMaterialPlatformStats()
 	MaterialStatsManager->SetMaterial(MaterialForStats, bStatsFromPreviewMaterial ? DerivedMaterialInstances : OriginalDerivedMaterialInstances);
 	MaterialStatsManager->Update();
 
+	// Update the num shaders UI
+	if (NumShadersText.IsValid())
+	{
+		NumShadersText->SetText(MaterialStatsManager->GetShaderCountText(GMaxRHIShaderPlatform, GetQualityPreviewLevel(), 0));
+	}
+
 	// check if any derived instances fail, if so show the window
 	if (MaterialStatsManager->AnyNewCompilationErrors(1))
 	{
@@ -5455,8 +5467,27 @@ TSharedRef<SDockTab> FMaterialEditor::SpawnTab_Preview(const FSpawnTabArgs& Args
 			[
 				PreviewUIViewport.ToSharedRef()
 			]
+			+ SOverlay::Slot()
+			.Padding(8)
+			.VAlign(VAlign_Bottom)
+			.HAlign(HAlign_Right)
+			[
+				SNew(SNotificationBackground)
+				[
+					SNew(SHorizontalBox)
+					+ SHorizontalBox::Slot()
+					.Padding(6.f)
+					.VAlign(VAlign_Center)
+					.AutoWidth()
+					[
+						SAssignNew(NumShadersText, STextBlock)
+					]
+				]
+			]
 		];
 
+	NumShadersText->SetVisibility(MaterialFunction ? EVisibility::Collapsed : EVisibility::Visible);
+
 	PreviewViewport->OnAddedToTab( SpawnedTab );
 
 	return SpawnedTab;
diff --git a/Engine/Source/Editor/MaterialEditor/Private/MaterialEditor.h b/Engine/Source/Editor/MaterialEditor/Private/MaterialEditor.h
--- a/Engine/Source/Editor/MaterialEditor/Private/MaterialEditor.h
+++ b/Engine/Source/Editor/MaterialEditor/Private/MaterialEditor.h
@@ -729,6 +729,7 @@ private:
 	/** Helper functions for the quality level node display toggling */
 	void SetQualityPreview(EMaterialQualityLevel::Type NewQuality);
 	bool IsQualityPreviewChecked(EMaterialQualityLevel::Type TestQuality);
+	EMaterialQualityLevel::Type GetQualityPreviewLevel() const;
 
 	/** Helper functions for the feature level node display toggling */
 	void SetFeaturePreview(ERHIFeatureLevel::Type NewFeatureLevel);
@@ -1152,4 +1153,7 @@ private:
 	/** Delegate handles for FWorldDelegates registrations. */
 	FDelegateHandle OnWorldCleanupDelegateHandle;
 	FDelegateHandle OnPostWorldCleanupDelegateHandle;
+
+	/** The text for the number of shaders in the material. */
+	TSharedPtr<STextBlock> NumShadersText;
 };
diff --git a/Engine/Source/Editor/MaterialEditor/Private/MaterialInstanceEditor.cpp b/Engine/Source/Editor/MaterialEditor/Private/MaterialInstanceEditor.cpp
--- a/Engine/Source/Editor/MaterialEditor/Private/MaterialInstanceEditor.cpp
+++ b/Engine/Source/Editor/MaterialEditor/Private/MaterialInstanceEditor.cpp
@@ -65,6 +65,7 @@
 #include "VT/RuntimeVirtualTexture.h"
 #include "SparseVolumeTexture/SparseVolumeTexture.h"
 #include "Widgets/Input/SButton.h"
+#include "Widgets/Notifications/SNotificationBackground.h"
 #include "Subsystems/AssetEditorSubsystem.h"
 #include "SMaterialEditorSubstrateWidget.h"
 #include "ViewportToolbar/UnrealEdViewportToolbar.h"
@@ -78,7 +79,7 @@ DEFINE_LOG_CATEGORY_STATIC(LogMaterialInstanceEditor, Log, All);
 static TAutoConsoleVariable<bool> CVarEnableTextureValidationInMaterialInstanceEditor(
 	TEXT("r.MaterialEditor.EnableTextureValidation"),
 	true,
-	TEXT("Enables expensive texture validation in the material instance editor.  On by default.  Disable this to improve performance when changing texture paramters on large materials."),
+	TEXT("Enables expensive texture validation in the material instance editor.  On by default.  Disable this to improve performance when changing texture parameters on large materials."),
 	ECVF_Default
 );
 
@@ -1045,6 +1046,23 @@ TSharedRef<SDockTab> FMaterialInstanceEditor::SpawnTab_Preview( const FSpawnTabA
 			[
 				PreviewUIViewport.ToSharedRef()
 			]
+			+ SOverlay::Slot()
+			.Padding(8)
+			.VAlign(VAlign_Bottom)
+			.HAlign(HAlign_Right)
+			[
+				SNew(SNotificationBackground)
+				[
+					SNew(SHorizontalBox)
+					+ SHorizontalBox::Slot()
+					.Padding(6.f)
+					.VAlign(VAlign_Center)
+					.AutoWidth()
+					[
+						SAssignNew(NumShadersPillText, STextBlock)
+					]
+				]
+			]
 		];
 
 	PreviewVC->OnAddedToTab( SpawnedTab );
@@ -1674,6 +1692,19 @@ void FMaterialInstanceEditor::Tick(float DeltaTime)
 {
 	MaterialStatsManager->SetMaterial(MaterialEditorInstance->SourceInstance);
 	MaterialStatsManager->Update();
+
+	if (NumShadersPillText.IsValid())
+	{
+		// Update the Num Shaders text in the viewport.
+		if (MaterialEditorInstance->SourceInstance->bHasStaticPermutationResource)
+		{
+			NumShadersPillText->SetText(MaterialStatsManager->GetShaderCountText(GMaxRHIShaderPlatform, GetCachedScalabilityCVars().MaterialQualityLevel, 0));
+		}
+		else
+		{
+			NumShadersPillText->SetText(NSLOCTEXT("MaterialInstanceEditor", "NumShadersText", "Shaders: 0"));
+		}
+	}
 }
 
 TStatId FMaterialInstanceEditor::GetStatId() const
diff --git a/Engine/Source/Editor/MaterialEditor/Private/MaterialInstanceEditor.h b/Engine/Source/Editor/MaterialEditor/Private/MaterialInstanceEditor.h
--- a/Engine/Source/Editor/MaterialEditor/Private/MaterialInstanceEditor.h
+++ b/Engine/Source/Editor/MaterialEditor/Private/MaterialInstanceEditor.h
@@ -282,7 +282,6 @@ private:
 	TSharedPtr<FExtensibilityManager> MenuExtensibilityManager;
 	TSharedPtr<FExtensibilityManager> ToolBarExtensibilityManager;
 
-	/** */
 	UMaterialFunctionInstance*	MaterialFunctionOriginal;
 	UMaterialFunctionInstance*	MaterialFunctionInstance;
 	UMaterial*					FunctionMaterialProxy; 
@@ -290,4 +289,7 @@ private:
 
 	/** Object used as material statistics manager */
 	TSharedPtr<class FMaterialStats> MaterialStatsManager;
+
+	/** The text for the number of shaders in the material. */
+	TSharedPtr<STextBlock> NumShadersPillText;
 };
diff --git a/Engine/Source/Editor/MaterialEditor/Private/MaterialStats.cpp b/Engine/Source/Editor/MaterialEditor/Private/MaterialStats.cpp
--- a/Engine/Source/Editor/MaterialEditor/Private/MaterialStats.cpp
+++ b/Engine/Source/Editor/MaterialEditor/Private/MaterialStats.cpp
@@ -106,7 +106,7 @@ void FShaderPlatformSettings::OnShaderViewComboSelectionChanged(TSharedPtr<FMate
 	}
 }
 
-FText FShaderPlatformSettings::GetShaderCode(const EMaterialQualityLevel::Type QualityType, const int32 InstanceIndex)
+FText FShaderPlatformSettings::GetShaderCodeText(const EMaterialQualityLevel::Type QualityType, const int32 InstanceIndex)
 {
 	if (PlatformData[QualityType].Instances.IsEmpty())
 	{
@@ -152,6 +152,42 @@ FText FShaderPlatformSettings::GetShaderCode(const EMaterialQualityLevel::Type Q
 	return Instance.ShaderCode;
 }
 
+FText FShaderPlatformSettings::GetShaderCountText(const EMaterialQualityLevel::Type QualityType, const int32 InstanceIndex)
+{
+	if (PlatformData[QualityType].Instances.IsEmpty())
+	{
+		return FText::GetEmpty();
+	}
+
+	check(InstanceIndex >= 0 && InstanceIndex < PlatformData[QualityType].Instances.Num());
+
+	int TotalShadersForMaterial = 0;
+
+	if (FMaterialResource* Resource = PlatformData[QualityType].Instances[InstanceIndex].MaterialResourcesStats)
+	{
+		FPlatformTypeLayoutParameters LayoutParams;
+		LayoutParams.InitializeForPlatform(nullptr);
+
+		TArray<FDebugShaderTypeInfo> OutShaderInfo;
+		Resource->GetShaderTypes(LayoutParams, OutShaderInfo);
+
+		for (const FDebugShaderTypeInfo& ShaderInfo : OutShaderInfo)
+		{
+			TotalShadersForMaterial += ShaderInfo.ShaderTypes.Num();
+
+			for (const FDebugShaderPipelineInfo& PipelineInfo : ShaderInfo.Pipelines)
+			{
+				TotalShadersForMaterial += PipelineInfo.ShaderTypes.Num();
+			}
+		}
+	}
+
+	return FText::Format(
+		NSLOCTEXT("MaterialEditor", "NumShadersText", "Shaders: {0}"),
+		TotalShadersForMaterial
+	);
+}
+
 /** returns all shaders' stats concatenated together*/
 FString FShaderPlatformSettings::GetShadersStats() const
 {
@@ -782,15 +818,26 @@ TSharedPtr<FShaderPlatformSettings> FMaterialStats::GetPlatformSettings(const FN
 	return GetPlatformSettings(PlatformID);
 }
 
-FText FMaterialStats::GetShaderCode(const EShaderPlatform PlatformID, const EMaterialQualityLevel::Type QualityType, const int32 InstanceIndex)
+FText FMaterialStats::GetShaderCodeText(const EShaderPlatform PlatformID, const EMaterialQualityLevel::Type QualityType, const int32 InstanceIndex)
+{
+	auto* Entry = ShaderPlatformStatsDB.Find(PlatformID);
+	if (Entry == nullptr)
+	{
+		return NSLOCTEXT("MaterialStats", "ShaderStatsNotAvailableMsg", "Shader code compiling or not available!");
+	}
+
+	return (*Entry)->GetShaderCodeText(QualityType, InstanceIndex);
+}
+
+FText FMaterialStats::GetShaderCountText(const EShaderPlatform PlatformID, const EMaterialQualityLevel::Type QualityType, const int32 InstanceIndex)
 {
 	auto* Entry = ShaderPlatformStatsDB.Find(PlatformID);
 	if (Entry == nullptr)
 	{
-		return FText::FromString(TEXT("Shader code compiling or not available!"));
+		return NSLOCTEXT("MaterialStats", "ShaderCountNotAvailableMsg", "Shader count not available!");
 	}
 
-	return (*Entry)->GetShaderCode(QualityType, InstanceIndex);
+	return (*Entry)->GetShaderCountText(QualityType, InstanceIndex);
 }
 
 FString FMaterialStats::GetShadersStats() const
@@ -1273,7 +1320,7 @@ void FMaterialStats::BuildViewShaderCodeMenus()
 							auto StatsPtr = MaterialStats.Pin();
 							if (StatsPtr.IsValid())
 							{
-								return StatsPtr->GetShaderCode(PlatformID, QualityLevel, InstanceIndex);
+								return StatsPtr->GetShaderCodeText(PlatformID, QualityLevel, InstanceIndex);
 							}
 							return FText::FromString(TEXT("Error reading shader code!"));
 						},
@@ -1281,7 +1328,7 @@ void FMaterialStats::BuildViewShaderCodeMenus()
 							auto StatsPtr = MaterialStats.Pin();
 							if (StatsPtr.IsValid())
 							{
-								FText ShaderCode = StatsPtr->GetShaderCode(PlatformID, QualityLevel, InstanceIndex);
+								FText ShaderCode = StatsPtr->GetShaderCodeText(PlatformID, QualityLevel, InstanceIndex);
 								FPlatformApplicationMisc::ClipboardCopy(*ShaderCode.ToString());
 								return FReply::Handled();
 							}
diff --git a/Engine/Source/Editor/MaterialEditor/Private/MaterialStats.h b/Engine/Source/Editor/MaterialEditor/Private/MaterialStats.h
--- a/Engine/Source/Editor/MaterialEditor/Private/MaterialStats.h
+++ b/Engine/Source/Editor/MaterialEditor/Private/MaterialStats.h
@@ -284,7 +284,10 @@ public:
 	void OnShaderViewComboSelectionChanged(TSharedPtr<FMaterialShaderEntry> Item, EMaterialQualityLevel::Type QualityType, const int32 InstanceIndex);
 
 	/** returns the actual shader source selected by the shaders viewer's combo-box */
-	FText GetShaderCode(const EMaterialQualityLevel::Type QualityType, const int32 InstanceIndex);
+	FText GetShaderCodeText(const EMaterialQualityLevel::Type QualityType, const int32 InstanceIndex);
+
+	/** returns the number of shaders in this material w/o compiling shaders */
+	FText GetShaderCountText(const EMaterialQualityLevel::Type QualityType, const int32 InstanceIndex);
 
 	/** returns all shaders' stats concatenated together*/
 	FString GetShadersStats() const;
@@ -486,7 +489,10 @@ public:
 	TSharedPtr<FShaderPlatformSettings> GetPlatformSettings(const FName PlatformName);
 
 	/** returns the shader code computed by the specified platform with some quality level */
-	FText GetShaderCode(const EShaderPlatform PlatformID, const EMaterialQualityLevel::Type QualityType, const int32 InstanceIndex);
+	FText GetShaderCodeText(const EShaderPlatform PlatformID, const EMaterialQualityLevel::Type QualityType, const int32 InstanceIndex);
+
+	/** returns the number of shaders in the material with the specified platform with some quality level.  Does this without compiling shaders. */
+	FText GetShaderCountText(const EShaderPlatform PlatformID, const EMaterialQualityLevel::Type QualityType, const int32 InstanceIndex);
 
 	/** returns all shaders' stats concatenated together*/
 	FString GetShadersStats() const;
