Use the matching material quality level for building texture streaming when saving materials and material instances through the editor

#Synced-CL 48498858

diff --git a/Engine/Source/Editor/MaterialEditor/Private/MaterialEditorUtilities.cpp b/Engine/Source/Editor/MaterialEditor/Private/MaterialEditorUtilities.cpp
--- a/Engine/Source/Editor/MaterialEditor/Private/MaterialEditorUtilities.cpp
+++ b/Engine/Source/Editor/MaterialEditor/Private/MaterialEditorUtilities.cpp
@@ -816,13 +816,12 @@ bool FMaterialEditorUtilities::HasCompatibleConnection(UClass* ExpressionClass,
 
 void FMaterialEditorUtilities::BuildTextureStreamingData(UMaterialInterface* UpdatedMaterial)
 {
-	const EMaterialQualityLevel::Type QualityLevel = EMaterialQualityLevel::High;
+	const EMaterialQualityLevel::Type QualityLevel = GetCachedScalabilityCVars().MaterialQualityLevel;
 	const ERHIFeatureLevel::Type FeatureLevel = GMaxRHIFeatureLevel;
 	const EShaderPlatform ShaderPlatform = GMaxRHIShaderPlatform;
 
 	if (UpdatedMaterial)
 	{
-
 		CollectGarbage(GARBAGE_COLLECTION_KEEPFLAGS);
 
 		FScopedSlowTask SlowTask(2.f, (LOCTEXT("MaterialEditorUtilities_UpdatingTextureStreamingData", "Updating Texture Streaming Data")));
