Expose Material and Material Function GetExpressions to Blueprints and Python.

[RN] minor | Rendering Architecture
#Synced-CL 48503420

diff --git a/Engine/Source/Editor/MaterialEditor/Private/MaterialEditingLibrary.cpp b/Engine/Source/Editor/MaterialEditor/Private/MaterialEditingLibrary.cpp
--- a/Engine/Source/Editor/MaterialEditor/Private/MaterialEditingLibrary.cpp
+++ b/Engine/Source/Editor/MaterialEditor/Private/MaterialEditingLibrary.cpp
@@ -403,6 +403,26 @@ void UMaterialEditingLibrary::RebuildMaterialInstanceEditors(UMaterialFunction*
 	}
 }
 
+TArray<UMaterialExpression*> UMaterialEditingLibrary::GetMaterialFunctionExpressions(const UMaterialFunction* MaterialFunction)
+{
+	if (MaterialFunction)
+	{
+		return TArray<UMaterialExpression*>(MaterialFunction->GetExpressions());
+	}
+
+	return {};
+}
+
+TArray<UMaterialExpression*> UMaterialEditingLibrary::GetMaterialExpressions(const UMaterial* Material)
+{
+	if (Material)
+	{
+		return TArray<UMaterialExpression*>(Material->GetExpressions());
+	}
+
+	return {};
+}
+
 int32 UMaterialEditingLibrary::GetNumMaterialExpressions(const UMaterial* Material)
 {
 	int32 Result = 0;
diff --git a/Engine/Source/Editor/MaterialEditor/Public/MaterialEditingLibrary.h b/Engine/Source/Editor/MaterialEditor/Public/MaterialEditingLibrary.h
--- a/Engine/Source/Editor/MaterialEditor/Public/MaterialEditingLibrary.h
+++ b/Engine/Source/Editor/MaterialEditor/Public/MaterialEditingLibrary.h
@@ -106,11 +106,19 @@ public:
 
 	///////////// MATERIAL EDITING
 
+	/** returns the Expressions contained in this Material Function. Does not include nested Expressions. */
+	UFUNCTION(BlueprintPure, Category = "MaterialEditing")
+	static UE_API TArray<UMaterialExpression*> GetMaterialFunctionExpressions(const UMaterialFunction* MaterialFunction);
+
+	/** returns the Expressions contained in this Material. Does not include nested Expressions. */
+	UFUNCTION(BlueprintPure, Category = "MaterialEditing")
+	static UE_API TArray<UMaterialExpression*> GetMaterialExpressions(const UMaterial* Material);
 
 	/** Returns number of material expressions in the supplied material */
 	UFUNCTION(BlueprintPure, Category = "MaterialEditing")
 	static UE_API int32 GetNumMaterialExpressions(const UMaterial* Material);
 
+
 	/** Delete all material expressions in the supplied material */
 	UFUNCTION(BlueprintCallable, Category = "MaterialEditing")
 	static UE_API void DeleteAllMaterialExpressions(UMaterial* Material);
