Use a ListView to display all the shaders a material has in the Material Editor instead of a combo box.

- Update the code to use a SListView instead of SComboBox
- Move the SListView into a splitter so it is side-by-side with the Shader Code.
- This makes it easier to see all the shaders at a glance.
- Update the Shader Code Tab to include the full platform string which includes the quality level.

[RN] minor | Rendering - Materials and Shaders
#Synced-CL 48812887

diff --git a/Engine/Source/Editor/MaterialEditor/Private/MaterialStats.cpp b/Engine/Source/Editor/MaterialEditor/Private/MaterialStats.cpp
--- a/Engine/Source/Editor/MaterialEditor/Private/MaterialStats.cpp
+++ b/Engine/Source/Editor/MaterialEditor/Private/MaterialStats.cpp
@@ -1154,39 +1154,42 @@ TSharedRef<SDockTab> FMaterialStats::SpawnTab_ShaderCode(const FSpawnTabArgs& Ar
 
 	if (ShaderEntriesPtr == nullptr)
 	{
-		// SComboBox needs a valid pointer otherwise follow up SetItemsSource will fail.
+		// SListView needs a valid pointer otherwise follow up SetItemsSource will fail.
 		static TArray<TSharedPtr<FMaterialShaderEntry>> EmptyShaderEntriesArray = {};
 		ShaderEntriesPtr = &EmptyShaderEntriesArray;
 	}
 
 	ensure(ShaderEntriesPtr != nullptr);
-	TSharedRef<SComboBox<TSharedPtr<FMaterialShaderEntry>>> ShaderBox = SNew(SComboBox<TSharedPtr<FMaterialShaderEntry>>)
-		.OptionsSource(ShaderEntriesPtr)
-		.OnGenerateWidget_Lambda([](TSharedPtr<FMaterialShaderEntry> Value) { return SNew(STextBlock).Text(FText::FromString(Value->Text)); })
-		.OnSelectionChanged_Lambda([PlatformPtr, QualityLevel, InstanceIndex](TSharedPtr<FMaterialShaderEntry> Item, ESelectInfo::Type SelectInfo) { PlatformPtr->OnShaderViewComboSelectionChanged(Item, QualityLevel, InstanceIndex); })
-		[
-			SNew(STextBlock)
-			.Text_Lambda([PlatformPtr, QualityLevel, InstanceIndex]() { return PlatformPtr->GetSelectedShaderViewComboText(QualityLevel, InstanceIndex); })
-		];
+	TSharedRef<SListView<TSharedPtr<FMaterialShaderEntry>>> ShaderListView = SNew(SListView<TSharedPtr<FMaterialShaderEntry>>)
+		.ListItemsSource(ShaderEntriesPtr)
+		.SelectionMode(ESelectionMode::Single)
+		.OnGenerateRow_Lambda([](TSharedPtr<FMaterialShaderEntry> Value, const TSharedRef<STableViewBase>& OwnerTable) {
+			return SNew(STableRow<TSharedPtr<FMaterialShaderEntry>>, OwnerTable)
+			[
+				SNew(STextBlock).Text(FText::FromString(Value->Text))
+			];
+		})
+		.OnSelectionChanged_Lambda([PlatformPtr, QualityLevel, InstanceIndex](TSharedPtr<FMaterialShaderEntry> Item, ESelectInfo::Type SelectInfo) { PlatformPtr->OnShaderViewComboSelectionChanged(Item, QualityLevel, InstanceIndex); });
 
-	// Refresh ShaderBox when shaders are updated
-	FDelegateHandle ShaderBoxUpdater = RefreshDependentTabs.AddLambda([PlatformPtrWeakPtr = PlatformPtr.ToWeakPtr(), ShaderBoxWeakPtr = ShaderBox.ToWeakPtr(), QualityLevel, InstanceIndex]()
+	// Refresh ShaderList when shaders are updated
+	FDelegateHandle ShaderListUpdater = RefreshDependentTabs.AddLambda([PlatformPtrWeakPtr = PlatformPtr.ToWeakPtr(), ShaderListWeakPtr = ShaderListView.ToWeakPtr(), QualityLevel, InstanceIndex]()
 	{
 		TSharedPtr<FShaderPlatformSettings> PlatformPtr = PlatformPtrWeakPtr.Pin();
-		TSharedPtr<SComboBox<TSharedPtr<FMaterialShaderEntry>>> ShaderBox = ShaderBoxWeakPtr.Pin();
-		if (PlatformPtr.IsValid() && ShaderBox.IsValid())
+		TSharedPtr<SListView<TSharedPtr<FMaterialShaderEntry>>> ShaderListView = ShaderListWeakPtr.Pin();
+		if (PlatformPtr.IsValid() && ShaderListView.IsValid())
 		{
-			ShaderBox->SetItemsSource(PlatformPtr->GetShaderEntries(QualityLevel, InstanceIndex));
+			ShaderListView->SetItemsSource(PlatformPtr->GetShaderEntries(QualityLevel, InstanceIndex));
+			ShaderListView->RequestListRefresh();
 		}
 	});
 
 	TSharedRef<SDockTab> SpawnedTab = SNew(SDockTab)
-		.OnTabClosed_Lambda([MaterialStatsWeakPtr = this->AsWeak(), ShaderBoxUpdater](TSharedRef<SDockTab> Tab)
+		.OnTabClosed_Lambda([MaterialStatsWeakPtr = this->AsWeak(), ShaderListUpdater](TSharedRef<SDockTab> Tab)
 		{
 			TSharedPtr<FMaterialStats> MaterialStats = MaterialStatsWeakPtr.Pin();
-			if (MaterialStats.IsValid() && ShaderBoxUpdater.IsValid())
+			if (MaterialStats.IsValid() && ShaderListUpdater.IsValid())
 			{
-				MaterialStats->RefreshDependentTabs.Remove(ShaderBoxUpdater);
+				MaterialStats->RefreshDependentTabs.Remove(ShaderListUpdater);
 			}
 		})
 		[
@@ -1194,38 +1197,32 @@ TSharedRef<SDockTab> FMaterialStats::SpawnTab_ShaderCode(const FSpawnTabArgs& Ar
 			+ SVerticalBox::Slot()
 			.AutoHeight()
 			[
-				SNew(SVerticalBox)
-				// Copy Button
-				+ SVerticalBox::Slot()
-				.AutoHeight()
+				SNew(SHorizontalBox)
+				+ SHorizontalBox::Slot()
+				.AutoWidth()
+				.Padding(2.0f, 2.0f)
+				.VAlign(VAlign_Center)
+				.HAlign(HAlign_Left)
 				[
-					SNew(SHorizontalBox)
-					+ SHorizontalBox::Slot()
-					.AutoWidth()
-					.Padding(2.0f, 2.0f)
-					.VAlign(VAlign_Center)
-					.HAlign(HAlign_Left)
-					[
-						SNew(STextBlock)
+					SNew(STextBlock)
 						.Text(FText::FromString(FullPlatformName))
-					]
-				]
-				// Separator
-				+ SVerticalBox::Slot()
-				.FillHeight(1)
-				[
-					SNew(SSeparator)
 				]
 			]
 			+ SVerticalBox::Slot()
-			.AutoHeight()
+			.FillHeight(1.0f)
 			[
-				ShaderBox
-			]
-			+ SVerticalBox::Slot()
-			.FillHeight(1)
-			[
-				PlatformPtr->GetShaderViewerScrollBox(QualityLevel).ToSharedRef()
+				SNew(SSplitter)
+				.Orientation(EOrientation::Orient_Horizontal)
+				+ SSplitter::Slot()
+				.Value(0.3f)
+				[
+					ShaderListView
+				]
+				+ SSplitter::Slot()
+				.Value(0.7f)
+				[
+					PlatformPtr->GetShaderViewerScrollBox(QualityLevel).ToSharedRef()
+				]
 			]
 		];
 
