Show shader count for the preview platform in material editor preview windows.
#rnx
#Synced-CL 48841205

diff --git a/Engine/Source/Editor/MaterialEditor/Private/MaterialEditor.cpp b/Engine/Source/Editor/MaterialEditor/Private/MaterialEditor.cpp
--- a/Engine/Source/Editor/MaterialEditor/Private/MaterialEditor.cpp
+++ b/Engine/Source/Editor/MaterialEditor/Private/MaterialEditor.cpp
@@ -3142,7 +3142,12 @@ void FMaterialEditor::UpdateMaterialPlatformStats()
 	// Update the num shaders UI
 	if (NumShadersText.IsValid())
 	{
-		NumShadersText->SetText(MaterialStatsManager->GetShaderCountText(GMaxRHIShaderPlatform, GetQualityPreviewLevel(), 0));
+		EShaderPlatform Platform = GEditor->GetActiveShaderPlatform();
+		if (FDataDrivenShaderPlatformInfo::GetIsPreviewPlatform(Platform))
+		{
+			Platform = FDataDrivenShaderPlatformInfo::GetPreviewShaderPlatformParent(Platform);
+		}
+		NumShadersText->SetText(MaterialStatsManager->GetShaderCountText(Platform, GetQualityPreviewLevel(), 0));
 	}
 
 	// check if any derived instances fail, if so show the window
diff --git a/Engine/Source/Editor/MaterialEditor/Private/MaterialInstanceEditor.cpp b/Engine/Source/Editor/MaterialEditor/Private/MaterialInstanceEditor.cpp
--- a/Engine/Source/Editor/MaterialEditor/Private/MaterialInstanceEditor.cpp
+++ b/Engine/Source/Editor/MaterialEditor/Private/MaterialInstanceEditor.cpp
@@ -1698,7 +1698,12 @@ void FMaterialInstanceEditor::Tick(float DeltaTime)
 		// Update the Num Shaders text in the viewport.
 		if (MaterialEditorInstance->SourceInstance->bHasStaticPermutationResource)
 		{
-			NumShadersPillText->SetText(MaterialStatsManager->GetShaderCountText(GMaxRHIShaderPlatform, GetCachedScalabilityCVars().MaterialQualityLevel, 0));
+			EShaderPlatform Platform = GEditor->GetActiveShaderPlatform();
+			if (FDataDrivenShaderPlatformInfo::GetIsPreviewPlatform(Platform))
+			{
+				Platform = FDataDrivenShaderPlatformInfo::GetPreviewShaderPlatformParent(Platform);
+			}
+			NumShadersPillText->SetText(MaterialStatsManager->GetShaderCountText(Platform, GetCachedScalabilityCVars().MaterialQualityLevel, 0));
 		}
 		else
 		{
