Fix MaterialStats Manager triggering shader compilation on Scalar and Vector params modification.
Clean up property class name usage.

[RN] minor | Rendering Architecture

#Synced-CL 48863451

diff --git a/Engine/Source/Editor/MaterialEditor/Private/MaterialInstanceEditor.cpp b/Engine/Source/Editor/MaterialEditor/Private/MaterialInstanceEditor.cpp
--- a/Engine/Source/Editor/MaterialEditor/Private/MaterialInstanceEditor.cpp
+++ b/Engine/Source/Editor/MaterialEditor/Private/MaterialInstanceEditor.cpp
@@ -13,11 +13,7 @@
 #include "Editor.h"
 #include "Styling/AppStyle.h"
 #include "Styling/CoreStyle.h"
-#include "MaterialEditor/DEditorTextureParameterValue.h"
-#include "MaterialEditor/DEditorRuntimeVirtualTextureParameterValue.h"
-#include "MaterialEditor/DEditorSparseVolumeTextureParameterValue.h"
 #include "Materials/Material.h"
-#include "MaterialEditor/MaterialEditorInstanceConstant.h"
 #include "ThumbnailRendering/SceneThumbnailInfoWithPrimitive.h"
 #include "Materials/MaterialFunction.h"
 #include "Materials/MaterialInstance.h"
@@ -37,6 +33,11 @@
 #include "MaterialEditor.h"
 #include "MaterialEditorActions.h"
 #include "MaterialEditorUtilities.h"
+#include "MaterialEditor/DEditorTextureParameterValue.h"
+#include "MaterialEditor/DEditorRuntimeVirtualTextureParameterValue.h"
+#include "MaterialEditor/DEditorSparseVolumeTextureParameterValue.h"
+#include "MaterialEditor/DEditorStaticSwitchParameterValue.h"
+#include "MaterialEditor/MaterialEditorInstanceConstant.h"
 
 #include "PropertyEditorModule.h"
 #include "MaterialEditorInstanceDetailCustomization.h"
@@ -62,6 +63,7 @@
 #include "MaterialEditorViewportToolbarSections.h"
 #include "Materials/MaterialFunctionMaterialLayer.h"
 #include "Materials/MaterialFunctionMaterialLayerBlend.h"
+#include "Materials/MaterialInstanceBasePropertyOverrides.h"
 #include "VT/RuntimeVirtualTexture.h"
 #include "SparseVolumeTexture/SparseVolumeTexture.h"
 #include "Widgets/Input/SButton.h"
@@ -1326,7 +1328,7 @@ void FMaterialInstanceEditor::NotifyPostChange( const FPropertyChangedEvent& Pro
 	// Rebuild the property window to account for the possibility that  
 	// the item changed was a static switch or function call parameter
 	UObject* PropertyClass = PropertyThatChanged->GetOwner<UObject>();
-	if(PropertyClass && (PropertyClass->GetName() == TEXT("DEditorStaticSwitchParameterValue") || PropertyClass->GetName() == TEXT("EditorParameterGroup"))//DEditorMaterialLayerParameters"))
+	if(PropertyClass && (PropertyClass == UDEditorStaticSwitchParameterValue::StaticClass() || PropertyClass == FEditorParameterGroup::StaticStruct())
 		&& MaterialEditorInstance->Parent && MaterialEditorInstance->SourceInstance )
 	{
 		// TODO: We need to hit this on MaterialLayerParam updates but only get notifications for their array elements changing, hence the overly generic test above
@@ -1337,15 +1339,23 @@ void FMaterialInstanceEditor::NotifyPostChange( const FPropertyChangedEvent& Pro
 	}
 
 	if (PropertyClass && 
-	   (PropertyClass->GetName() == TEXT("DEditorTextureParameterValue") || 
-		PropertyClass->GetName() == TEXT("DEditorRuntimeVirtualTextureParameterValue") ||
-		PropertyClass->GetName() == TEXT("DEditorSparseVolumeTextureParameterValue")))
+	   (PropertyClass == UDEditorTextureParameterValue::StaticClass() ||
+		PropertyClass == UDEditorRuntimeVirtualTextureParameterValue::StaticClass()||
+		PropertyClass == UDEditorSparseVolumeTextureParameterValue::StaticClass()))
 	{
 		RefreshOnScreenTextureWarnings();
 	}
+	
+	const bool bRefreshStats = PropertyClass && (
+		PropertyClass == UDEditorStaticSwitchParameterValue::StaticClass() ||
+		PropertyClass == FMaterialInstanceBasePropertyOverrides::StaticStruct()
+		);
 
-	// something was changed in the material so we need to reflect this in the stats
-	MaterialStatsManager->SignalMaterialChanged();
+	if (bRefreshStats)
+	{
+		// something was changed in the material so we need to reflect this in the stats
+		MaterialStatsManager->SignalMaterialChanged();
+	}
 
 	// Update the preview window when the user changes a property.
 	PreviewVC->RefreshViewport();
@@ -1892,3 +1902,4 @@ void FMaterialInstanceEditor::NotifyUserSceneTextureLoadOrUnload()
 }
 
 #undef LOCTEXT_NAMESPACE
+
