Use TWeakPtr in lambda to prevent crash in Material Instance Editor widget.
#rnx
#Synced-CL 48982928

diff --git a/Engine/Source/Editor/MaterialEditor/Private/MaterialEditorInstanceDetailCustomization.cpp b/Engine/Source/Editor/MaterialEditor/Private/MaterialEditorInstanceDetailCustomization.cpp
--- a/Engine/Source/Editor/MaterialEditor/Private/MaterialEditorInstanceDetailCustomization.cpp
+++ b/Engine/Source/Editor/MaterialEditor/Private/MaterialEditorInstanceDetailCustomization.cpp
@@ -89,27 +89,35 @@ void FMaterialInstanceParameterDetails::CustomizeDetails(IDetailLayoutBuilder& D
 	IDetailCategoryBuilder& GroupsCategory = DetailLayout.EditCategory(GroupsCategoryName, LOCTEXT("MICParamGroupsTitle", "Parameter Groups"));
 	
 	TWeakObjectPtr<UMaterialInterface> MaterialWeakPtr(MaterialEditorInstance->GetMaterialInterface());
+	TWeakPtr<FMaterialEditorValidatorContext> ValidationContextWeakPtr = ValidationContext;
 
 	// Insert optional validation warning row at the top of the category.
 	GroupsCategory.AddCustomRow(LOCTEXT("ValidationWarningRow", "Validation Warning"))
 	.WholeRowContent()
 	[
 		SNew(SWarningOrErrorBox)
-		.Visibility_Lambda([this, MaterialWeakPtr]() -> EVisibility
+		.Visibility_Lambda([MaterialWeakPtr, ValidationContextWeakPtr]() -> EVisibility
 		{
 			bool bWarn = false;
 			if (MaterialWeakPtr.IsValid())
 			{
 				if (UMaterialEditorValidationSubsystem* Subsystem = GEditor->GetEditorSubsystem<UMaterialEditorValidationSubsystem>())
 				{
-					bWarn = !Subsystem->Validate(MaterialWeakPtr.Get(), *ValidationContext.Get());
+					if (TSharedPtr<FMaterialEditorValidatorContext> ValidationContext = ValidationContextWeakPtr.Pin())
+					{
+						bWarn = !Subsystem->Validate(MaterialWeakPtr.Get(), *ValidationContext.Get());
+					}
 				}
 			}
 			return bWarn ? EVisibility::Visible : EVisibility::Collapsed;
 		})
-		.Message_Lambda([this]() -> FText
+		.Message_Lambda([ValidationContextWeakPtr]() -> FText
 		{
-			return ValidationContext->Message;
+			if (TSharedPtr<FMaterialEditorValidatorContext> ValidationContext = ValidationContextWeakPtr.Pin())
+			{
+				return ValidationContext->Message;
+			}
+			return {};
 		})
 	];
 	
