Fix Material Instance Editor ResetToDefault not sending PostEditChange messages.
This matters when using a child material instance for the preview window.
#rnx
#Synced-CL 49018876

diff --git a/Engine/Source/Editor/MaterialEditor/Private/MaterialPropertyHelpers.cpp b/Engine/Source/Editor/MaterialEditor/Private/MaterialPropertyHelpers.cpp
--- a/Engine/Source/Editor/MaterialEditor/Private/MaterialPropertyHelpers.cpp
+++ b/Engine/Source/Editor/MaterialEditor/Private/MaterialPropertyHelpers.cpp
@@ -587,7 +587,7 @@ void FMaterialPropertyHelpers::ConfigurePropertyRowForParameter(IDetailPropertyR
 	PropertyRow
 		.DisplayName(FText::FromName(Parameter->ParameterInfo.Name))
 		.ToolTip(FMaterialPropertyHelpers::GetParameterTooltip(Parameter, MaterialEditorInstance))
-		.OverrideResetToDefault(FMaterialPropertyHelpers::CreateResetToDefaultOverride(Parameter, MaterialEditorInstance));
+		.OverrideResetToDefault(FMaterialPropertyHelpers::CreateResetToDefaultOverride(Parameter, MaterialEditorInstance, ParameterValueProperty));
 
 	// Set Visibility
 	if (RowOverrides.Visibility.IsSet())
@@ -1193,7 +1193,7 @@ FText FMaterialPropertyHelpers::GetParameterTooltip(UDEditorParameterValue* Para
 	return TooltipText;
 }
 
-FResetToDefaultOverride FMaterialPropertyHelpers::CreateResetToDefaultOverride(UDEditorParameterValue* Parameter, UMaterialEditorInstanceConstant* MaterialEditorInstance)
+FResetToDefaultOverride FMaterialPropertyHelpers::CreateResetToDefaultOverride(UDEditorParameterValue* Parameter, UMaterialEditorInstanceConstant* MaterialEditorInstance, TSharedPtr<IPropertyHandle> ParameterValueProperty)
 {
 	UDEditorScalarParameterValue* ScalarParam = Cast<UDEditorScalarParameterValue>(Parameter);
 	if (ScalarParam && ScalarParam->AtlasData.bIsUsedAsAtlasPosition)
@@ -1204,11 +1204,11 @@ FResetToDefaultOverride FMaterialPropertyHelpers::CreateResetToDefaultOverride(U
 	}
 
 	TAttribute<bool> IsResetVisible = TAttribute<bool>::Create(TAttribute<bool>::FGetter::CreateStatic(&FMaterialPropertyHelpers::ShouldShowResetToDefault, Parameter, MaterialEditorInstance));
-	FSimpleDelegate ResetHandler = FSimpleDelegate::CreateStatic(&FMaterialPropertyHelpers::ResetToDefault, Parameter, MaterialEditorInstance);
+	FSimpleDelegate ResetHandler = FSimpleDelegate::CreateStatic(&FMaterialPropertyHelpers::ResetToDefault, Parameter, MaterialEditorInstance, ParameterValueProperty);
 	return FResetToDefaultOverride::Create(IsResetVisible, ResetHandler);
 }
 
-void FMaterialPropertyHelpers::ResetToDefault(class UDEditorParameterValue* Parameter, UMaterialEditorInstanceConstant* MaterialEditorInstance)
+void FMaterialPropertyHelpers::ResetToDefault(class UDEditorParameterValue* Parameter, UMaterialEditorInstanceConstant* MaterialEditorInstance, TSharedPtr<IPropertyHandle> ParameterValueProperty)
 {
 	const FScopedTransaction Transaction( LOCTEXT( "ResetToDefault", "Reset To Default" ) );
 	Parameter->Modify();
@@ -1219,7 +1219,10 @@ void FMaterialPropertyHelpers::ResetToDefault(class UDEditorParameterValue* Para
 		FMaterialParameterMetadata DefaultValue;
 		if (MaterialEditorInstance->SourceInstance->GetParameterDefaultValue(EditorValue.Value.Type, Parameter->ParameterInfo, DefaultValue))
 		{
+			ParameterValueProperty->NotifyPreChange();
 			Parameter->SetValue(DefaultValue.Value);
+			ParameterValueProperty->NotifyPostChange(EPropertyChangeType::ValueSet|EPropertyChangeType::ResetToDefault);
+
 			MaterialEditorInstance->CopyToSourceInstance();
 		}
 	}
diff --git a/Engine/Source/Editor/MaterialEditor/Private/SMaterialLayersFunctionsTree.cpp b/Engine/Source/Editor/MaterialEditor/Private/SMaterialLayersFunctionsTree.cpp
--- a/Engine/Source/Editor/MaterialEditor/Private/SMaterialLayersFunctionsTree.cpp
+++ b/Engine/Source/Editor/MaterialEditor/Private/SMaterialLayersFunctionsTree.cpp
@@ -573,7 +573,7 @@ void SMaterialLayersFunctionsInstanceTreeItem::Construct(const FArguments& InArg
 
 		ResetWidget = SNew(SResetToDefaultPropertyEditor, PropertyHandle)
 			.IsEnabled(PropertyRowOverrides.EditConditionValue)
-			.CustomResetToDefault(FMaterialPropertyHelpers::CreateResetToDefaultOverride(StackParameterData->Parameter, MaterialEditorInstance));
+			.CustomResetToDefault(FMaterialPropertyHelpers::CreateResetToDefaultOverride(StackParameterData->Parameter, MaterialEditorInstance, PropertyHandle));
 
 		PropertyHandle->SetOnPropertyValueChanged(FSimpleDelegate::CreateSP(Tree, &SMaterialLayersFunctionsInstanceTree::UpdateThumbnailMaterial, StackParameterData->ParameterInfo.Association, StackParameterData->ParameterInfo.Index, false));
 		PropertyHandle->SetOnChildPropertyValueChanged(FSimpleDelegate::CreateSP(Tree, &SMaterialLayersFunctionsInstanceTree::UpdateThumbnailMaterial, StackParameterData->ParameterInfo.Association, StackParameterData->ParameterInfo.Index, false));
diff --git a/Engine/Source/Editor/MaterialEditor/Public/MaterialPropertyHelpers.h b/Engine/Source/Editor/MaterialEditor/Public/MaterialPropertyHelpers.h
--- a/Engine/Source/Editor/MaterialEditor/Public/MaterialPropertyHelpers.h
+++ b/Engine/Source/Editor/MaterialEditor/Public/MaterialPropertyHelpers.h
@@ -221,8 +221,8 @@ public:
 	static EVisibility ShouldShowExpression(TObjectPtr<UDEditorParameterValue> Parameter, UMaterialEditorInstanceConstant* MaterialEditorInstance, FGetShowHiddenParameters ShowHiddenDelegate) { return ShouldShowExpression(Parameter.Get(), MaterialEditorInstance, ShowHiddenDelegate); }
 
 	/** Generic material property reset to default implementation.  Resets Parameter to default */
-	static UE_API FResetToDefaultOverride CreateResetToDefaultOverride(UDEditorParameterValue* Parameter, UMaterialEditorInstanceConstant* MaterialEditorInstance);
-	static UE_API void ResetToDefault(UDEditorParameterValue* Parameter, UMaterialEditorInstanceConstant* MaterialEditorInstance);
+	static UE_API FResetToDefaultOverride CreateResetToDefaultOverride(UDEditorParameterValue* Parameter, UMaterialEditorInstanceConstant* MaterialEditorInstance, TSharedPtr<IPropertyHandle> ParameterValueProperty);
+	static UE_API void ResetToDefault(UDEditorParameterValue* Parameter, UMaterialEditorInstanceConstant* MaterialEditorInstance, TSharedPtr<IPropertyHandle> ParameterValueProperty);
 	static void ResetToDefault(TObjectPtr<UDEditorParameterValue> Parameter, UMaterialEditorInstanceConstant* MaterialEditorInstance) { ResetToDefault(Parameter.Get(), MaterialEditorInstance); }
 	static UE_API bool ShouldShowResetToDefault(UDEditorParameterValue* Parameter, UMaterialEditorInstanceConstant* MaterialEditorInstance);
 	static bool ShouldShowResetToDefault(TObjectPtr<UDEditorParameterValue> Parameter, UMaterialEditorInstanceConstant* MaterialEditorInstance) { return ShouldShowResetToDefault(Parameter.Get(), MaterialEditorInstance); }
