Build fix
#rnx
#Synced-CL 49018876

diff --git a/Engine/Source/Editor/MaterialEditor/Public/MaterialPropertyHelpers.h b/Engine/Source/Editor/MaterialEditor/Public/MaterialPropertyHelpers.h
--- a/Engine/Source/Editor/MaterialEditor/Public/MaterialPropertyHelpers.h
+++ b/Engine/Source/Editor/MaterialEditor/Public/MaterialPropertyHelpers.h
@@ -223,7 +223,7 @@ public:
 	/** Generic material property reset to default implementation.  Resets Parameter to default */
 	static UE_API FResetToDefaultOverride CreateResetToDefaultOverride(UDEditorParameterValue* Parameter, UMaterialEditorInstanceConstant* MaterialEditorInstance, TSharedPtr<IPropertyHandle> ParameterValueProperty);
 	static UE_API void ResetToDefault(UDEditorParameterValue* Parameter, UMaterialEditorInstanceConstant* MaterialEditorInstance, TSharedPtr<IPropertyHandle> ParameterValueProperty);
-	static void ResetToDefault(TObjectPtr<UDEditorParameterValue> Parameter, UMaterialEditorInstanceConstant* MaterialEditorInstance) { ResetToDefault(Parameter.Get(), MaterialEditorInstance); }
+	static void ResetToDefault(TObjectPtr<UDEditorParameterValue> Parameter, UMaterialEditorInstanceConstant* MaterialEditorInstance, TSharedPtr<IPropertyHandle> ParameterValueProperty) { ResetToDefault(Parameter.Get(), MaterialEditorInstance, ParameterValueProperty); }
 	static UE_API bool ShouldShowResetToDefault(UDEditorParameterValue* Parameter, UMaterialEditorInstanceConstant* MaterialEditorInstance);
 	static bool ShouldShowResetToDefault(TObjectPtr<UDEditorParameterValue> Parameter, UMaterialEditorInstanceConstant* MaterialEditorInstance) { return ShouldShowResetToDefault(Parameter.Get(), MaterialEditorInstance); }
 	
