Remove unused wrapper functions
#rnx
#Synced-CL 49065038

diff --git a/Engine/Source/Editor/MaterialEditor/Public/MaterialPropertyHelpers.h b/Engine/Source/Editor/MaterialEditor/Public/MaterialPropertyHelpers.h
--- a/Engine/Source/Editor/MaterialEditor/Public/MaterialPropertyHelpers.h
+++ b/Engine/Source/Editor/MaterialEditor/Public/MaterialPropertyHelpers.h
@@ -223,9 +223,7 @@ public:
 	/** Generic material property reset to default implementation.  Resets Parameter to default */
 	static UE_API FResetToDefaultOverride CreateResetToDefaultOverride(UDEditorParameterValue* Parameter, UMaterialEditorInstanceConstant* MaterialEditorInstance, TSharedPtr<IPropertyHandle> ParameterValueProperty);
 	static UE_API void ResetToDefault(UDEditorParameterValue* Parameter, UMaterialEditorInstanceConstant* MaterialEditorInstance, TSharedPtr<IPropertyHandle> ParameterValueProperty);
-	static void ResetToDefault(TObjectPtr<UDEditorParameterValue> Parameter, UMaterialEditorInstanceConstant* MaterialEditorInstance, TSharedPtr<IPropertyHandle> ParameterValueProperty) { ResetToDefault(Parameter.Get(), MaterialEditorInstance, ParameterValueProperty); }
 	static UE_API bool ShouldShowResetToDefault(UDEditorParameterValue* Parameter, UMaterialEditorInstanceConstant* MaterialEditorInstance);
-	static bool ShouldShowResetToDefault(TObjectPtr<UDEditorParameterValue> Parameter, UMaterialEditorInstanceConstant* MaterialEditorInstance) { return ShouldShowResetToDefault(Parameter.Get(), MaterialEditorInstance); }
 	
 	/** Specific resets for layer and blend asses */
 	static UE_API void ResetLayerAssetToDefault(UDEditorParameterValue* InParameter, TEnumAsByte<EMaterialParameterAssociation> InAssociation, int32 Index, UMaterialEditorInstanceConstant* MaterialEditorInstance);
@@ -247,9 +245,7 @@ public:
 
 	static UE_API void GetVectorChannelMaskComboBoxStrings(TArray<TSharedPtr<FString>>& OutComboBoxStrings, TArray<TSharedPtr<class SToolTip>>& OutToolTips, TArray<bool>& OutRestrictedItems);
 	static UE_API FString GetVectorChannelMaskValue(UDEditorParameterValue* InParameter);
-	static FString GetVectorChannelMaskValue(TObjectPtr<UDEditorParameterValue> InParameter) { return GetVectorChannelMaskValue(InParameter.Get()); }
 	static UE_API void SetVectorChannelMaskValue(const FString& StringValue, TSharedPtr<IPropertyHandle> PropertyHandle, UDEditorParameterValue* InParameter, UObject* MaterialEditorInstance);
-	static void SetVectorChannelMaskValue(const FString& StringValue, TSharedPtr<IPropertyHandle> PropertyHandle, TObjectPtr<UDEditorParameterValue> InParameter, UObject* MaterialEditorInstance) { SetVectorChannelMaskValue(StringValue, PropertyHandle, InParameter.Get(), MaterialEditorInstance); }
 
 	static UE_API TArray<class UFactory*> GetAssetFactories(EMaterialParameterAssociation AssetType);
 	/**
@@ -266,7 +262,6 @@ public:
 	static UE_API void SetPositionFromCurveAsset(const FAssetData& AssetData, TSoftObjectPtr<UCurveLinearColorAtlas> InAtlas, class UDEditorScalarParameterValue* InParameter, TSharedPtr<IPropertyHandle> PropertyHandle, UObject* MaterialEditorInstance);
 
 	static UE_API void ResetCurveToDefault(UDEditorParameterValue* Parameter, UMaterialEditorInstanceConstant* MaterialEditorInstance);
-	static void ResetCurveToDefault(TObjectPtr<UDEditorParameterValue> Parameter, UMaterialEditorInstanceConstant* MaterialEditorInstance) { ResetCurveToDefault(Parameter.Get(), MaterialEditorInstance); }
 
 	static  TSharedRef<SWidget> CreateScalarEnumerationParameterValueWidget(TSoftObjectPtr<UObject> const& SoftEnumeration, TSharedPtr<IPropertyHandle> ParameterValueProperty);
 
