Add the dynamic slate shader permutation to the material stats so instruction count shows up in the Material Editor for UI materials.

todo: We should make this type based and not name based since this is extremely brittle.

#rnx
#Synced-CL 49034721

diff --git a/Engine/Source/Editor/MaterialEditor/Private/MaterialStatsCommon.cpp b/Engine/Source/Editor/MaterialEditor/Private/MaterialStatsCommon.cpp
--- a/Engine/Source/Editor/MaterialEditor/Private/MaterialStatsCommon.cpp
+++ b/Engine/Source/Editor/MaterialEditor/Private/MaterialStatsCommon.cpp
@@ -222,6 +222,9 @@ FString FMaterialStatsUtils::RepresentativeShaderTypeToString(const ERepresentat
 		case ERepresentativeShader::UIDefaultFragmentShader:
 			return TEXT("UI Pixel Shader");
 
+		case ERepresentativeShader::UIDynamicFragmentShader:
+			return TEXT("UI Pixel Shader (Dynamic)");
+
 		case ERepresentativeShader::UIDefaultVertexShader:
 			return TEXT("UI Vertex Shader");
 
@@ -292,6 +295,10 @@ void FMaterialStatsUtils::GetRepresentativeShaderTypesAndDescriptions(TMap<FName
 		ShaderTypeNamesAndDescriptions.FindOrAdd(FLocalVertexFactoryName)
 			.Add(FRepresentativeShaderInfo(ERepresentativeShader::UIDefaultFragmentShader, TSlateMaterialShaderPSDefaultName, TEXT("Default UI Pixel Shader")));
 
+		static FName TSlateMaterialShaderPSDynamicName = TEXT("TSlateMaterialShaderPSDynamic");
+		ShaderTypeNamesAndDescriptions.FindOrAdd(FLocalVertexFactoryName)
+		.Add(FRepresentativeShaderInfo(ERepresentativeShader::UIDynamicFragmentShader, TSlateMaterialShaderPSDynamicName, TEXT("Dynamic UI Pixel Shader")));
+
 		static FName TSlateMaterialShaderVSfalseName = TEXT("TSlateMaterialShaderVSfalse");
 		ShaderTypeNamesAndDescriptions.FindOrAdd(FLocalVertexFactoryName)
 			.Add(FRepresentativeShaderInfo(ERepresentativeShader::UIDefaultVertexShader, TSlateMaterialShaderVSfalseName, TEXT("Default UI Vertex Shader")));
diff --git a/Engine/Source/Editor/MaterialEditor/Public/MaterialStatsCommon.h b/Engine/Source/Editor/MaterialEditor/Public/MaterialStatsCommon.h
--- a/Engine/Source/Editor/MaterialEditor/Public/MaterialStatsCommon.h
+++ b/Engine/Source/Editor/MaterialEditor/Public/MaterialStatsCommon.h
@@ -56,7 +56,8 @@ enum class ERepresentativeShader
 	DynamicallyLitObject,
 	RuntimeVirtualTextureOutput,
 	UIDefaultFragmentShader,
-	LastFragmentShader = UIDefaultFragmentShader,
+	UIDynamicFragmentShader,
+	LastFragmentShader = UIDynamicFragmentShader,
 
 	FirstVertexShader,
 	StaticMesh = FirstVertexShader,
