Change the material editor preview mesh selector to use new dialogue boxes.
1 - The static mesh selection now pops up a content browser highlighting the originally selected mesh but also allowing the user to select a different mesh for the preview material.
2 - The material slot selection window allows the user to specify which material slots they want the preview material applied to (slot options change based on selected mesh and refresh if the mesh is swapped).

#Synced-CL 49406606

diff --git a/Engine/Source/Editor/MaterialEditor/Private/MaterialEditorActions.cpp b/Engine/Source/Editor/MaterialEditor/Private/MaterialEditorActions.cpp
--- a/Engine/Source/Editor/MaterialEditor/Private/MaterialEditorActions.cpp
+++ b/Engine/Source/Editor/MaterialEditor/Private/MaterialEditorActions.cpp
@@ -38,6 +38,7 @@ void FMaterialEditorCommands::RegisterCommands()
 	UI_COMMAND( SetPlanePreview, "Plane", "Sets the preview mesh to a plane primitive.", EUserInterfaceActionType::RadioButton, FInputChord() );
 	UI_COMMAND( SetCubePreview, "Cube", "Sets the preview mesh to a cube primitive.", EUserInterfaceActionType::RadioButton, FInputChord() );
 	UI_COMMAND( SetPreviewMeshFromSelection, "Mesh", "Sets the preview mesh based on the current content browser selection.", EUserInterfaceActionType::RadioButton, FInputChord() );
+	UI_COMMAND( SelectPreviewMeshMaterialSlots, "Mesh Slots", "Sets the preview mesh slots to apply the material to.", EUserInterfaceActionType::RadioButton, FInputChord() );
 	UI_COMMAND( TogglePreviewBackground, "Background", "Toggles the preview pane's background.", EUserInterfaceActionType::ToggleButton, FInputChord() );
 
 	UI_COMMAND( CameraHome, "Home", "Goes home on the canvas.", EUserInterfaceActionType::Button, FInputChord() );
diff --git a/Engine/Source/Editor/MaterialEditor/Private/MaterialEditorViewportToolbarSections.cpp b/Engine/Source/Editor/MaterialEditor/Private/MaterialEditorViewportToolbarSections.cpp
--- a/Engine/Source/Editor/MaterialEditor/Private/MaterialEditorViewportToolbarSections.cpp
+++ b/Engine/Source/Editor/MaterialEditor/Private/MaterialEditorViewportToolbarSections.cpp
@@ -3,6 +3,7 @@
 #include "MaterialEditorViewportToolbarSections.h"
 #include "EditorViewportCommands.h"
 #include "MaterialEditorActions.h"
+#include "Materials/Material.h"
 #include "SEditorViewport.h"
 #include "SMaterialEditorViewport.h"
 #include "ToolMenu.h"
@@ -121,6 +122,13 @@ void UE::MaterialEditor::ExtendPreviewSceneSettingsSubmenu(FName InSubmenuName)
 	MeshFromSelectionEntry.Icon = FSlateIconFinder::FindIconForClass(UStaticMesh::StaticClass());
 	MeshFromSelectionEntry.SetShowInToolbarTopLevel(true);
 	MeshFromSelectionEntry.ToolBarData = ToolbarData;
+
+	FToolMenuEntry& MaterialSlotsEntry =
+		PreviewMeshSection.AddMenuEntry(FMaterialEditorCommands::Get().SelectPreviewMeshMaterialSlots);
+	MaterialSlotsEntry.Label = LOCTEXT("SelectPreviewMeshMaterialSlotsLabel", "Mesh material slots to use when previewing material");
+	MaterialSlotsEntry.Icon = FSlateIconFinder::FindIconForClass(UMaterial::StaticClass());
+	MaterialSlotsEntry.SetShowInToolbarTopLevel(true);
+	MaterialSlotsEntry.ToolBarData = ToolbarData;
 }
 
 #undef LOCTEXT_NAMESPACE
diff --git a/Engine/Source/Editor/MaterialEditor/Private/SMaterialEditorViewport.cpp b/Engine/Source/Editor/MaterialEditor/Private/SMaterialEditorViewport.cpp
--- a/Engine/Source/Editor/MaterialEditor/Private/SMaterialEditorViewport.cpp
+++ b/Engine/Source/Editor/MaterialEditor/Private/SMaterialEditorViewport.cpp
@@ -46,6 +46,9 @@
 #include "MaterialEditorViewportToolbarSections.h"
 #include "ToolMenus.h"
 #include "ViewportToolbar/UnrealEdViewportToolbar.h"
+#include "AssetThumbnail.h"
+#include "ContentBrowserModule.h"
+#include "IContentBrowserSingleton.h"
 
 #define LOCTEXT_NAMESPACE "MaterialEditor"
 #include "AdvancedPreviewSceneMenus.h"
@@ -423,6 +426,24 @@ bool SMaterialEditor3DPreviewViewport::SetPreviewAsset(UObject* InAsset)
 		AdvancedPreviewScene->AddComponent(PreviewMeshComponent, Transform);
 		AdvancedPreviewScene->SetFloorOffset(-PreviewMeshComponent->Bounds.Origin.Z + PreviewMeshComponent->Bounds.BoxExtent.Z);
 
+		//Repopulate the material slot states as the mesh option is selected
+		MaterialSlots.Empty();
+		if (UStaticMesh* SM = Cast<UStaticMesh>(InAsset))
+		{
+			const TArray<FStaticMaterial>& StaticMatInfos = SM->GetStaticMaterials();
+			for (const FStaticMaterial& Slot : StaticMatInfos)
+			{
+				MaterialSlots.Add({ Slot.MaterialSlotName, Slot.MaterialInterface, true });
+			}
+		}
+		else if (USkeletalMesh* SK = Cast<USkeletalMesh>(InAsset))
+		{
+			const TArray<FSkeletalMaterial>& SkeletalMatInfos = SK->GetMaterials();
+			for (const FSkeletalMaterial& Slot : SkeletalMatInfos)
+			{
+				MaterialSlots.Add({ Slot.MaterialSlotName, Slot.MaterialInterface, true });
+			}
+		}
 	}
 
 	// Make sure the preview material is applied to the component
@@ -693,18 +714,28 @@ void SMaterialEditor3DPreviewViewport::SetPreviewMaterial(UMaterialInterface* In
 		// Add the preview material to the preview mesh.
 		if (PreviewMeshComponent != nullptr)
 		{
-			PreviewMeshComponent->OverrideMaterials.Empty();
-
-			if (PreviewMaterial)
+			int32 NumMaterialSlots = PreviewMeshComponent->GetNumMaterials();
+			ensure(NumMaterialSlots == MaterialSlots.Num());
+			//Ensure the override materials has enough array slots to apply the preview material
+			PreviewMeshComponent->OverrideMaterials.SetNum(NumMaterialSlots, EAllowShrinking::Yes);
+			for (int32 SlotIndex = 0; SlotIndex < NumMaterialSlots; SlotIndex++)
 			{
-				PreviewMeshComponent->OverrideMaterials.Add(PreviewMaterial);
+				PreviewMeshComponent->OverrideMaterials[SlotIndex] = MaterialSlots[SlotIndex].bSlotEnabled ? PreviewMaterial : nullptr;
 			}
-
 			PreviewMeshComponent->MarkRenderStateDirty();
 		}
 		
 		PostProcessVolumeActor = nullptr;
 	}
+
+	if (PreviewMaterial)
+	{
+		PreviewMaterial->PreviewMesh = PreviewMeshComponent;
+		FMaterialEditor::UpdateThumbnailInfoPreviewMesh(PreviewMaterial);
+		PreviewMaterial->MarkPackageDirty();
+	}
+
+	RefreshViewport();
 }
 
 void SMaterialEditor3DPreviewViewport::OnAddedToTab( const TSharedRef<SDockTab>& OwnerTab )
@@ -759,6 +790,12 @@ void SMaterialEditor3DPreviewViewport::BindCommands()
 		FCanExecuteAction(),
 		FIsActionChecked::CreateSP( this, &SMaterialEditor3DPreviewViewport::IsPreviewMeshFromSelectionChecked ) );
 
+	CommandList->MapAction(
+		Commands.SelectPreviewMeshMaterialSlots,
+		FExecuteAction::CreateSP(this, &SMaterialEditor3DPreviewViewport::OnSetPreviewMeshMaterialSlots),
+		FCanExecuteAction(),
+		FIsActionChecked::CreateSP(this, &SMaterialEditor3DPreviewViewport::DoesPreviewMeshContainSlots));
+
 	CommandList->MapAction(
 		Commands.TogglePreviewBackground,
 		FExecuteAction::CreateSP( this, &SMaterialEditor3DPreviewViewport::TogglePreviewBackground ),
@@ -795,19 +832,6 @@ void SMaterialEditor3DPreviewViewport::OnSetPreviewPrimitive(EThumbnailPrimType
 		if (Primitive != nullptr)
 		{
 			SetPreviewAsset(Primitive);
-			
-			// Clear the thumbnail preview mesh
-			if (UMaterialInterface* MaterialInterface = GetPreviewingMaterialInterface())
-			{
-				MaterialInterface->PreviewMesh = nullptr;
-				FMaterialEditor::UpdateThumbnailInfoPreviewMesh(MaterialInterface);
-				if (!bInitialLoad)
-				{
-					MaterialInterface->MarkPackageDirty();
-				}
-			}
-			
-			RefreshViewport();
 		}
 	}
 }
@@ -817,53 +841,92 @@ bool SMaterialEditor3DPreviewViewport::IsPreviewPrimitiveChecked(EThumbnailPrimT
 	return PreviewPrimType == PrimType;
 }
 
-void SMaterialEditor3DPreviewViewport::OnSetPreviewMeshFromSelection()
+void SMaterialEditor3DPreviewViewport::ApplySelectedMesh(UObject* InAsset)
 {
-	bool bFoundPreviewMesh = false;
-	FEditorDelegates::LoadSelectedAssetsIfNeeded.Broadcast();
+	if (USkeletalMesh* SkeletalMesh = Cast<USkeletalMesh>(InAsset))
+	{
+		if (PreviewMaterial)
+		{
+			PreviewMaterial->SetMaterialUsage(MATUSAGE_SkeletalMesh);
+		}
+	}
+
+	SetPreviewAsset(InAsset);
+}
 
+void SMaterialEditor3DPreviewViewport::OnSetPreviewMeshFromSelection()
+{
 	UMaterialInterface* MaterialInterface = GetPreviewingMaterialInterface();
+	if (!MaterialInterface)
+	{
+		return;
+	}
 
+	FEditorDelegates::LoadSelectedAssetsIfNeeded.Broadcast();
+	UObject* InitialSelection = nullptr;
 	// Look for a selected asset that can be converted to a mesh component
-	for (FSelectionIterator SelectionIt(*GEditor->GetSelectedObjects()); SelectionIt && !bFoundPreviewMesh; ++SelectionIt)
+	for (FSelectionIterator SelectionIt(*GEditor->GetSelectedObjects()); SelectionIt; ++SelectionIt)
 	{
 		UObject* TestAsset = *SelectionIt;
-		if (TestAsset->IsAsset())
+		if (TestAsset && (TestAsset->IsA<UStaticMesh>() || TestAsset->IsA<USkeletalMesh>()))
 		{
-			if (TSubclassOf<UActorComponent> ComponentClass = FComponentAssetBrokerage::GetPrimaryComponentForAsset(TestAsset->GetClass()))
-			{
-				if (ComponentClass->IsChildOf(UMeshComponent::StaticClass()))
-				{
-					if (USkeletalMesh* SkeletalMesh = Cast<USkeletalMesh>(TestAsset))
-					{
-						// Special case handling for skeletal meshes, sets the material to be usable with them
-						MaterialInterface->SetMaterialUsage(MATUSAGE_SkeletalMesh);
-					}
-
-					SetPreviewAsset(TestAsset);
-					MaterialInterface->PreviewMesh = TestAsset->GetPathName();
-					bFoundPreviewMesh = true;
-				}
-			}
+			InitialSelection = TestAsset;
+			break;
 		}
 	}
 
-	if (bFoundPreviewMesh)
-	{
-		FMaterialEditor::UpdateThumbnailInfoPreviewMesh(MaterialInterface);
-
-		MaterialInterface->MarkPackageDirty();
-		RefreshViewport();
-	}
-	else
+	if (InitialSelection)
 	{
-		FSuppressableWarningDialog::FSetupInfo Info(NSLOCTEXT("UnrealEd", "Warning_NoPreviewMeshFound_Message", "You need to select a mesh-based asset in the content browser to preview it."),
-			NSLOCTEXT("UnrealEd", "Warning_NoPreviewMeshFound", "Warning: No Preview Mesh Found"), "Warning_NoPreviewMeshFound");
-		Info.ConfirmText = NSLOCTEXT("UnrealEd", "Warning_NoPreviewMeshFound_Confirm", "Continue");
-		
-		FSuppressableWarningDialog NoPreviewMeshWarning( Info );
-		NoPreviewMeshWarning.ShowModal();
+		ApplySelectedMesh(InitialSelection);
 	}
+
+	FContentBrowserModule& CBModule =
+		FModuleManager::LoadModuleChecked<FContentBrowserModule>("ContentBrowser");
+
+	FAssetPickerConfig PickerConfig;
+
+	// Filter: allow only static & skeletal meshes
+	PickerConfig.Filter.ClassPaths.Add(UStaticMesh::StaticClass()->GetClassPathName());
+	PickerConfig.Filter.ClassPaths.Add(USkeletalMesh::StaticClass()->GetClassPathName());
+
+	PickerConfig.InitialAssetSelection = InitialSelection;
+	PickerConfig.bAllowNullSelection = false;
+	PickerConfig.bFocusSearchBoxWhenOpened = true;
+	PickerConfig.InitialAssetViewType = EAssetViewType::Tile;
+
+	PickerConfig.OnAssetSelected = FOnAssetSelected::CreateLambda(
+		[this](const FAssetData& AssetData)
+		{
+			if (UObject* SelectedAsset = AssetData.GetAsset())
+			{
+				ApplySelectedMesh(SelectedAsset);
+			}
+		}
+	);
+
+	TSharedRef<SWidget> AssetPicker =
+		CBModule.Get().CreateAssetPicker(PickerConfig);
+
+	TSharedRef<SWidget> WrappedPicker =
+		SNew(SBorder)
+		.Padding(FMargin(6.f))
+		.BorderImage(FAppStyle::Get().GetBrush("Menu.Background"))
+		[
+			SNew(SBox)
+				.WidthOverride(400.0f)
+				.HeightOverride(400.0f)
+				[
+					AssetPicker
+				]
+		];
+
+	const FVector2D ScreenPos = FSlateApplication::Get().GetCursorPos();
+	FSlateApplication::Get().PushMenu(
+		AsShared(),
+		FWidgetPath(),
+		WrappedPicker,
+		ScreenPos,
+		FPopupTransitionEffect(FPopupTransitionEffect::SubMenu));
 }
 
 bool SMaterialEditor3DPreviewViewport::IsPreviewMeshFromSelectionChecked() const
@@ -871,6 +934,27 @@ bool SMaterialEditor3DPreviewViewport::IsPreviewMeshFromSelectionChecked() const
 	return (PreviewPrimType == TPT_None && PreviewMeshComponent != nullptr);
 }
 
+void SMaterialEditor3DPreviewViewport::OnSetPreviewMeshMaterialSlots()
+{
+	TSharedRef<SWidget> MenuContent = GenerateMaterialSlotMenu();
+	TSharedRef<SWidget> OwnerWidget = AsShared();
+
+	const FVector2D ScreenPos = FSlateApplication::Get().GetCursorPos();
+
+	FSlateApplication::Get().PushMenu(
+		OwnerWidget,
+		FWidgetPath(),
+		MenuContent,
+		ScreenPos,
+		FPopupTransitionEffect(FPopupTransitionEffect::SubMenu)
+	);
+}
+
+bool SMaterialEditor3DPreviewViewport::DoesPreviewMeshContainSlots() const
+{
+	return PreviewMeshComponent != nullptr && PreviewMeshComponent->GetNumMaterials() > 0;
+}
+
 void SMaterialEditor3DPreviewViewport::TogglePreviewBackground()
 {
 	UAssetViewerSettings* Settings = UAssetViewerSettings::Get();
@@ -1004,6 +1088,115 @@ TSharedPtr<IPreviewProfileController> SMaterialEditor3DPreviewViewport::CreatePr
 	return MakeShared<FPreviewProfileController>();
 }
 
+void SMaterialEditor3DPreviewViewport::ApplyStateToAllSlots(bool bState)
+{
+	for (const MaterialSlotPreviewInfo& SlotInfo : MaterialSlots)
+	{
+		SlotInfo.bSlotEnabled = bState;
+	}
+
+	SetPreviewMaterial(PreviewMaterial);
+}
+
+void SMaterialEditor3DPreviewViewport::OnSelectAllSlots()
+{
+	ApplyStateToAllSlots(true);
+}
+
+void SMaterialEditor3DPreviewViewport::OnSelectNoSlots()
+{
+	ApplyStateToAllSlots(false);
+}
+
+void SMaterialEditor3DPreviewViewport::OnMaterialSlotCheckStateChanged(ECheckBoxState NewState, int32 SlotIndex)
+{
+	if (MaterialSlots.IsValidIndex(SlotIndex))
+	{
+		MaterialSlots[SlotIndex].bSlotEnabled = NewState == ECheckBoxState::Checked;
+		SetPreviewMaterial(PreviewMaterial);
+	}
+}
+
+ECheckBoxState SMaterialEditor3DPreviewViewport::IsMaterialSlotChecked(int32 SlotIndex) const
+{
+	return MaterialSlots.IsValidIndex(SlotIndex) && MaterialSlots[SlotIndex].bSlotEnabled ? ECheckBoxState::Checked : ECheckBoxState::Unchecked;
+}
+
+TSharedRef<SWidget> SMaterialEditor3DPreviewViewport::GenerateMaterialSlotMenu()
+{
+	FMenuBuilder MenuBuilder(/*bInShouldCloseWindowAfterMenuSelection*/ false, nullptr);
+
+	MenuBuilder.AddMenuEntry(
+		LOCTEXT("SelectAllSlots", "All Slots"),
+		LOCTEXT("SelectAllSlots_Tooltip", "Apply the preview material to all slots."),
+		FSlateIcon(),
+		FUIAction(FExecuteAction::CreateSP(this, &SMaterialEditor3DPreviewViewport::OnSelectAllSlots))
+	);
+
+	MenuBuilder.AddMenuEntry(
+		LOCTEXT("SelectNoSlots", "None"),
+		LOCTEXT("SelectNoSlots_Tooltip", "Apply the preview material to no slots."),
+		FSlateIcon(),
+		FUIAction(FExecuteAction::CreateSP(this, &SMaterialEditor3DPreviewViewport::OnSelectNoSlots))
+	);
+
+	MenuBuilder.AddSeparator();
+
+	TSharedPtr<FAssetThumbnailPool> ThumbnailPool = UThumbnailManager::Get().GetSharedThumbnailPool();	
+	TArray<TSharedPtr<FAssetThumbnail>> SlotThumbnails;
+	const int32 ThumbSize = 40;
+	for (int32 SlotIndex = 0; SlotIndex < MaterialSlots.Num(); SlotIndex++)
+	{
+		SlotThumbnails.Add(
+			MakeShareable(new FAssetThumbnail(MaterialSlots[SlotIndex].OriginalMaterial, ThumbSize, ThumbSize, ThumbnailPool))
+		);
+
+		const FText SlotLabel = FText::Format(
+			LOCTEXT("SlotLabelFmt", "{0}: {1}"),
+			FText::AsNumber(SlotIndex),
+			FText::FromName(MaterialSlots[SlotIndex].SlotName)
+		);
+
+		TSharedRef<SHorizontalBox> Row =
+			SNew(SHorizontalBox)
+			+ SHorizontalBox::Slot()
+			.AutoWidth()
+			.VAlign(VAlign_Center)
+			.Padding(2.f, 0.f)
+			[
+				SNew(SBox)
+					.WidthOverride(40.f)
+					.HeightOverride(40.f)
+					[
+						SlotThumbnails.IsValidIndex(SlotIndex)
+							? SlotThumbnails[SlotIndex]->MakeThumbnailWidget()
+							: SNullWidget::NullWidget
+					]
+			]
+		+ SHorizontalBox::Slot()
+			.FillWidth(1.f)
+			.VAlign(VAlign_Center)
+			.Padding(4.f, 0.f)
+			[
+				SNew(SCheckBox)
+					.IsChecked(this, &SMaterialEditor3DPreviewViewport::IsMaterialSlotChecked, SlotIndex)
+					.OnCheckStateChanged(this,
+						&SMaterialEditor3DPreviewViewport::OnMaterialSlotCheckStateChanged, SlotIndex)
+					[
+						SNew(STextBlock)
+							.Text(SlotLabel)
+					]
+			];
+
+		MenuBuilder.AddWidget(
+			Row,
+			FText::GetEmpty()
+		);
+	}
+
+	return MenuBuilder.MakeWidget();
+}
+
 void SMaterialEditor3DPreviewViewport::PopulateViewportOverlays(TSharedRef<class SOverlay> Overlay)
 {
 	// add the feature level display widget
diff --git a/Engine/Source/Editor/MaterialEditor/Private/SMaterialEditorViewport.h b/Engine/Source/Editor/MaterialEditor/Private/SMaterialEditorViewport.h
--- a/Engine/Source/Editor/MaterialEditor/Private/SMaterialEditorViewport.h
+++ b/Engine/Source/Editor/MaterialEditor/Private/SMaterialEditorViewport.h
@@ -22,6 +22,14 @@ class SMaterialEditorUIPreviewZoomer;
 class UMaterialInterface;
 class UMeshComponent;
 
+/** Helped struct for containing data for the material slot submenu UI and where to apply the preview material */
+struct MaterialSlotPreviewInfo
+{
+	FName SlotName;
+	TObjectPtr<UMaterialInterface> OriginalMaterial;
+	mutable bool bSlotEnabled;
+};
+
 /**
  * Material Editor Preview viewport widget
  */
@@ -83,6 +91,9 @@ public:
 	bool IsPreviewPrimitiveChecked(EThumbnailPrimType PrimType) const;
 	void OnSetPreviewMeshFromSelection();
 	bool IsPreviewMeshFromSelectionChecked() const;
+	void OnSetPreviewMeshMaterialSlots();
+	bool DoesPreviewMeshContainSlots() const;
+
 	void TogglePreviewBackground();
 	bool IsTogglePreviewBackgroundChecked() const;
 
@@ -97,6 +108,8 @@ protected:
 	virtual TSharedRef<FEditorViewportClient> MakeEditorViewportClient() override;
 	virtual TSharedPtr<SWidget> BuildViewportToolbar() override;
 	virtual TSharedPtr<IPreviewProfileController> CreatePreviewProfileController() override;
+	TSharedRef<SWidget> GenerateMaterialSlotMenu();
+
 	virtual void PopulateViewportOverlays(TSharedRef<class SOverlay> Overlay) override;
 	virtual EVisibility OnGetViewportContentVisibility() const override;
 	virtual void BindCommands() override;
@@ -107,6 +120,24 @@ private:
 
 	void OnPropertyChanged(UObject* ObjectBeingModified, FPropertyChangedEvent& PropertyChangedEvent);
 
+	/* Apply the selected mesh object to the preview mesh component */
+	void ApplySelectedMesh(UObject* InAsset);
+
+	/** Helper function for applying state logic to all material slots + updates*/
+	void ApplyStateToAllSlots(bool bState);
+
+	/** Set true for usage to all material slots */
+	void OnSelectAllSlots();
+
+	/** Set false for usage to all material slots */
+	void OnSelectNoSlots();
+
+	/** Updates the scene preview mesh when a checkbox state changes */
+	void OnMaterialSlotCheckStateChanged(ECheckBoxState NewState, int32 SlotIndex);
+
+	/** Returns if a material slot is checked/should show the preview material */
+	ECheckBoxState IsMaterialSlotChecked(int32 SlotIndex) const;
+
 	/** The parent tab where this viewport resides */
 	TWeakPtr<SDockTab> ParentTab;
 
@@ -138,6 +169,9 @@ private:
 
 	/** Handle to the registered OnPreviewFeatureLevelChanged delegate. */
 	FDelegateHandle PreviewFeatureLevelChangedHandle;
+
+	/** Material Slot data for populating and processing the material slot selections */
+	TArray<const MaterialSlotPreviewInfo> MaterialSlots;
 };
 
 /**
diff --git a/Engine/Source/Editor/MaterialEditor/Public/MaterialEditorActions.h b/Engine/Source/Editor/MaterialEditor/Public/MaterialEditorActions.h
--- a/Engine/Source/Editor/MaterialEditor/Public/MaterialEditorActions.h
+++ b/Engine/Source/Editor/MaterialEditor/Public/MaterialEditorActions.h
@@ -75,6 +75,9 @@ public:
 	
 	/** Sets the preview mesh to the current selection in the level editor */
 	TSharedPtr< FUICommandInfo > SetPreviewMeshFromSelection;
+
+	/** Set the material slots of the PreviewMesh to apply this material to */
+	TSharedPtr< FUICommandInfo > SelectPreviewMeshMaterialSlots;
 	
 	/** Toggles the preview pane's background */
 	TSharedPtr< FUICommandInfo > TogglePreviewBackground;
