Substrate - added blend mode tooltip.

#tests Pc editor
#rnx

#Synced-CL 49615688

diff --git a/Engine/Source/Editor/MaterialEditor/Private/MaterialEditorDetailCustomization.cpp b/Engine/Source/Editor/MaterialEditor/Private/MaterialEditorDetailCustomization.cpp
--- a/Engine/Source/Editor/MaterialEditor/Private/MaterialEditorDetailCustomization.cpp
+++ b/Engine/Source/Editor/MaterialEditor/Private/MaterialEditorDetailCustomization.cpp
@@ -68,6 +68,20 @@ void UpdateBlendModeEnum()
 
 		// BLEND_TranslucentColoredTransmittance is only supported in Substrate mode
 		BlendModeEnum->SetMetaData(TEXT("DisplayName"), TEXT("TranslucentColoredTransmittance"), BLEND_TranslucentColoredTransmittance);
+
+		static const FString ToolTipName = TEXT("ToolTip");
+		static const FString TranslucentToolTip = TEXT("Translucent surface accounting for coverage. Colored transmittance over the background is not possible and will be converted to grey scale transmittance.\nThis is the same as the legacy \"Translucent\" blend mode.");
+		static const FString ModulateToolTip = TEXT("Surface where only transmittance is applied over the background.\nThis is similar to the legacy \"Modulate\" blend mode.");
+		BlendModeEnum->SetMetaData(*ToolTipName, TEXT("Opaque surface"), BLEND_Opaque);
+		BlendModeEnum->SetMetaData(*ToolTipName, TEXT("Masked opaque surface, visible only when opacity is greater than the OpacityMaskClipValue from the material."), BLEND_Masked);
+		BlendModeEnum->SetMetaData(*ToolTipName, *TranslucentToolTip, BLEND_Translucent);
+		BlendModeEnum->SetMetaData(*ToolTipName, TEXT("Surface where only emissive and coverage are accounted for."), BLEND_Additive);
+		BlendModeEnum->SetMetaData(*ToolTipName, *ModulateToolTip, BLEND_Modulate);
+		BlendModeEnum->SetMetaData(*ToolTipName, TEXT("Translucent surface for which the throughput over the background can be overriden independently from the material itself."), BLEND_AlphaComposite);
+		BlendModeEnum->SetMetaData(*ToolTipName, TEXT("Alpha hold out blend mode use to punch holes in translucent surfaces."), BLEND_AlphaHoldout);
+		BlendModeEnum->SetMetaData(*ToolTipName, TEXT("Translucent surface accounting for coverage as well as colored transmittance over the background.\nThis is similar to the legacy \"ThinTranslucent\" shading model."), BLEND_TranslucentColoredTransmittance);
+		BlendModeEnum->SetMetaData(*ToolTipName, *TranslucentToolTip, BLEND_TranslucentGreyTransmittance);
+		BlendModeEnum->SetMetaData(*ToolTipName, *ModulateToolTip, BLEND_ColoredTransmittanceOnly);
 	}
 	else
 	{
