Material Editor pixel shader stats always use skylight base pass permutation
* Fixes missing stats on translucent materials: as of cl 47605465 they always use the skylight permutation

#Synced-CL 49613554

diff --git a/Engine/Source/Editor/MaterialEditor/Private/MaterialStatsCommon.cpp b/Engine/Source/Editor/MaterialEditor/Private/MaterialStatsCommon.cpp
--- a/Engine/Source/Editor/MaterialEditor/Private/MaterialStatsCommon.cpp
+++ b/Engine/Source/Editor/MaterialEditor/Private/MaterialStatsCommon.cpp
@@ -312,16 +312,16 @@ void FMaterialStatsUtils::GetRepresentativeShaderTypesAndDescriptions(TMap<FName
 		if (TargetMaterial->GetShadingModels().IsUnlit())
 		{
 			//unlit materials are never lightmapped
-			static FName TBasePassPSFNoLightMapPolicyName = TEXT("TBasePassPSFNoLightMapPolicy");
+			static FName TBasePassPSFNoLightMapPolicyName = TEXT("TBasePassPSFNoLightMapPolicySkylight");
 			ShaderTypeNamesAndDescriptions.FindOrAdd(FLocalVertexFactoryName)
 				.Add(FRepresentativeShaderInfo(ERepresentativeShader::StationarySurface, TBasePassPSFNoLightMapPolicyName, TEXT("Base pass shader without light map")));
 		}
 		else
 		{
 			//also show a dynamically lit shader
-			static FName TBasePassPSFNoLightMapPolicyName = TEXT("TBasePassPSFNoLightMapPolicy");
+			static FName TBasePassPSFNoLightMapPolicyName = TEXT("TBasePassPSFNoLightMapPolicySkylight");
 			ShaderTypeNamesAndDescriptions.FindOrAdd(FLocalVertexFactoryName)
-				.Add(FRepresentativeShaderInfo(ERepresentativeShader::DynamicallyLitObject, TBasePassPSFNoLightMapPolicyName, TEXT("Base pass shader")));
+				.Add(FRepresentativeShaderInfo(ERepresentativeShader::DynamicallyLitObject, TBasePassPSFNoLightMapPolicyName, TEXT("Base pass shader with skylight")));
 
 			if (IsStaticLightingAllowed())
 			{
