Material function editor no longer complains if a skeletal mesh is set on the preview window.
We use CheckMaterialUsage() to determine if skel mesh usage is supported and that will turn on the usage if auto usage is set on the preview material.
Material function preview material always has that enabled.
#rnx
#Synced-CL 49843598

diff --git a/Engine/Source/Editor/MaterialEditor/Private/MaterialEditor.cpp b/Engine/Source/Editor/MaterialEditor/Private/MaterialEditor.cpp
--- a/Engine/Source/Editor/MaterialEditor/Private/MaterialEditor.cpp
+++ b/Engine/Source/Editor/MaterialEditor/Private/MaterialEditor.cpp
@@ -2140,7 +2140,7 @@ bool FMaterialEditor::ApproveSetPreviewAsset(UObject* InAsset)
 	// Only permit the use of a skeletal mesh if the material has SkeletalMesh usage flag.
 	if (USkeletalMesh* SkeletalMesh = Cast<USkeletalMesh>(InAsset))
 	{
-		if (!Material->GetUsageByFlag(MATUSAGE_SkeletalMesh))
+		if (!Material->CheckMaterialUsage(MATUSAGE_SkeletalMesh))
 		{
 			FMessageDialog::Open(EAppMsgType::Ok, NSLOCTEXT("UnrealEd", "Error_MaterialEditor_CantPreviewOnSkelMesh", "Can't preview on the specified skeletal mesh because the material has not been compiled with SkeletalMesh usage flag."));
 			bApproved = false;
diff --git a/Engine/Source/Editor/MaterialEditor/Private/SMaterialEditorViewport.cpp b/Engine/Source/Editor/MaterialEditor/Private/SMaterialEditorViewport.cpp
--- a/Engine/Source/Editor/MaterialEditor/Private/SMaterialEditorViewport.cpp
+++ b/Engine/Source/Editor/MaterialEditor/Private/SMaterialEditorViewport.cpp
@@ -843,14 +843,6 @@ bool SMaterialEditor3DPreviewViewport::IsPreviewPrimitiveChecked(EThumbnailPrimT
 
 void SMaterialEditor3DPreviewViewport::ApplySelectedMesh(UObject* InAsset)
 {
-	if (USkeletalMesh* SkeletalMesh = Cast<USkeletalMesh>(InAsset))
-	{
-		if (PreviewMaterial)
-		{
-			PreviewMaterial->SetMaterialUsage(MATUSAGE_SkeletalMesh);
-		}
-	}
-
 	SetPreviewAsset(InAsset);
 }
 
