Material editor. Fix issue where custom preview mesh would revert to sphere mesh during material edits.
#rnx
#Synced-CL 49843598

diff --git a/Engine/Source/Editor/MaterialEditor/Private/SMaterialEditorViewport.cpp b/Engine/Source/Editor/MaterialEditor/Private/SMaterialEditorViewport.cpp
--- a/Engine/Source/Editor/MaterialEditor/Private/SMaterialEditorViewport.cpp
+++ b/Engine/Source/Editor/MaterialEditor/Private/SMaterialEditorViewport.cpp
@@ -354,6 +354,12 @@ bool SMaterialEditor3DPreviewViewport::SetPreviewAsset(UObject* InAsset)
 		return false;
 	}
 
+	if (PreviewMeshComponent == InAsset)
+	{
+		// No change
+		return true;
+	}
+
 	// Unregister the current component
 	if (PreviewMeshComponent != nullptr)
 	{
