Fix a bug where base property and usage flag changes in a parent material would not propogate to loaded children MaterialInstances and not propogate to opened Material Instance Editors.

- The issue was that we would call `RebuildMaterialInstanceEditors` to refresh the MI Editors _before_ we actually updated the MI itself.
- The MI itself would get updated in the destructor of `FMaterialUpdateContext UpdateContext`, but the call to `RebuildMaterialInstanceEditors` was still in the scope of the update context causing it to pull stale data.
- To fix this I've moved the call to RebuildMaterialInstanceEditors and the UpdateContext to inside `RecompileMaterialInternal`.  This way each call site doesn't need to manage the context and the call to RebuildMaterialInstanceEditors.
- The call to `UMaterialEditingLibrary::SetBaseMaterialUsage` was correctly calling RebuildMaterialInstanceEditors outside the scope of the update context.  To avoid a duplicate call, I've removed the call to rebuild instance editors.
- Remove UE_API on private functions in the mat editing library.

[RN] minor | Rendering - Materials and Shaders
#Synced-CL 49784419

diff --git a/Engine/Source/Editor/MaterialEditor/Private/MaterialEditingLibrary.cpp b/Engine/Source/Editor/MaterialEditor/Private/MaterialEditingLibrary.cpp
--- a/Engine/Source/Editor/MaterialEditor/Private/MaterialEditingLibrary.cpp
+++ b/Engine/Source/Editor/MaterialEditor/Private/MaterialEditingLibrary.cpp
@@ -698,15 +698,10 @@ void UMaterialEditingLibrary::SetBaseMaterialUsage(UMaterial* Material, EMateria
 		return;
 	}
 
-	{
-		FMaterialUpdateContext UpdateContext;
-		Material->SetUsageByFlag(Usage, bValue);
-		RecompileMaterialInternal(UpdateContext, Material);
-	}
+	Material->SetUsageByFlag(Usage, bValue);
+	RecompileMaterialInternal( { Material }, nullptr);
 
 	FEditorSupportDelegates::MaterialUsageFlagsChanged.Broadcast(Material, Usage);
-
-	RebuildMaterialInstanceEditors(Material);
 }
 
 bool UMaterialEditingLibrary::HasMaterialUsageOverride(UMaterialInstanceConstant* MaterialInstance, EMaterialUsage Usage)
@@ -804,32 +799,55 @@ bool UMaterialEditingLibrary::ConnectMaterialExpressions(UMaterialExpression* Fr
 }
 
 
-void UMaterialEditingLibrary::RecompileMaterialInternal(FMaterialUpdateContext& UpdateContext, UMaterial* Material)
+void UMaterialEditingLibrary::RecompileMaterialInternal(const TArray<UMaterial*>& Materials, const FOnItemComplete* OnItemComplete)
 {
-	UpdateContext.AddMaterial(Material);
+	// 1. Update all children Material Instances via the update context
+	{
+		FMaterialUpdateContext UpdateContext;
 
-	// Propagate the change to this material
-	Material->PreEditChange(nullptr);
-	Material->PostEditChange();
+		for (UMaterial* Material : Materials)
+		{
+			if (ensureAsRuntimeWarning(Material != nullptr))
+			{
+				UpdateContext.AddMaterial(Material);
 
-	Material->MarkPackageDirty();
+				// Propagate the change to this material
+				Material->PreEditChange(nullptr);
+				Material->PostEditChange();
 
-	// Force particle components to update their view relevance.
-	for (TObjectIterator<UParticleSystemComponent> It(/*AdditionalExclusionFlags = */RF_ClassDefaultObject, /*bIncludeDerivedClasses = */true, /*InInternalExclusionFlags = */EInternalObjectFlags::Garbage); It; ++It)
-	{
-		It->bIsViewRelevanceDirty = true;
+				Material->MarkPackageDirty();
+
+				// Force particle components to update their view relevance.
+				for (TObjectIterator<UParticleSystemComponent> It(/*AdditionalExclusionFlags = */RF_ClassDefaultObject, /*bIncludeDerivedClasses = */true, /*InInternalExclusionFlags = */EInternalObjectFlags::Garbage); It; ++It)
+				{
+					It->bIsViewRelevanceDirty = true;
+				}
+
+				// Update parameter names on any child material instances
+				for (TObjectIterator<UMaterialInstance> It(/*AdditionalExclusionFlags = */RF_ClassDefaultObject, /*bIncludeDerivedClasses = */true, /*InInternalExclusionFlags = */EInternalObjectFlags::Garbage); It; ++It)
+				{
+					if (It->Parent == Material)
+					{
+						It->UpdateParameterNames();
+					}
+				}
+
+				if (OnItemComplete)
+				{
+					OnItemComplete->Execute();
+				}
+			}
+		}
 	}
 
-	// Update parameter names on any child material instances
-	for (TObjectIterator<UMaterialInstance> It(/*AdditionalExclusionFlags = */RF_ClassDefaultObject, /*bIncludeDerivedClasses = */true, /*InInternalExclusionFlags = */EInternalObjectFlags::Garbage); It; ++It)
+	// 2. Once the MIs are updated, update the opened MI Editors.
+	for (UMaterial* Material : Materials)
 	{
-		if (It->Parent == Material)
+		if (Material)
 		{
-			It->UpdateParameterNames();
+			UMaterialEditingLibrary::RebuildMaterialInstanceEditors(Material);
 		}
 	}
-
-	UMaterialEditingLibrary::RebuildMaterialInstanceEditors(Material);
 }
 
 void UMaterialEditingLibrary::RecompileMaterial(UMaterial* Material)
@@ -839,12 +857,7 @@ void UMaterialEditingLibrary::RecompileMaterial(UMaterial* Material)
 		return;
 	}
 
-	{
-		// Create a material update context so we can safely update materials using this function.
-		FMaterialUpdateContext UpdateContext;
-
-		RecompileMaterialInternal(UpdateContext, Material);
-	}
+	RecompileMaterialInternal( { Material }, nullptr);
 
 	// Update the world's viewports
 	FEditorDelegates::RefreshEditor.Broadcast();
@@ -861,19 +874,7 @@ void UMaterialEditingLibrary::RecompileMaterials(TArray<UMaterial*>& Materials,
 		return;
 	}
 
-	{
-		// Create a material update context so we can safely update materials using this function.
-		FMaterialUpdateContext UpdateContext;
-
-		for (UMaterial* Material : Materials)
-		{
-			if (ensureAsRuntimeWarning(Material != nullptr))
-			{
-				RecompileMaterialInternal(UpdateContext, Material);
-				OnItemComplete.Execute();
-			}
-		}
-	}
+	RecompileMaterialInternal(Materials, &OnItemComplete);
 
 	// Update the world's viewports
 	FEditorDelegates::RefreshEditor.Broadcast();
diff --git a/Engine/Source/Editor/MaterialEditor/Public/MaterialEditingLibrary.h b/Engine/Source/Editor/MaterialEditor/Public/MaterialEditingLibrary.h
--- a/Engine/Source/Editor/MaterialEditor/Public/MaterialEditingLibrary.h
+++ b/Engine/Source/Editor/MaterialEditor/Public/MaterialEditingLibrary.h
@@ -523,8 +523,8 @@ public:
 	static UE_API UMaterialInterface* GetNaniteOverrideMaterial(UMaterialInterface* Material);
 
 private:
-	static UE_API void RecompileMaterialInternal(FMaterialUpdateContext& UpdateContext, UMaterial* Material);
-	static UE_API void UpdateMaterialFunctionInternal(FMaterialUpdateContext& UpdateContext, UMaterialFunctionInterface* MaterialFunction, UMaterial* PreviewMaterial);
+	static void RecompileMaterialInternal(const TArray<UMaterial*>& Material, const FOnItemComplete* OnItemComplete);
+	static void UpdateMaterialFunctionInternal(FMaterialUpdateContext& UpdateContext, UMaterialFunctionInterface* MaterialFunction, UMaterial* PreviewMaterial);
 };
 
 #undef UE_API
