Fix bug where material edits via the property matrix editor would not propagate to loaded material instances and material instance editors.

- In order to keep the property matrix editor and the material editor decoupled this has been implemented as a delegate exposed my the PropertyEditorModule that the MaterialEditorModule subscribes too.
- I've structured this in such a way that we do a bulk update of all the materials saved.
- FPropertyEditorModule exposes a OnPropertyMatrixAssetsSaved one parameter delegate which takes the array of objects that were saved.
- FPropertyEditorToolkit overrides OnAssetsSaved on the base class FAssetEditorToolkit.  This overridden implementation simple broadcasts the delegate OnPropertyMatrixAssetsSaved passing in the saved objects array.
- In FMaterialEditorModule::StartupModule() we subscribe to the delegate and simply call UMaterialEditingLibrary::RecompileMaterials to do the refresh.

## Misc
- Remove unused delegate handle ContentBrowserAssetExtenderDelegateHandle
- Fix a bug where we weren't unregistering the custom class layout for MaterialExpressionOperator when we shutdown.

[RN] major | Rendering - Materials and Shaders

#Synced-CL 49784419

diff --git a/Engine/Source/Editor/MaterialEditor/Private/MaterialEditorModule.cpp b/Engine/Source/Editor/MaterialEditor/Private/MaterialEditorModule.cpp
--- a/Engine/Source/Editor/MaterialEditor/Private/MaterialEditorModule.cpp
+++ b/Engine/Source/Editor/MaterialEditor/Private/MaterialEditorModule.cpp
@@ -26,6 +26,7 @@
 #include "SubstrateMaterialEditorStyle.h"
 #include "MaterialEditorDetailCustomization.h"
 #include "MaterialGraphConnectionDrawingPolicy.h"
+#include "MaterialEditingLibrary.h"
 
 const FName MaterialEditorAppIdentifier = FName(TEXT("MaterialEditorApp"));
 const FName MaterialInstanceEditorAppIdentifier = FName(TEXT("MaterialInstanceEditorApp"));
@@ -88,6 +89,21 @@ public:
 			FOnGetDetailCustomizationInstance::CreateStatic(&FMaterialExpressionOperatorCustomization::MakeInstance)
 		);
 
+		// Register a callback to bulk update materials when many materials are edited and saved via the property matrix.
+		OnPropertyMatrixAssetsSavedDelegateHandle = PropertyModule.OnPropertyMatrixAssetsSaved().AddLambda([](const TArray<UObject*>& SavedObjects)
+		{
+			TArray<UMaterial*> Materials;
+			Materials.Reserve(SavedObjects.Num());
+			for (UObject* Object : SavedObjects)
+			{
+				if (UMaterial* Material = Cast<UMaterial>(Object))
+				{
+					Materials.Add(Material);
+				}
+			}
+			UMaterialEditingLibrary::RecompileMaterials(Materials, UMaterialEditingLibrary::FOnItemComplete::CreateLambda([]() {}));
+		});
+
 		GraphPanelPinFactory = MakeShared<FMaterialEditorGraphPanelPinFactory>();
 		FEdGraphUtilities::RegisterVisualPinFactory(GraphPanelPinFactory);
 
@@ -111,6 +127,8 @@ public:
 		{
 			FPropertyEditorModule& PropertyModule = FModuleManager::GetModuleChecked<FPropertyEditorModule>("PropertyEditor");
 			PropertyModule.UnregisterCustomClassLayout("MaterialExpressionColorRamp");
+			PropertyModule.UnregisterCustomClassLayout("MaterialExpressionOperator");
+			PropertyModule.OnPropertyMatrixAssetsSaved().Remove(OnPropertyMatrixAssetsSavedDelegateHandle);
 		}
 		
 		FSubstrateMaterialEditorStyle::Shutdown();
@@ -169,9 +187,9 @@ public:
 private:
 	TSharedPtr<FExtensibilityManager> MenuExtensibilityManager;
 	TSharedPtr<FExtensibilityManager> ToolBarExtensibilityManager;
-	FDelegateHandle ContentBrowserAssetExtenderDelegateHandle;
 	TSharedPtr<FMaterialEditorGraphPanelPinFactory> GraphPanelPinFactory;
 	TSharedPtr<FGraphPanelPinConnectionFactory> GraphConnectionFactory;
+	FDelegateHandle OnPropertyMatrixAssetsSavedDelegateHandle;
 };
 
 IMPLEMENT_MODULE( FMaterialEditorModule, MaterialEditor );
diff --git a/Engine/Source/Editor/PropertyEditor/Private/PropertyEditorToolkit.cpp b/Engine/Source/Editor/PropertyEditor/Private/PropertyEditorToolkit.cpp
--- a/Engine/Source/Editor/PropertyEditor/Private/PropertyEditorToolkit.cpp
+++ b/Engine/Source/Editor/PropertyEditor/Private/PropertyEditorToolkit.cpp
@@ -369,6 +369,11 @@ bool FPropertyEditorToolkit::TableHasCustomColumns() const
 	return PropertyPathsAddedAsColumns.Num() > 0;
 }
 
+void FPropertyEditorToolkit::OnAssetsSaved(const TArray<UObject*>& SavedObjects)
+{
+	FPropertyEditorModule& PropertyEditorModule = FModuleManager::GetModuleChecked<FPropertyEditorModule>("PropertyEditor");
+	PropertyEditorModule.OnPropertyMatrixAssetsSaved().Broadcast(SavedObjects);
+}
 
 bool FPropertyEditorToolkit::IsExposedAsColumn( const TWeakPtr< IPropertyTreeRow >& Row ) const
 {
diff --git a/Engine/Source/Editor/PropertyEditor/Private/PropertyEditorToolkit.h b/Engine/Source/Editor/PropertyEditor/Private/PropertyEditorToolkit.h
--- a/Engine/Source/Editor/PropertyEditor/Private/PropertyEditorToolkit.h
+++ b/Engine/Source/Editor/PropertyEditor/Private/PropertyEditorToolkit.h
@@ -48,6 +48,8 @@ public:
 
 	virtual bool IsPrimaryEditor() const override{ return false; };
 
+	virtual void OnAssetsSaved(const TArray<UObject*>& SavedObjects) override;
+
 public:
 
 	static TSharedRef<FPropertyEditorToolkit> CreateEditor( const EToolkitMode::Type Mode, const TSharedPtr< IToolkitHost >& InitToolkitHost, UObject* ObjectToEdit );
diff --git a/Engine/Source/Editor/PropertyEditor/Public/PropertyEditorModule.h b/Engine/Source/Editor/PropertyEditor/Public/PropertyEditorModule.h
--- a/Engine/Source/Editor/PropertyEditor/Public/PropertyEditorModule.h
+++ b/Engine/Source/Editor/PropertyEditor/Public/PropertyEditorModule.h
@@ -454,6 +454,9 @@ public:
 	DECLARE_EVENT(PropertyEditorModule, FPropertyEditorOpenedEvent);
 	virtual FPropertyEditorOpenedEvent& OnPropertyEditorOpened() { return PropertyEditorOpened; }
 
+	DECLARE_EVENT_OneParam(PropertyEditorModule, FOnPropertyMatrixAssetsSaved, const TArray<UObject*>&);
+	virtual FOnPropertyMatrixAssetsSaved& OnPropertyMatrixAssetsSaved() { return PropertyMatrixAssetsSaved; }
+
 	const FCustomDetailLayoutNameMap& GetClassNameToDetailLayoutNameMap() const { return ClassNameToDetailLayoutNameMap; }
 
 	/** Get the global row extension generators. */
@@ -514,6 +517,8 @@ private:
 	TMap<FName, TSharedPtr<FClassSectionMapping>> ClassSectionMappings;
 	/** Event to be called when a property editor is opened */
 	FPropertyEditorOpenedEvent PropertyEditorOpened;
+	/** Event to be called when assets are saved in bulk from the property matrix editor. */
+	FOnPropertyMatrixAssetsSaved PropertyMatrixAssetsSaved;
 	/** Mapping of registered floating UStructs to their struct proxy so they show correctly in the details panel */
 	TMap<FName, FStructProperty*> RegisteredStructToProxyMap;
 	/** Shared thumbnail pool used by property row generators */
