Added the same tooltip to material parameters in the Parameters panel of the material editor tab as in the material instance editor in order to help the user understand where the parameters come from (especially useful in combination with when r.Material.PedanticErrorChecksEnabled 1 is used, in order to detect where the differing parameter values come from)

#Synced-CL 49919419

diff --git a/Engine/Source/Editor/MaterialEditor/Private/MaterialPropertyHelpers.cpp b/Engine/Source/Editor/MaterialEditor/Private/MaterialPropertyHelpers.cpp
--- a/Engine/Source/Editor/MaterialEditor/Private/MaterialPropertyHelpers.cpp
+++ b/Engine/Source/Editor/MaterialEditor/Private/MaterialPropertyHelpers.cpp
@@ -1158,9 +1158,9 @@ FText FMaterialPropertyHelpers::GetParameterTooltip(UDEditorParameterValue* Para
 	return GetParameterTooltip(Parameter, static_cast<UMaterialEditorInstanceConstant*>(MaterialEditorInstance));
 }
 
-FText FMaterialPropertyHelpers::GetParameterTooltip(UDEditorParameterValue* Parameter, UMaterialEditorInstanceConstant* MaterialEditorInstance)
+FText FMaterialPropertyHelpers::GetParameterTooltip(UDEditorParameterValue* Parameter, UMaterial* Material)
 {
-	const FMaterialCachedExpressionData* CachedExpressionData = MaterialEditorInstance->Parent ? &MaterialEditorInstance->Parent->GetCachedExpressionData() : nullptr;
+	const FMaterialCachedExpressionData* CachedExpressionData = Material ? &Material->GetCachedExpressionData() : nullptr;
 
 	FText AssetPath;
 	if (CachedExpressionData)
@@ -1176,12 +1176,7 @@ FText FMaterialPropertyHelpers::GetParameterTooltip(UDEditorParameterValue* Para
 	}
 
 	FText TooltipText;
-	// If the material instance disallows the creation of new shader permutations, prevent overriding the static parameter.
-	if (Parameter->IsStaticParameter() && MaterialEditorInstance->SourceInstance->bDisallowStaticParameterPermutations)
-	{
-		return FText::FromString(TEXT("This material instance parent restricts the creation of new shader permutations. Overriding this parameter would result in the generation of additional shader permutations."));
-	}
-	else if (!Parameter->Description.IsEmpty())
+	if (!Parameter->Description.IsEmpty())
 	{
 		TooltipText = FText::Format(LOCTEXT("ParameterInfoDescAndLocation", "{0} \nFound in: {1}"), FText::FromString(Parameter->Description), AssetPath);
 	}
@@ -1193,6 +1188,18 @@ FText FMaterialPropertyHelpers::GetParameterTooltip(UDEditorParameterValue* Para
 	return TooltipText;
 }
 
+FText FMaterialPropertyHelpers::GetParameterTooltip(UDEditorParameterValue* Parameter, UMaterialEditorInstanceConstant* MaterialEditorInstance)
+{
+	// If the material instance disallows the creation of new shader permutations, prevent overriding the static parameter.
+	if (Parameter->IsStaticParameter() && MaterialEditorInstance->SourceInstance && MaterialEditorInstance->SourceInstance->bDisallowStaticParameterPermutations)
+	{
+		return FText::FromString(TEXT("This material instance parent restricts the creation of new shader permutations. Overriding this parameter would result in the generation of additional shader permutations."));
+	}
+	
+	// Otherwise, fallback to the generic material parameter tooltip 
+	return GetParameterTooltip(Parameter, MaterialEditorInstance->Parent ? MaterialEditorInstance->Parent->GetMaterial() : nullptr);
+}
+
 FResetToDefaultOverride FMaterialPropertyHelpers::CreateResetToDefaultOverride(UDEditorParameterValue* Parameter, UMaterialEditorInstanceConstant* MaterialEditorInstance, TSharedPtr<IPropertyHandle> ParameterValueProperty)
 {
 	UDEditorScalarParameterValue* ScalarParam = Cast<UDEditorScalarParameterValue>(Parameter);
diff --git a/Engine/Source/Editor/MaterialEditor/Private/SMaterialParametersOverviewWidget.cpp b/Engine/Source/Editor/MaterialEditor/Private/SMaterialParametersOverviewWidget.cpp
--- a/Engine/Source/Editor/MaterialEditor/Private/SMaterialParametersOverviewWidget.cpp
+++ b/Engine/Source/Editor/MaterialEditor/Private/SMaterialParametersOverviewWidget.cpp
@@ -96,6 +96,8 @@ void SMaterialParametersOverviewTreeItem::Construct(const FArguments& InArgs, co
 		LeftSideWidget = NodeWidgets.NameWidget.ToSharedRef();
 		RightSideWidget = NodeWidgets.ValueWidget.ToSharedRef();
 
+		LeftSideWidget->SetToolTipText(FMaterialPropertyHelpers::GetParameterTooltip(StackParameterData->Parameter, MaterialEditorInstance->PreviewMaterial));
+
 		RightSidePadding = FMargin(12.0f, 0.0f, 2.0f, 0.0f);	// Match DetailWidgetConstants::RightRowPadding
 	}
 // END PROPERTY
diff --git a/Engine/Source/Editor/MaterialEditor/Public/MaterialPropertyHelpers.h b/Engine/Source/Editor/MaterialEditor/Public/MaterialPropertyHelpers.h
--- a/Engine/Source/Editor/MaterialEditor/Public/MaterialPropertyHelpers.h
+++ b/Engine/Source/Editor/MaterialEditor/Public/MaterialPropertyHelpers.h
@@ -209,6 +209,7 @@ public:
 	static	UE_API FText GetParameterExpressionDescription(UDEditorParameterValue* Parameter, UObject* MaterialEditorInstance);
 	
 	static UE_API FText GetParameterTooltip(UDEditorParameterValue* Parameter, UMaterialEditorInstanceConstant* MaterialEditorInstance);
+	static UE_API FText GetParameterTooltip(UDEditorParameterValue* Parameter, UMaterial* Material);
 
 	UE_DEPRECATED(5.8, "Please use the fully typed method instead.")
 	static UE_API FText GetParameterTooltip(UDEditorParameterValue* Parameter, UObject* MaterialEditorInstance);
