[Backout] - CL49410344
#rnx

Original CL Desc
-----------------------------------------------------------------
Change the material editor preview mesh selector to use new dialogue boxes.
1 - The static mesh selection now pops up a content browser highlighting the originally selected mesh but also allowing the user to select a different mesh for the preview material.
2 - The material slot selection window allows the user to specify which material slots they want the preview material applied to (slot options change based on selected mesh and refresh if the mesh is swapped).

#Synced-CL 49944062

diff --git a/Engine/Source/Editor/MaterialEditor/Private/MaterialEditorActions.cpp b/Engine/Source/Editor/MaterialEditor/Private/MaterialEditorActions.cpp
--- a/Engine/Source/Editor/MaterialEditor/Private/MaterialEditorActions.cpp
+++ b/Engine/Source/Editor/MaterialEditor/Private/MaterialEditorActions.cpp
@@ -38,7 +38,6 @@ void FMaterialEditorCommands::RegisterCommands()
 	UI_COMMAND( SetPlanePreview, "Plane", "Sets the preview mesh to a plane primitive.", EUserInterfaceActionType::RadioButton, FInputChord() );
 	UI_COMMAND( SetCubePreview, "Cube", "Sets the preview mesh to a cube primitive.", EUserInterfaceActionType::RadioButton, FInputChord() );
 	UI_COMMAND( SetPreviewMeshFromSelection, "Mesh", "Sets the preview mesh based on the current content browser selection.", EUserInterfaceActionType::RadioButton, FInputChord() );
-	UI_COMMAND( SelectPreviewMeshMaterialSlots, "Mesh Slots", "Sets the preview mesh slots to apply the material to.", EUserInterfaceActionType::RadioButton, FInputChord() );
 	UI_COMMAND( TogglePreviewBackground, "Background", "Toggles the preview pane's background.", EUserInterfaceActionType::ToggleButton, FInputChord() );
 
 	UI_COMMAND( CameraHome, "Home", "Goes home on the canvas.", EUserInterfaceActionType::Button, FInputChord() );
diff --git a/Engine/Source/Editor/MaterialEditor/Private/MaterialEditorViewportToolbarSections.cpp b/Engine/Source/Editor/MaterialEditor/Private/MaterialEditorViewportToolbarSections.cpp
--- a/Engine/Source/Editor/MaterialEditor/Private/MaterialEditorViewportToolbarSections.cpp
+++ b/Engine/Source/Editor/MaterialEditor/Private/MaterialEditorViewportToolbarSections.cpp
@@ -3,7 +3,6 @@
 #include "MaterialEditorViewportToolbarSections.h"
 #include "EditorViewportCommands.h"
 #include "MaterialEditorActions.h"
-#include "Materials/Material.h"
 #include "SEditorViewport.h"
 #include "SMaterialEditorViewport.h"
 #include "ToolMenu.h"
@@ -122,13 +121,6 @@ void UE::MaterialEditor::ExtendPreviewSceneSettingsSubmenu(FName InSubmenuName)
 	MeshFromSelectionEntry.Icon = FSlateIconFinder::FindIconForClass(UStaticMesh::StaticClass());
 	MeshFromSelectionEntry.SetShowInToolbarTopLevel(true);
 	MeshFromSelectionEntry.ToolBarData = ToolbarData;
-
-	FToolMenuEntry& MaterialSlotsEntry =
-		PreviewMeshSection.AddMenuEntry(FMaterialEditorCommands::Get().SelectPreviewMeshMaterialSlots);
-	MaterialSlotsEntry.Label = LOCTEXT("SelectPreviewMeshMaterialSlotsLabel", "Mesh material slots to use when previewing material");
-	MaterialSlotsEntry.Icon = FSlateIconFinder::FindIconForClass(UMaterial::StaticClass());
-	MaterialSlotsEntry.SetShowInToolbarTopLevel(true);
-	MaterialSlotsEntry.ToolBarData = ToolbarData;
 }
 
 #undef LOCTEXT_NAMESPACE
diff --git a/Engine/Source/Editor/MaterialEditor/Private/SMaterialEditorViewport.cpp b/Engine/Source/Editor/MaterialEditor/Private/SMaterialEditorViewport.cpp
--- a/Engine/Source/Editor/MaterialEditor/Private/SMaterialEditorViewport.cpp
+++ b/Engine/Source/Editor/MaterialEditor/Private/SMaterialEditorViewport.cpp
@@ -46,9 +46,6 @@
 #include "MaterialEditorViewportToolbarSections.h"
 #include "ToolMenus.h"
 #include "ViewportToolbar/UnrealEdViewportToolbar.h"
-#include "AssetThumbnail.h"
-#include "ContentBrowserModule.h"
-#include "IContentBrowserSingleton.h"
 
 #define LOCTEXT_NAMESPACE "MaterialEditor"
 #include "AdvancedPreviewSceneMenus.h"
@@ -432,24 +429,6 @@ bool SMaterialEditor3DPreviewViewport::SetPreviewAsset(UObject* InAsset)
 		AdvancedPreviewScene->AddComponent(PreviewMeshComponent, Transform);
 		AdvancedPreviewScene->SetFloorOffset(-PreviewMeshComponent->Bounds.Origin.Z + PreviewMeshComponent->Bounds.BoxExtent.Z);
 
-		//Repopulate the material slot states as the mesh option is selected
-		MaterialSlots.Empty();
-		if (UStaticMesh* SM = Cast<UStaticMesh>(InAsset))
-		{
-			const TArray<FStaticMaterial>& StaticMatInfos = SM->GetStaticMaterials();
-			for (const FStaticMaterial& Slot : StaticMatInfos)
-			{
-				MaterialSlots.Add({ Slot.MaterialSlotName, Slot.MaterialInterface, true });
-			}
-		}
-		else if (USkeletalMesh* SK = Cast<USkeletalMesh>(InAsset))
-		{
-			const TArray<FSkeletalMaterial>& SkeletalMatInfos = SK->GetMaterials();
-			for (const FSkeletalMaterial& Slot : SkeletalMatInfos)
-			{
-				MaterialSlots.Add({ Slot.MaterialSlotName, Slot.MaterialInterface, true });
-			}
-		}
 	}
 
 	// Make sure the preview material is applied to the component
@@ -720,28 +699,18 @@ void SMaterialEditor3DPreviewViewport::SetPreviewMaterial(UMaterialInterface* In
 		// Add the preview material to the preview mesh.
 		if (PreviewMeshComponent != nullptr)
 		{
-			int32 NumMaterialSlots = PreviewMeshComponent->GetNumMaterials();
-			ensure(NumMaterialSlots == MaterialSlots.Num());
-			//Ensure the override materials has enough array slots to apply the preview material
-			PreviewMeshComponent->OverrideMaterials.SetNum(NumMaterialSlots, EAllowShrinking::Yes);
-			for (int32 SlotIndex = 0; SlotIndex < NumMaterialSlots; SlotIndex++)
+			PreviewMeshComponent->OverrideMaterials.Empty();
+
+			if (PreviewMaterial)
 			{
-				PreviewMeshComponent->OverrideMaterials[SlotIndex] = MaterialSlots[SlotIndex].bSlotEnabled ? PreviewMaterial : nullptr;
+				PreviewMeshComponent->OverrideMaterials.Add(PreviewMaterial);
 			}
+
 			PreviewMeshComponent->MarkRenderStateDirty();
 		}
 		
 		PostProcessVolumeActor = nullptr;
 	}
-
-	if (PreviewMaterial)
-	{
-		PreviewMaterial->PreviewMesh = PreviewMeshComponent;
-		FMaterialEditor::UpdateThumbnailInfoPreviewMesh(PreviewMaterial);
-		PreviewMaterial->MarkPackageDirty();
-	}
-
-	RefreshViewport();
 }
 
 void SMaterialEditor3DPreviewViewport::OnAddedToTab( const TSharedRef<SDockTab>& OwnerTab )
@@ -796,12 +765,6 @@ void SMaterialEditor3DPreviewViewport::BindCommands()
 		FCanExecuteAction(),
 		FIsActionChecked::CreateSP( this, &SMaterialEditor3DPreviewViewport::IsPreviewMeshFromSelectionChecked ) );
 
-	CommandList->MapAction(
-		Commands.SelectPreviewMeshMaterialSlots,
-		FExecuteAction::CreateSP(this, &SMaterialEditor3DPreviewViewport::OnSetPreviewMeshMaterialSlots),
-		FCanExecuteAction(),
-		FIsActionChecked::CreateSP(this, &SMaterialEditor3DPreviewViewport::DoesPreviewMeshContainSlots));
-
 	CommandList->MapAction(
 		Commands.TogglePreviewBackground,
 		FExecuteAction::CreateSP( this, &SMaterialEditor3DPreviewViewport::TogglePreviewBackground ),
@@ -838,6 +801,19 @@ void SMaterialEditor3DPreviewViewport::OnSetPreviewPrimitive(EThumbnailPrimType
 		if (Primitive != nullptr)
 		{
 			SetPreviewAsset(Primitive);
+			
+			// Clear the thumbnail preview mesh
+			if (UMaterialInterface* MaterialInterface = GetPreviewingMaterialInterface())
+			{
+				MaterialInterface->PreviewMesh = nullptr;
+				FMaterialEditor::UpdateThumbnailInfoPreviewMesh(MaterialInterface);
+				if (!bInitialLoad)
+				{
+					MaterialInterface->MarkPackageDirty();
+				}
+			}
+			
+			RefreshViewport();
 		}
 	}
 }
@@ -847,84 +823,47 @@ bool SMaterialEditor3DPreviewViewport::IsPreviewPrimitiveChecked(EThumbnailPrimT
 	return PreviewPrimType == PrimType;
 }
 
-void SMaterialEditor3DPreviewViewport::ApplySelectedMesh(UObject* InAsset)
-{
-	SetPreviewAsset(InAsset);
-}
-
 void SMaterialEditor3DPreviewViewport::OnSetPreviewMeshFromSelection()
 {
+	bool bFoundPreviewMesh = false;
+	FEditorDelegates::LoadSelectedAssetsIfNeeded.Broadcast();
+
 	UMaterialInterface* MaterialInterface = GetPreviewingMaterialInterface();
-	if (!MaterialInterface)
-	{
-		return;
-	}
 
-	FEditorDelegates::LoadSelectedAssetsIfNeeded.Broadcast();
-	UObject* InitialSelection = nullptr;
 	// Look for a selected asset that can be converted to a mesh component
-	for (FSelectionIterator SelectionIt(*GEditor->GetSelectedObjects()); SelectionIt; ++SelectionIt)
+	for (FSelectionIterator SelectionIt(*GEditor->GetSelectedObjects()); SelectionIt && !bFoundPreviewMesh; ++SelectionIt)
 	{
 		UObject* TestAsset = *SelectionIt;
-		if (TestAsset && (TestAsset->IsA<UStaticMesh>() || TestAsset->IsA<USkeletalMesh>()))
-		{
-			InitialSelection = TestAsset;
-			break;
-		}
-	}
-
-	if (InitialSelection)
-	{
-		ApplySelectedMesh(InitialSelection);
-	}
-
-	FContentBrowserModule& CBModule =
-		FModuleManager::LoadModuleChecked<FContentBrowserModule>("ContentBrowser");
-
-	FAssetPickerConfig PickerConfig;
-
-	// Filter: allow only static & skeletal meshes
-	PickerConfig.Filter.ClassPaths.Add(UStaticMesh::StaticClass()->GetClassPathName());
-	PickerConfig.Filter.ClassPaths.Add(USkeletalMesh::StaticClass()->GetClassPathName());
-
-	PickerConfig.InitialAssetSelection = InitialSelection;
-	PickerConfig.bAllowNullSelection = false;
-	PickerConfig.bFocusSearchBoxWhenOpened = true;
-	PickerConfig.InitialAssetViewType = EAssetViewType::Tile;
-
-	PickerConfig.OnAssetSelected = FOnAssetSelected::CreateLambda(
-		[this](const FAssetData& AssetData)
+		if (TestAsset->IsAsset())
 		{
-			if (UObject* SelectedAsset = AssetData.GetAsset())
+			if (TSubclassOf<UActorComponent> ComponentClass = FComponentAssetBrokerage::GetPrimaryComponentForAsset(TestAsset->GetClass()))
 			{
-				ApplySelectedMesh(SelectedAsset);
+				if (ComponentClass->IsChildOf(UMeshComponent::StaticClass()))
+				{
+					SetPreviewAsset(TestAsset);
+					MaterialInterface->PreviewMesh = TestAsset->GetPathName();
+					bFoundPreviewMesh = true;
+				}
 			}
 		}
-	);
-
-	TSharedRef<SWidget> AssetPicker =
-		CBModule.Get().CreateAssetPicker(PickerConfig);
+	}
 
-	TSharedRef<SWidget> WrappedPicker =
-		SNew(SBorder)
-		.Padding(FMargin(6.f))
-		.BorderImage(FAppStyle::Get().GetBrush("Menu.Background"))
-		[
-			SNew(SBox)
-				.WidthOverride(400.0f)
-				.HeightOverride(400.0f)
-				[
-					AssetPicker
-				]
-		];
+	if (bFoundPreviewMesh)
+	{
+		FMaterialEditor::UpdateThumbnailInfoPreviewMesh(MaterialInterface);
 
-	const FVector2D ScreenPos = FSlateApplication::Get().GetCursorPos();
-	FSlateApplication::Get().PushMenu(
-		AsShared(),
-		FWidgetPath(),
-		WrappedPicker,
-		ScreenPos,
-		FPopupTransitionEffect(FPopupTransitionEffect::SubMenu));
+		MaterialInterface->MarkPackageDirty();
+		RefreshViewport();
+	}
+	else
+	{
+		FSuppressableWarningDialog::FSetupInfo Info(NSLOCTEXT("UnrealEd", "Warning_NoPreviewMeshFound_Message", "You need to select a mesh-based asset in the content browser to preview it."),
+			NSLOCTEXT("UnrealEd", "Warning_NoPreviewMeshFound", "Warning: No Preview Mesh Found"), "Warning_NoPreviewMeshFound");
+		Info.ConfirmText = NSLOCTEXT("UnrealEd", "Warning_NoPreviewMeshFound_Confirm", "Continue");
+		
+		FSuppressableWarningDialog NoPreviewMeshWarning( Info );
+		NoPreviewMeshWarning.ShowModal();
+	}
 }
 
 bool SMaterialEditor3DPreviewViewport::IsPreviewMeshFromSelectionChecked() const
@@ -932,27 +871,6 @@ bool SMaterialEditor3DPreviewViewport::IsPreviewMeshFromSelectionChecked() const
 	return (PreviewPrimType == TPT_None && PreviewMeshComponent != nullptr);
 }
 
-void SMaterialEditor3DPreviewViewport::OnSetPreviewMeshMaterialSlots()
-{
-	TSharedRef<SWidget> MenuContent = GenerateMaterialSlotMenu();
-	TSharedRef<SWidget> OwnerWidget = AsShared();
-
-	const FVector2D ScreenPos = FSlateApplication::Get().GetCursorPos();
-
-	FSlateApplication::Get().PushMenu(
-		OwnerWidget,
-		FWidgetPath(),
-		MenuContent,
-		ScreenPos,
-		FPopupTransitionEffect(FPopupTransitionEffect::SubMenu)
-	);
-}
-
-bool SMaterialEditor3DPreviewViewport::DoesPreviewMeshContainSlots() const
-{
-	return PreviewMeshComponent != nullptr && PreviewMeshComponent->GetNumMaterials() > 0;
-}
-
 void SMaterialEditor3DPreviewViewport::TogglePreviewBackground()
 {
 	UAssetViewerSettings* Settings = UAssetViewerSettings::Get();
@@ -1086,115 +1004,6 @@ TSharedPtr<IPreviewProfileController> SMaterialEditor3DPreviewViewport::CreatePr
 	return MakeShared<FPreviewProfileController>();
 }
 
-void SMaterialEditor3DPreviewViewport::ApplyStateToAllSlots(bool bState)
-{
-	for (const MaterialSlotPreviewInfo& SlotInfo : MaterialSlots)
-	{
-		SlotInfo.bSlotEnabled = bState;
-	}
-
-	SetPreviewMaterial(PreviewMaterial);
-}
-
-void SMaterialEditor3DPreviewViewport::OnSelectAllSlots()
-{
-	ApplyStateToAllSlots(true);
-}
-
-void SMaterialEditor3DPreviewViewport::OnSelectNoSlots()
-{
-	ApplyStateToAllSlots(false);
-}
-
-void SMaterialEditor3DPreviewViewport::OnMaterialSlotCheckStateChanged(ECheckBoxState NewState, int32 SlotIndex)
-{
-	if (MaterialSlots.IsValidIndex(SlotIndex))
-	{
-		MaterialSlots[SlotIndex].bSlotEnabled = NewState == ECheckBoxState::Checked;
-		SetPreviewMaterial(PreviewMaterial);
-	}
-}
-
-ECheckBoxState SMaterialEditor3DPreviewViewport::IsMaterialSlotChecked(int32 SlotIndex) const
-{
-	return MaterialSlots.IsValidIndex(SlotIndex) && MaterialSlots[SlotIndex].bSlotEnabled ? ECheckBoxState::Checked : ECheckBoxState::Unchecked;
-}
-
-TSharedRef<SWidget> SMaterialEditor3DPreviewViewport::GenerateMaterialSlotMenu()
-{
-	FMenuBuilder MenuBuilder(/*bInShouldCloseWindowAfterMenuSelection*/ false, nullptr);
-
-	MenuBuilder.AddMenuEntry(
-		LOCTEXT("SelectAllSlots", "All Slots"),
-		LOCTEXT("SelectAllSlots_Tooltip", "Apply the preview material to all slots."),
-		FSlateIcon(),
-		FUIAction(FExecuteAction::CreateSP(this, &SMaterialEditor3DPreviewViewport::OnSelectAllSlots))
-	);
-
-	MenuBuilder.AddMenuEntry(
-		LOCTEXT("SelectNoSlots", "None"),
-		LOCTEXT("SelectNoSlots_Tooltip", "Apply the preview material to no slots."),
-		FSlateIcon(),
-		FUIAction(FExecuteAction::CreateSP(this, &SMaterialEditor3DPreviewViewport::OnSelectNoSlots))
-	);
-
-	MenuBuilder.AddSeparator();
-
-	TSharedPtr<FAssetThumbnailPool> ThumbnailPool = UThumbnailManager::Get().GetSharedThumbnailPool();	
-	TArray<TSharedPtr<FAssetThumbnail>> SlotThumbnails;
-	const int32 ThumbSize = 40;
-	for (int32 SlotIndex = 0; SlotIndex < MaterialSlots.Num(); SlotIndex++)
-	{
-		SlotThumbnails.Add(
-			MakeShareable(new FAssetThumbnail(MaterialSlots[SlotIndex].OriginalMaterial, ThumbSize, ThumbSize, ThumbnailPool))
-		);
-
-		const FText SlotLabel = FText::Format(
-			LOCTEXT("SlotLabelFmt", "{0}: {1}"),
-			FText::AsNumber(SlotIndex),
-			FText::FromName(MaterialSlots[SlotIndex].SlotName)
-		);
-
-		TSharedRef<SHorizontalBox> Row =
-			SNew(SHorizontalBox)
-			+ SHorizontalBox::Slot()
-			.AutoWidth()
-			.VAlign(VAlign_Center)
-			.Padding(2.f, 0.f)
-			[
-				SNew(SBox)
-					.WidthOverride(40.f)
-					.HeightOverride(40.f)
-					[
-						SlotThumbnails.IsValidIndex(SlotIndex)
-							? SlotThumbnails[SlotIndex]->MakeThumbnailWidget()
-							: SNullWidget::NullWidget
-					]
-			]
-		+ SHorizontalBox::Slot()
-			.FillWidth(1.f)
-			.VAlign(VAlign_Center)
-			.Padding(4.f, 0.f)
-			[
-				SNew(SCheckBox)
-					.IsChecked(this, &SMaterialEditor3DPreviewViewport::IsMaterialSlotChecked, SlotIndex)
-					.OnCheckStateChanged(this,
-						&SMaterialEditor3DPreviewViewport::OnMaterialSlotCheckStateChanged, SlotIndex)
-					[
-						SNew(STextBlock)
-							.Text(SlotLabel)
-					]
-			];
-
-		MenuBuilder.AddWidget(
-			Row,
-			FText::GetEmpty()
-		);
-	}
-
-	return MenuBuilder.MakeWidget();
-}
-
 void SMaterialEditor3DPreviewViewport::PopulateViewportOverlays(TSharedRef<class SOverlay> Overlay)
 {
 	// add the feature level display widget
diff --git a/Engine/Source/Editor/MaterialEditor/Private/SMaterialEditorViewport.h b/Engine/Source/Editor/MaterialEditor/Private/SMaterialEditorViewport.h
--- a/Engine/Source/Editor/MaterialEditor/Private/SMaterialEditorViewport.h
+++ b/Engine/Source/Editor/MaterialEditor/Private/SMaterialEditorViewport.h
@@ -22,14 +22,6 @@ class SMaterialEditorUIPreviewZoomer;
 class UMaterialInterface;
 class UMeshComponent;
 
-/** Helped struct for containing data for the material slot submenu UI and where to apply the preview material */
-struct MaterialSlotPreviewInfo
-{
-	FName SlotName;
-	TObjectPtr<UMaterialInterface> OriginalMaterial;
-	mutable bool bSlotEnabled;
-};
-
 /**
  * Material Editor Preview viewport widget
  */
@@ -91,9 +83,6 @@ public:
 	bool IsPreviewPrimitiveChecked(EThumbnailPrimType PrimType) const;
 	void OnSetPreviewMeshFromSelection();
 	bool IsPreviewMeshFromSelectionChecked() const;
-	void OnSetPreviewMeshMaterialSlots();
-	bool DoesPreviewMeshContainSlots() const;
-
 	void TogglePreviewBackground();
 	bool IsTogglePreviewBackgroundChecked() const;
 
@@ -108,8 +97,6 @@ protected:
 	virtual TSharedRef<FEditorViewportClient> MakeEditorViewportClient() override;
 	virtual TSharedPtr<SWidget> BuildViewportToolbar() override;
 	virtual TSharedPtr<IPreviewProfileController> CreatePreviewProfileController() override;
-	TSharedRef<SWidget> GenerateMaterialSlotMenu();
-
 	virtual void PopulateViewportOverlays(TSharedRef<class SOverlay> Overlay) override;
 	virtual EVisibility OnGetViewportContentVisibility() const override;
 	virtual void BindCommands() override;
@@ -120,24 +107,6 @@ private:
 
 	void OnPropertyChanged(UObject* ObjectBeingModified, FPropertyChangedEvent& PropertyChangedEvent);
 
-	/* Apply the selected mesh object to the preview mesh component */
-	void ApplySelectedMesh(UObject* InAsset);
-
-	/** Helper function for applying state logic to all material slots + updates*/
-	void ApplyStateToAllSlots(bool bState);
-
-	/** Set true for usage to all material slots */
-	void OnSelectAllSlots();
-
-	/** Set false for usage to all material slots */
-	void OnSelectNoSlots();
-
-	/** Updates the scene preview mesh when a checkbox state changes */
-	void OnMaterialSlotCheckStateChanged(ECheckBoxState NewState, int32 SlotIndex);
-
-	/** Returns if a material slot is checked/should show the preview material */
-	ECheckBoxState IsMaterialSlotChecked(int32 SlotIndex) const;
-
 	/** The parent tab where this viewport resides */
 	TWeakPtr<SDockTab> ParentTab;
 
@@ -169,9 +138,6 @@ private:
 
 	/** Handle to the registered OnPreviewFeatureLevelChanged delegate. */
 	FDelegateHandle PreviewFeatureLevelChangedHandle;
-
-	/** Material Slot data for populating and processing the material slot selections */
-	TArray<const MaterialSlotPreviewInfo> MaterialSlots;
 };
 
 /**
diff --git a/Engine/Source/Editor/MaterialEditor/Public/MaterialEditorActions.h b/Engine/Source/Editor/MaterialEditor/Public/MaterialEditorActions.h
--- a/Engine/Source/Editor/MaterialEditor/Public/MaterialEditorActions.h
+++ b/Engine/Source/Editor/MaterialEditor/Public/MaterialEditorActions.h
@@ -75,9 +75,6 @@ public:
 	
 	/** Sets the preview mesh to the current selection in the level editor */
 	TSharedPtr< FUICommandInfo > SetPreviewMeshFromSelection;
-
-	/** Set the material slots of the PreviewMesh to apply this material to */
-	TSharedPtr< FUICommandInfo > SelectPreviewMeshMaterialSlots;
 	
 	/** Toggles the preview pane's background */
 	TSharedPtr< FUICommandInfo > TogglePreviewBackground;
