Fix bug where opaque materials were not showing base pass pixel shader instruction counts.

- Opaque materials use TBasePassPSFNoLightMapPolicy (w/o skylight) but were requesting the skylight version TBasePassPSFNoLightMapPolicySkylight.  This shader isn't compiled anymore so we wouldn't find the shader and not show the user an instruction count.
- Translucent materials use TBasePassPSFNoLightMapPolicySkylight.
- The string description is now just the shader name instead of a hardcoded string literal so we don't have to keep updating the description as the shader name changes and ensure they match.

#rnx

#Synced-CL 50102962

diff --git a/Engine/Source/Editor/MaterialEditor/Private/MaterialStatsCommon.cpp b/Engine/Source/Editor/MaterialEditor/Private/MaterialStatsCommon.cpp
--- a/Engine/Source/Editor/MaterialEditor/Private/MaterialStatsCommon.cpp
+++ b/Engine/Source/Editor/MaterialEditor/Private/MaterialStatsCommon.cpp
@@ -309,19 +309,25 @@ void FMaterialStatsUtils::GetRepresentativeShaderTypesAndDescriptions(TMap<FName
 	}
 	else if (TargetMaterial->GetFeatureLevel() >= ERHIFeatureLevel::SM5)
 	{
+		static FName TBasePassPSFNoLightMapPolicyName = TEXT("TBasePassPSFNoLightMapPolicy");
+		static FName TBasePassPSFNoLightMapPolicySkylightName = TEXT("TBasePassPSFNoLightMapPolicySkylight");
+
+		const bool bTranslucent = IsTranslucentBlendMode(TargetMaterial);
 		if (TargetMaterial->GetShadingModels().IsUnlit())
 		{
-			//unlit materials are never lightmapped
-			static FName TBasePassPSFNoLightMapPolicyName = TEXT("TBasePassPSFNoLightMapPolicySkylight");
 			ShaderTypeNamesAndDescriptions.FindOrAdd(FLocalVertexFactoryName)
-				.Add(FRepresentativeShaderInfo(ERepresentativeShader::StationarySurface, TBasePassPSFNoLightMapPolicyName, TEXT("Base pass shader without light map")));
+				.Add(FRepresentativeShaderInfo(
+						ERepresentativeShader::StationarySurface, 
+						bTranslucent ? TBasePassPSFNoLightMapPolicySkylightName : TBasePassPSFNoLightMapPolicyName, 
+						bTranslucent ? TBasePassPSFNoLightMapPolicySkylightName.ToString() : TBasePassPSFNoLightMapPolicyName.ToString()));
 		}
 		else
 		{
-			//also show a dynamically lit shader
-			static FName TBasePassPSFNoLightMapPolicyName = TEXT("TBasePassPSFNoLightMapPolicySkylight");
 			ShaderTypeNamesAndDescriptions.FindOrAdd(FLocalVertexFactoryName)
-				.Add(FRepresentativeShaderInfo(ERepresentativeShader::DynamicallyLitObject, TBasePassPSFNoLightMapPolicyName, TEXT("Base pass shader with skylight")));
+				.Add(FRepresentativeShaderInfo(
+						ERepresentativeShader::DynamicallyLitObject, 
+						bTranslucent ? TBasePassPSFNoLightMapPolicySkylightName : TBasePassPSFNoLightMapPolicyName, 
+						bTranslucent ? TBasePassPSFNoLightMapPolicySkylightName.ToString() : TBasePassPSFNoLightMapPolicyName.ToString()));
 
 			if (IsStaticLightingAllowed())
 			{
