Add Material Instance Parameter Groups.
Remove unused code which was used to display parameters in the main panel for selected layers. This is an unused feature as layer selection is disabled.

#13442

[RN] minor | Rendering Architecture

#Synced-CL 50036996

diff --git a/Engine/Source/Editor/MaterialEditor/Private/MaterialEditorInstanceDetailCustomization.cpp b/Engine/Source/Editor/MaterialEditor/Private/MaterialEditorInstanceDetailCustomization.cpp
--- a/Engine/Source/Editor/MaterialEditor/Private/MaterialEditorInstanceDetailCustomization.cpp
+++ b/Engine/Source/Editor/MaterialEditor/Private/MaterialEditorInstanceDetailCustomization.cpp
@@ -51,7 +51,6 @@
 #define LOCTEXT_NAMESPACE "MaterialInstanceEditor"
 
 
-
 TSharedRef<IDetailCustomization> FMaterialInstanceParameterDetails::MakeInstance(UMaterialEditorInstanceConstant* MaterialInstance, SMaterialLayersFunctionsInstanceWrapper* InMaterialLayersFunctionsInstance, FGetShowHiddenParameters InShowHiddenDelegate)
 {
 	return MakeShareable(new FMaterialInstanceParameterDetails(MaterialInstance, InMaterialLayersFunctionsInstance, InShowHiddenDelegate));
@@ -122,57 +121,66 @@ void FMaterialInstanceParameterDetails::CustomizeDetails(IDetailLayoutBuilder& D
 	];
 	
 	TSharedRef<IPropertyHandle> ParameterGroupsProperty = DetailLayout.GetProperty("ParameterGroups");
+	CreateGroupsWidget(ParameterGroupsProperty, GroupsCategory);
+	CreateParameterCollectionOverrideCategory(ParameterGroupsProperty, DetailLayout);
 
-	// check if tree has any selection, we show parameter properties for selected layer only
-	if (MaterialLayersFunctionsInstance != nullptr && MaterialLayersFunctionsInstance->NestedTree->GetNumItemsSelected() > 0)
+	// Create default category for class properties
+	const FName DefaultCategoryName = NAME_None;
+	IDetailCategoryBuilder& DefaultCategory = DetailLayout.EditCategory(DefaultCategoryName);
+	DetailLayout.HideProperty("MaterialLayersParameterValues");
+	if (MaterialEditorInstance->bIsFunctionPreviewMaterial)
 	{
-		// for each selected FSortedParamData item (type stack)
-		TSharedPtr<FSortedParamData> SelectedItem = MaterialLayersFunctionsInstance->NestedTree->GetSelectedItems().Last();
-		
-			// make sure we selected a stack item
-		check (SelectedItem->StackDataType == EStackDataType::Stack)
-
-		// we should now gather all sub-stack items to loop through all together
-		TArray<TSharedPtr<FSortedParamData>> AssetStacksCollection;
-		TArray<uint32> AssetParentNodeCollection;
-		MaterialLayersFunctionsInstance->NestedTree->CollectAssetStackItemsRecursively(SelectedItem, AssetStacksCollection, AssetParentNodeCollection);
-
-		// we go through list of assets
-		for (int32 Index = 0; Index < AssetStacksCollection.Num(); ++Index)
+		// Customize Parent property so we can check for recursively set parents
+		bool bShowParent = false;
+		if(MaterialEditorInstance->SourceFunction->GetMaterialFunctionUsage() != EMaterialFunctionUsage::Default)
+		{
+			bShowParent = true;
+		}
+		if (bShowParent)
 		{
-			TSharedPtr<FSortedParamData> AssetItem = AssetStacksCollection[Index];
-			uint32 NodeID = AssetParentNodeCollection[Index];
+			TSharedRef<IPropertyHandle> ParentPropertyHandle = DetailLayout.GetProperty("Parent");
+			IDetailPropertyRow& ParentPropertyRow = DefaultCategory.AddProperty(ParentPropertyHandle);
+			ParentPropertyHandle->MarkResetToDefaultCustomized();
 
-			for (TSharedPtr<FSortedParamData> GroupParamData : AssetItem->Children)
-			{
-				if (GroupParamData->StackDataType == EStackDataType::Group)
-				{
-					int32 GroupIdx = MaterialEditorInstance->ParameterGroups.IndexOfByPredicate(
-						[&](const FEditorParameterGroup& Group)
-						{
-							return Group.GroupName == GroupParamData->Group.GroupName;
-						});
-					if (GroupIdx != INDEX_NONE)
-					{
-						FEditorParameterGroup& ParameterGroup = GroupParamData->Group;
-						IDetailGroup& DetailGroup = GroupsCategory.AddGroup(ParameterGroup.GroupName, FText::FromName(ParameterGroup.GroupName), false, true);
-						TSharedPtr<IPropertyHandle> GroupPropertyHandle = ParameterGroupsProperty->GetChildHandle(GroupIdx);
-
-						CreateSingleGroupWidget(ParameterGroup, GroupPropertyHandle, DetailGroup, GroupParamData->ParameterInfo.Index, true);
-
-						FSimpleDelegate UpdateThumbnails = FSimpleDelegate::CreateLambda([=, this]()
-							{
-								this->MaterialLayersFunctionsInstance->NestedTree->UpdateThumbnailMaterial(AssetItem->ParameterInfo.Association, NodeID);
-							});
-						GroupPropertyHandle->SetOnPropertyValueChanged(UpdateThumbnails);
-						GroupPropertyHandle->SetOnChildPropertyValueChanged(UpdateThumbnails);
-					}
-				}
-			}
+			TSharedPtr<SWidget> NameWidget;
+			TSharedPtr<SWidget> ValueWidget;
+			FDetailWidgetRow Row;
+
+			ParentPropertyRow.GetDefaultWidgets(NameWidget, ValueWidget, Row);
+
+			ParentPropertyHandle->ClearResetToDefaultCustomized();
+
+			const bool bShowChildren = true;
+			ParentPropertyRow.CustomWidget(bShowChildren)
+				.NameContent()
+				.MinDesiredWidth(Row.NameWidget.MinWidth)
+				.MaxDesiredWidth(Row.NameWidget.MaxWidth)
+				[
+					NameWidget.ToSharedRef()
+				]
+			.ValueContent()
+				.MinDesiredWidth(Row.ValueWidget.MinWidth)
+				.MaxDesiredWidth(Row.ValueWidget.MaxWidth)
+				[
+					SNew(SObjectPropertyEntryBox)
+					.ObjectPath(this, &FMaterialInstanceParameterDetails::GetFunctionParentPath)
+					.AllowedClass(UMaterialFunctionInterface::StaticClass())
+					.ThumbnailPool(DetailLayout.GetThumbnailPool())
+					.AllowClear(true)
+					.OnObjectChanged(this, &FMaterialInstanceParameterDetails::OnAssetChanged, ParentPropertyHandle)
+					.OnShouldSetAsset(this, &FMaterialInstanceParameterDetails::OnShouldSetAsset)
+					.NewAssetFactories(TArray<UFactory*>())
+				];
+
+			ValueWidget.Reset();
+
+
+		}
+		else
+		{
+			DetailLayout.HideProperty("Parent");
 		}
-		
-		DetailLayout.HideCategory("MaterialEditorInstanceConstant");
-		DetailLayout.HideProperty("Parent");
+
 		DetailLayout.HideProperty("PostProcessOverrides");
 		DetailLayout.HideProperty("PhysMaterial");
 		DetailLayout.HideProperty("LightmassSettings");
@@ -186,76 +194,71 @@ void FMaterialInstanceParameterDetails::CustomizeDetails(IDetailLayoutBuilder& D
 		DetailLayout.HideProperty("bOverrideToonProfile");
 		DetailLayout.HideProperty("ToonProfile");
 		DetailLayout.HideProperty("BasePropertyOverrides");
-		DetailLayout.HideProperty("MaterialLayersParameterValues");
 	}
 	else
 	{
-		CreateGroupsWidget(ParameterGroupsProperty, GroupsCategory);
-		CreateParameterCollectionOverrideCategory(ParameterGroupsProperty, DetailLayout);
-
-		// Create default category for class properties
-		const FName DefaultCategoryName = NAME_None;
-		IDetailCategoryBuilder& DefaultCategory = DetailLayout.EditCategory(DefaultCategoryName);
-		DetailLayout.HideProperty("MaterialLayersParameterValues");
-		if (MaterialEditorInstance->bIsFunctionPreviewMaterial)
-		{
-			// Customize Parent property so we can check for recursively set parents
-			bool bShowParent = false;
-			if(MaterialEditorInstance->SourceFunction->GetMaterialFunctionUsage() != EMaterialFunctionUsage::Default)
-			{
-				bShowParent = true;
-			}
-			if (bShowParent)
-			{
-				TSharedRef<IPropertyHandle> ParentPropertyHandle = DetailLayout.GetProperty("Parent");
-				IDetailPropertyRow& ParentPropertyRow = DefaultCategory.AddProperty(ParentPropertyHandle);
-				ParentPropertyHandle->MarkResetToDefaultCustomized();
+		DetailLayout.HideProperty("PostProcessOverrides");
+		CreatePostProcessOverrideWidgets(DetailLayout);
 
-				TSharedPtr<SWidget> NameWidget;
-				TSharedPtr<SWidget> ValueWidget;
-				FDetailWidgetRow Row;
+		// Add PhysMaterial property
+		DefaultCategory.AddProperty("PhysMaterial");
 
-				ParentPropertyRow.GetDefaultWidgets(NameWidget, ValueWidget, Row);
+		// Customize Parent property so we can check for recursively set parents
+		TSharedRef<IPropertyHandle> ParentPropertyHandle = DetailLayout.GetProperty("Parent");
+		IDetailPropertyRow& ParentPropertyRow = DefaultCategory.AddProperty(ParentPropertyHandle);
 
-				ParentPropertyHandle->ClearResetToDefaultCustomized();
+		TSharedRef<SObjectPropertyEntryBox> ParentValueWidget = SNew(SObjectPropertyEntryBox)
+			.PropertyHandle(ParentPropertyHandle)
+			.AllowedClass(UMaterialInterface::StaticClass())
+			.ThumbnailPool(DetailLayout.GetThumbnailPool())
+			.AllowClear(true)
+			.OnShouldSetAsset(this, &FMaterialInstanceParameterDetails::OnShouldSetAsset);
 
-				const bool bShowChildren = true;
-				ParentPropertyRow.CustomWidget(bShowChildren)
-					.NameContent()
-					.MinDesiredWidth(Row.NameWidget.MinWidth)
-					.MaxDesiredWidth(Row.NameWidget.MaxWidth)
-					[
-						NameWidget.ToSharedRef()
-					]
-				.ValueContent()
-					.MinDesiredWidth(Row.ValueWidget.MinWidth)
-					.MaxDesiredWidth(Row.ValueWidget.MaxWidth)
-					[
-						SNew(SObjectPropertyEntryBox)
-						.ObjectPath(this, &FMaterialInstanceParameterDetails::GetFunctionParentPath)
-						.AllowedClass(UMaterialFunctionInterface::StaticClass())
-						.ThumbnailPool(DetailLayout.GetThumbnailPool())
-						.AllowClear(true)
-						.OnObjectChanged(this, &FMaterialInstanceParameterDetails::OnAssetChanged, ParentPropertyHandle)
-						.OnShouldSetAsset(this, &FMaterialInstanceParameterDetails::OnShouldSetAsset)
-						.NewAssetFactories(TArray<UFactory*>())
-					];
+		float ParentValueMinDesiredWidth, ParentValueMaxDesiredWidth;
+		ParentValueWidget->GetDesiredWidth(ParentValueMinDesiredWidth, ParentValueMaxDesiredWidth);
 
-				ValueWidget.Reset();
+		const bool bShowChildren = true;
+		ParentPropertyRow.CustomWidget(bShowChildren)
+			.NameContent()
+			[
+				ParentPropertyHandle->CreatePropertyNameWidget()
+			]
+			.ValueContent()
+			.MinDesiredWidth(ParentValueMinDesiredWidth)
+			.MaxDesiredWidth(ParentValueMaxDesiredWidth)
+			[
+				ParentValueWidget
+			];
 
+		// Add/hide other properties
+		DetailLayout.HideProperty("LightmassSettings");
+		CreateLightmassOverrideWidgets(DetailLayout);
+		DetailLayout.HideProperty("bUseOldStyleMICEditorGroups");
+		DetailLayout.HideProperty("ParameterGroups");
 
-			}
-			else
-			{
-				DetailLayout.HideProperty("Parent");
-			}
+		{
+			FIsResetToDefaultVisible IsRefractionDepthBiasPropertyResetVisible = FIsResetToDefaultVisible::CreateLambda([this](TSharedPtr<IPropertyHandle> InHandle) {
+				float BiasValue;
+				float ParentBiasValue;
+				return MaterialEditorInstance->SourceInstance->GetRefractionSettings(BiasValue) 
+					&& MaterialEditorInstance->Parent->GetRefractionSettings(ParentBiasValue)
+					&& BiasValue != ParentBiasValue;
+			});
+			FResetToDefaultHandler ResetRefractionDepthBiasPropertyHandler = FResetToDefaultHandler::CreateLambda([this](TSharedPtr<IPropertyHandle> InHandle) {
+				MaterialEditorInstance->Parent->GetRefractionSettings(MaterialEditorInstance->RefractionDepthBias);
+			});
+			FResetToDefaultOverride ResetRefractionDepthBiasPropertyOverride = FResetToDefaultOverride::Create(IsRefractionDepthBiasPropertyResetVisible, ResetRefractionDepthBiasPropertyHandler);
+			IDetailPropertyRow& PropertyRow = DefaultCategory.AddProperty("RefractionDepthBias");
+			PropertyRow.Visibility(TAttribute<EVisibility>::Create(TAttribute<EVisibility>::FGetter::CreateSP(this, &FMaterialInstanceParameterDetails::ShouldShowMaterialRefractionSettings)));
+			PropertyRow.OverrideResetToDefault(ResetRefractionDepthBiasPropertyOverride);
+		}
 
-			DetailLayout.HideProperty("PostProcessOverrides");
-			DetailLayout.HideProperty("PhysMaterial");
-			DetailLayout.HideProperty("LightmassSettings");
-			DetailLayout.HideProperty("bUseOldStyleMICEditorGroups");
-			DetailLayout.HideProperty("ParameterGroups");
-			DetailLayout.HideProperty("RefractionDepthBias");
+		{
+			// Add the material property override group
+			static FName GroupName(TEXT("MaterialPropertyOverrideGroup"));
+			IDetailGroup& MaterialPropertyOverrideGroup = DefaultCategory.AddGroup(GroupName, LOCTEXT("MaterialPropertyOverrideGroup", "Material Property Overrides"), false, false);
+			
+			// Hide the originals, these will be recreated manually
 			DetailLayout.HideProperty("bOverrideSubsurfaceProfile");
 			DetailLayout.HideProperty("SubsurfaceProfile");
 			DetailLayout.HideProperty("bOverrideSpecularProfile");
@@ -263,221 +266,191 @@ void FMaterialInstanceParameterDetails::CustomizeDetails(IDetailLayoutBuilder& D
 			DetailLayout.HideProperty("bOverrideToonProfile");
 			DetailLayout.HideProperty("ToonProfile");
 			DetailLayout.HideProperty("BasePropertyOverrides");
-		}
-		else
-		{
-			DetailLayout.HideProperty("PostProcessOverrides");
-			CreatePostProcessOverrideWidgets(DetailLayout);
-
-			// Add PhysMaterial property
-			DefaultCategory.AddProperty("PhysMaterial");
-
-			// Customize Parent property so we can check for recursively set parents
-			TSharedRef<IPropertyHandle> ParentPropertyHandle = DetailLayout.GetProperty("Parent");
-			IDetailPropertyRow& ParentPropertyRow = DefaultCategory.AddProperty(ParentPropertyHandle);
-
-			TSharedRef<SObjectPropertyEntryBox> ParentValueWidget = SNew(SObjectPropertyEntryBox)
-				.PropertyHandle(ParentPropertyHandle)
-				.AllowedClass(UMaterialInterface::StaticClass())
-				.ThumbnailPool(DetailLayout.GetThumbnailPool())
-				.AllowClear(true)
-				.OnShouldSetAsset(this, &FMaterialInstanceParameterDetails::OnShouldSetAsset);
-
-			float ParentValueMinDesiredWidth, ParentValueMaxDesiredWidth;
-			ParentValueWidget->GetDesiredWidth(ParentValueMinDesiredWidth, ParentValueMaxDesiredWidth);
-
-			const bool bShowChildren = true;
-			ParentPropertyRow.CustomWidget(bShowChildren)
-				.NameContent()
-				[
-					ParentPropertyHandle->CreatePropertyNameWidget()
-				]
-				.ValueContent()
-				.MinDesiredWidth(ParentValueMinDesiredWidth)
-				.MaxDesiredWidth(ParentValueMaxDesiredWidth)
-				[
-					ParentValueWidget
-				];
-
-			// Add/hide other properties
-			DetailLayout.HideProperty("LightmassSettings");
-			CreateLightmassOverrideWidgets(DetailLayout);
-			DetailLayout.HideProperty("bUseOldStyleMICEditorGroups");
-			DetailLayout.HideProperty("ParameterGroups");
 
+			// Set up the override logic for the subsurface profile
 			{
-				FIsResetToDefaultVisible IsRefractionDepthBiasPropertyResetVisible = FIsResetToDefaultVisible::CreateLambda([this](TSharedPtr<IPropertyHandle> InHandle) {
-					float BiasValue;
-					float ParentBiasValue;
-					return MaterialEditorInstance->SourceInstance->GetRefractionSettings(BiasValue) 
-						&& MaterialEditorInstance->Parent->GetRefractionSettings(ParentBiasValue)
-						&& BiasValue != ParentBiasValue;
-				});
-				FResetToDefaultHandler ResetRefractionDepthBiasPropertyHandler = FResetToDefaultHandler::CreateLambda([this](TSharedPtr<IPropertyHandle> InHandle) {
-					MaterialEditorInstance->Parent->GetRefractionSettings(MaterialEditorInstance->RefractionDepthBias);
-				});
-				FResetToDefaultOverride ResetRefractionDepthBiasPropertyOverride = FResetToDefaultOverride::Create(IsRefractionDepthBiasPropertyResetVisible, ResetRefractionDepthBiasPropertyHandler);
-				IDetailPropertyRow& PropertyRow = DefaultCategory.AddProperty("RefractionDepthBias");
-				PropertyRow.Visibility(TAttribute<EVisibility>::Create(TAttribute<EVisibility>::FGetter::CreateSP(this, &FMaterialInstanceParameterDetails::ShouldShowMaterialRefractionSettings)));
-				PropertyRow.OverrideResetToDefault(ResetRefractionDepthBiasPropertyOverride);
+				TAttribute<bool> IsParamEnabled = TAttribute<bool>::Create(TAttribute<bool>::FGetter::CreateLambda([this](){ return (bool)MaterialEditorInstance->bOverrideSubsurfaceProfile; }));
+
+				IDetailPropertyRow& PropertyRow = MaterialPropertyOverrideGroup.AddPropertyRow(DetailLayout.GetProperty("SubsurfaceProfile"));
+				PropertyRow
+					.EditCondition(IsParamEnabled, 
+						FOnBooleanValueChanged::CreateLambda([this](bool NewValue) {
+							MaterialEditorInstance->bOverrideSubsurfaceProfile = (uint32)NewValue;
+							MaterialEditorInstance->PostEditChange();
+							FEditorSupportDelegates::RedrawAllViewports.Broadcast();
+					}))
+					.Visibility(TAttribute<EVisibility>::Create(TAttribute<EVisibility>::FGetter::CreateSP(this, &FMaterialInstanceParameterDetails::ShouldShowSubsurfaceProfile)));
 			}
 
+			// Set up the override logic for the specular and toon profiles
+			if (Substrate::IsSubstrateEnabled())
 			{
-				// Add the material property override group
-				static FName GroupName(TEXT("MaterialPropertyOverrideGroup"));
-				IDetailGroup& MaterialPropertyOverrideGroup = DefaultCategory.AddGroup(GroupName, LOCTEXT("MaterialPropertyOverrideGroup", "Material Property Overrides"), false, false);
-			
-				// Hide the originals, these will be recreated manually
-				DetailLayout.HideProperty("bOverrideSubsurfaceProfile");
-				DetailLayout.HideProperty("SubsurfaceProfile");
-				DetailLayout.HideProperty("bOverrideSpecularProfile");
-				DetailLayout.HideProperty("SpecularProfile");
-				DetailLayout.HideProperty("bOverrideToonProfile");
-				DetailLayout.HideProperty("ToonProfile");
-				DetailLayout.HideProperty("BasePropertyOverrides");
-
-				// Set up the override logic for the subsurface profile
-				{
-					TAttribute<bool> IsParamEnabled = TAttribute<bool>::Create(TAttribute<bool>::FGetter::CreateLambda([this](){ return (bool)MaterialEditorInstance->bOverrideSubsurfaceProfile; }));
-
-					IDetailPropertyRow& PropertyRow = MaterialPropertyOverrideGroup.AddPropertyRow(DetailLayout.GetProperty("SubsurfaceProfile"));
-					PropertyRow
-						.EditCondition(IsParamEnabled, 
-							FOnBooleanValueChanged::CreateLambda([this](bool NewValue) {
-								MaterialEditorInstance->bOverrideSubsurfaceProfile = (uint32)NewValue;
-								MaterialEditorInstance->PostEditChange();
-								FEditorSupportDelegates::RedrawAllViewports.Broadcast();
-						}))
-						.Visibility(TAttribute<EVisibility>::Create(TAttribute<EVisibility>::FGetter::CreateSP(this, &FMaterialInstanceParameterDetails::ShouldShowSubsurfaceProfile)));
-				}
-
-				// Set up the override logic for the specular and toon profiles
-				if (Substrate::IsSubstrateEnabled())
-				{
-					TAttribute<bool> IsSpecProfileParamEnabled = TAttribute<bool>::Create(TAttribute<bool>::FGetter::CreateLambda([this](){ return (bool)MaterialEditorInstance->bOverrideSpecularProfile; }));
-
-					IDetailPropertyRow& SpecProfilePropertyRow = MaterialPropertyOverrideGroup.AddPropertyRow(DetailLayout.GetProperty("SpecularProfile"));
-					SpecProfilePropertyRow
-					.EditCondition(IsSpecProfileParamEnabled,
-								   FOnBooleanValueChanged::CreateLambda([this](bool NewValue) {
-									   MaterialEditorInstance->bOverrideSpecularProfile = (uint32)NewValue;
-									   MaterialEditorInstance->PostEditChange();
-									   FEditorSupportDelegates::RedrawAllViewports.Broadcast();
-					}))
-					.Visibility(TAttribute<EVisibility>::Create(TAttribute<EVisibility>::FGetter::CreateSP(this, &FMaterialInstanceParameterDetails::ShouldShowSpecularProfile)));
-
-
-					TAttribute<bool> IsToonProfileParamEnabled = TAttribute<bool>::Create(TAttribute<bool>::FGetter::CreateLambda([this](){ return (bool)MaterialEditorInstance->bOverrideToonProfile; }));
-
-					IDetailPropertyRow& ToonProfilePropertyRow = MaterialPropertyOverrideGroup.AddPropertyRow(DetailLayout.GetProperty("ToonProfile"));
-					ToonProfilePropertyRow
-					.EditCondition(IsToonProfileParamEnabled,
-								   FOnBooleanValueChanged::CreateLambda([this](bool NewValue) {
-									   MaterialEditorInstance->bOverrideToonProfile = (uint32)NewValue;
-									   MaterialEditorInstance->PostEditChange();
-									   FEditorSupportDelegates::RedrawAllViewports.Broadcast();
-					}))
-					.Visibility(TAttribute<EVisibility>::Create(TAttribute<EVisibility>::FGetter::CreateSP(this, &FMaterialInstanceParameterDetails::ShouldShowToonProfile)));
-				}
+				TAttribute<bool> IsSpecProfileParamEnabled = TAttribute<bool>::Create(TAttribute<bool>::FGetter::CreateLambda([this](){ return (bool)MaterialEditorInstance->bOverrideSpecularProfile; }));
+
+				IDetailPropertyRow& SpecProfilePropertyRow = MaterialPropertyOverrideGroup.AddPropertyRow(DetailLayout.GetProperty("SpecularProfile"));
+				SpecProfilePropertyRow
+				.EditCondition(IsSpecProfileParamEnabled,
+								FOnBooleanValueChanged::CreateLambda([this](bool NewValue) {
+									MaterialEditorInstance->bOverrideSpecularProfile = (uint32)NewValue;
+									MaterialEditorInstance->PostEditChange();
+									FEditorSupportDelegates::RedrawAllViewports.Broadcast();
+				}))
+				.Visibility(TAttribute<EVisibility>::Create(TAttribute<EVisibility>::FGetter::CreateSP(this, &FMaterialInstanceParameterDetails::ShouldShowSpecularProfile)));
+
+
+				TAttribute<bool> IsToonProfileParamEnabled = TAttribute<bool>::Create(TAttribute<bool>::FGetter::CreateLambda([this](){ return (bool)MaterialEditorInstance->bOverrideToonProfile; }));
+
+				IDetailPropertyRow& ToonProfilePropertyRow = MaterialPropertyOverrideGroup.AddPropertyRow(DetailLayout.GetProperty("ToonProfile"));
+				ToonProfilePropertyRow
+				.EditCondition(IsToonProfileParamEnabled,
+								FOnBooleanValueChanged::CreateLambda([this](bool NewValue) {
+									MaterialEditorInstance->bOverrideToonProfile = (uint32)NewValue;
+									MaterialEditorInstance->PostEditChange();
+									FEditorSupportDelegates::RedrawAllViewports.Broadcast();
+				}))
+				.Visibility(TAttribute<EVisibility>::Create(TAttribute<EVisibility>::FGetter::CreateSP(this, &FMaterialInstanceParameterDetails::ShouldShowToonProfile)));
+			}
 			
-				// Append the base property overrides to the Material Property Override Group
-				CreateBasePropertyOverrideWidgets(DetailLayout, MaterialPropertyOverrideGroup);
+			// Append the base property overrides to the Material Property Override Group
+			CreateBasePropertyOverrideWidgets(DetailLayout, MaterialPropertyOverrideGroup);
 
-				// Append the nanite material override.
-				MaterialPropertyOverrideGroup.AddPropertyRow(DetailLayout.GetProperty("NaniteOverrideMaterial"));
+			// Append the nanite material override.
+			MaterialPropertyOverrideGroup.AddPropertyRow(DetailLayout.GetProperty("NaniteOverrideMaterial"));
 
-				// Add the usage flag overrides
-				CreateUsageFlagOverrideWidgets(DetailLayout, DefaultCategory);
-			}
+			// Add the usage flag overrides
+			CreateUsageFlagOverrideWidgets(DetailLayout, DefaultCategory);
+		}
 
+		{
+			// Refresh UI when enumeration objects change.
+			IDetailPropertyRow& PropertyRow = DefaultCategory.AddProperty("EnumerationObjects");
+			TSharedPtr<IPropertyHandle> PropertyHandle = PropertyRow.GetPropertyHandle();
+			if (PropertyHandle->IsValidHandle())
 			{
-				// Refresh UI when enumeration objects change.
-				IDetailPropertyRow& PropertyRow = DefaultCategory.AddProperty("EnumerationObjects");
-				TSharedPtr<IPropertyHandle> PropertyHandle = PropertyRow.GetPropertyHandle();
-				if (PropertyHandle->IsValidHandle())
-				{
-					// Needs to happen on property (add to array) and child (array entries).
-					auto ForceRefresh = FSimpleDelegate::CreateLambda([&DetailLayout]() { DetailLayout.ForceRefreshDetails(); });
-					PropertyHandle->SetOnPropertyValueChanged(ForceRefresh);
-					PropertyHandle->SetOnChildPropertyValueChanged(ForceRefresh);
-				}
+				// Needs to happen on property (add to array) and child (array entries).
+				auto ForceRefresh = FSimpleDelegate::CreateLambda([&DetailLayout]() { DetailLayout.ForceRefreshDetails(); });
+				PropertyHandle->SetOnPropertyValueChanged(ForceRefresh);
+				PropertyHandle->SetOnChildPropertyValueChanged(ForceRefresh);
 			}
 		}
+	}
 
-		// Add the preview mesh property directly from the material instance 
-		FName PreviewingCategoryName = TEXT("Previewing");
-		IDetailCategoryBuilder& PreviewingCategory = DetailLayout.EditCategory(PreviewingCategoryName, LOCTEXT("MICPreviewingCategoryTitle", "Previewing"));
+	// Add the preview mesh property directly from the material instance 
+	FName PreviewingCategoryName = TEXT("Previewing");
+	IDetailCategoryBuilder& PreviewingCategory = DetailLayout.EditCategory(PreviewingCategoryName, LOCTEXT("MICPreviewingCategoryTitle", "Previewing"));
 
-		TArray<UObject*> ExternalObjects;
-		ExternalObjects.Add(MaterialEditorInstance->SourceInstance);
+	TArray<UObject*> ExternalObjects;
+	ExternalObjects.Add(MaterialEditorInstance->SourceInstance);
 
-		PreviewingCategory.AddExternalObjectProperty(ExternalObjects, TEXT("PreviewMesh"));
+	PreviewingCategory.AddExternalObjectProperty(ExternalObjects, TEXT("PreviewMesh"));
 
-		DefaultCategory.AddExternalObjectProperty(ExternalObjects, TEXT("AssetUserData"), EPropertyLocation::Advanced);
-	}
+	DefaultCategory.AddExternalObjectProperty(ExternalObjects, TEXT("AssetUserData"), EPropertyLocation::Advanced);
 }
 
 void FMaterialInstanceParameterDetails::CreateGroupsWidget(TSharedRef<IPropertyHandle> ParameterGroupsProperty, IDetailCategoryBuilder& GroupsCategory)
 {
 	bool bShowSaveButtons = false;
 	check(MaterialEditorInstance);
+
+	struct FGroupData
+	{
+		IDetailGroup* Group = nullptr;
+		TArray<int32> AffectedGroupsIndices;
+	};
+	TMap<FName, FGroupData> GroupNameToData;
 	
 	for (int32 GroupIdx = 0; GroupIdx < MaterialEditorInstance->ParameterGroups.Num(); ++GroupIdx)
 	{
 		FEditorParameterGroup& ParameterGroup = MaterialEditorInstance->ParameterGroups[GroupIdx];
-		if (ParameterGroup.GroupAssociation == EMaterialParameterAssociation::GlobalParameter
-			&& ParameterGroup.GroupName != FMaterialPropertyHelpers::LayerParamName)
+		if (ParameterGroup.GroupAssociation != EMaterialParameterAssociation::GlobalParameter ||
+			ParameterGroup.GroupName == FMaterialPropertyHelpers::LayerParamName)
 		{
-			bShowSaveButtons = true;
-			bool bCreateGroup = false;
-			for (int32 ParamIdx = 0; ParamIdx < ParameterGroup.Parameters.Num() && !bCreateGroup; ++ParamIdx)
+			continue;
+		}
+
+		bShowSaveButtons = true;
+		bool bCreateGroup = false;
+		for (int32 ParamIdx = 0; ParamIdx < ParameterGroup.Parameters.Num() && !bCreateGroup; ++ParamIdx)
+		{
+			UDEditorParameterValue* Parameter = ParameterGroup.Parameters[ParamIdx];
+			const bool bIsVisible = MaterialEditorInstance->VisibleExpressions.Contains(Parameter->ParameterInfo) && FMaterialPropertyHelpers::ShouldCreatePropertyRowForParameter(Parameter);
+			bCreateGroup = bIsVisible && (!MaterialEditorInstance->bShowOnlyOverrides || FMaterialPropertyHelpers::IsOverriddenExpression(Parameter));
+		}
+
+		if (!bCreateGroup)
+		{
+			continue;
+		}
+
+		IDetailGroup* DetailGroup = nullptr;
+		TArray<FString> GroupParts;
+		ParameterGroup.GroupName.ToString().ParseIntoArray(GroupParts, TEXT("|"));
+
+		FString Path;
+		for (FString Part : GroupParts)
+		{
+			Part.TrimStartAndEndInline();
+
+			if (!Path.IsEmpty())
 			{
-				UDEditorParameterValue* Parameter = ParameterGroup.Parameters[ParamIdx];
-				const bool bIsVisible = MaterialEditorInstance->VisibleExpressions.Contains(Parameter->ParameterInfo) && FMaterialPropertyHelpers::ShouldCreatePropertyRowForParameter(Parameter);
-				bCreateGroup = bIsVisible && (!MaterialEditorInstance->bShowOnlyOverrides || FMaterialPropertyHelpers::IsOverriddenExpression(Parameter));
+				Path += "|";
 			}
-		
-			if (bCreateGroup)
-			{
-				IDetailGroup& DetailGroup = GroupsCategory.AddGroup(ParameterGroup.GroupName, FText::FromName(ParameterGroup.GroupName), false, false);
-				FUIAction CopyAction(
-					FExecuteAction::CreateSP(this, &FMaterialInstanceParameterDetails::OnCopyParameterValues, GroupIdx),
-					FCanExecuteAction::CreateSP(this, &FMaterialInstanceParameterDetails::CanCopyParameterValues, GroupIdx));
-				FUIAction PasteAction(
-					FExecuteAction::CreateSP(this, &FMaterialInstanceParameterDetails::OnPasteParameterValues, GroupIdx),
-					FCanExecuteAction::CreateSP(this, &FMaterialInstanceParameterDetails::CanPasteParameterValues, GroupIdx));
-				FDetailWidgetRow& HeaderRow = DetailGroup.HeaderRow()
-					.CopyAction(CopyAction)
-					.PasteAction(PasteAction)
-					.NameContent()
-					[
-						SNew(STextBlock)
-						.Text(FText::FromName(DetailGroup.GetGroupName()))
-					];
-
-				CreateSingleGroupWidget(ParameterGroup, ParameterGroupsProperty->GetChildHandle(GroupIdx), DetailGroup);
+			Path += Part;
 
-				HeaderRow.AddCustomContextMenuAction(FUIAction(
-						FExecuteAction::CreateLambda([&]()mutable 
-						{
-							EnableGroupParameters(ParameterGroup, true);
-						})),
-							LOCTEXT("ToggleParametersEnable", "Enable All Parameters"),
-							LOCTEXT("ToggleParametersEnableTooltip", "Enable All Parameters in group"),
-							FSlateIcon());
-
-				HeaderRow.AddCustomContextMenuAction(FUIAction(
-						FExecuteAction::CreateLambda([&]()mutable 
-						{
-							EnableGroupParameters(ParameterGroup, false);
-						})),
-							LOCTEXT("ToggleParametersDisable", "Disable All Parameters"),
-							LOCTEXT("ToggleParametersDisableTooltip", "Disable All Parameters in group"),
-							FSlateIcon());
+			FGroupData& GroupData = GroupNameToData.FindOrAdd(FName(Path));
+			if (!GroupData.Group)
+			{
+				if (!DetailGroup)
+				{
+					GroupData.Group = &GroupsCategory.AddGroup(FName(Part), FText::FromString(FName::NameToDisplayString(Part, false)), false, true);
+				}
+				else
+				{
+					GroupData.Group = &DetailGroup->AddGroup(FName(Part), FText::FromString(FName::NameToDisplayString(Part, false)), true);
+				}
 			}
+
+			DetailGroup = GroupData.Group;
+			GroupData.AffectedGroupsIndices.Add(GroupIdx);
 		}
+
+		check(DetailGroup);
+		CreateSingleGroupWidget(ParameterGroup, ParameterGroupsProperty->GetChildHandle(GroupIdx), *DetailGroup);
 	}
-	
+
+	for (const TPair<FName, FGroupData>& It : GroupNameToData)
+	{
+		FUIAction CopyAction(
+			FExecuteAction::CreateSP(this, &FMaterialInstanceParameterDetails::OnCopyParameterValues, It.Value.AffectedGroupsIndices),
+			FCanExecuteAction::CreateSP(this, &FMaterialInstanceParameterDetails::CanCopyParameterValues, It.Value.AffectedGroupsIndices));
+		FUIAction PasteAction(
+			FExecuteAction::CreateSP(this, &FMaterialInstanceParameterDetails::OnPasteParameterValues, It.Value.AffectedGroupsIndices),
+			FCanExecuteAction::CreateSP(this, &FMaterialInstanceParameterDetails::CanPasteParameterValues, It.Value.AffectedGroupsIndices));
+		FDetailWidgetRow& HeaderRow = It.Value.Group->HeaderRow()
+			.CopyAction(CopyAction)
+			.PasteAction(PasteAction)
+			.NameContent()
+			[
+				SNew(STextBlock)
+					.Text(FText::FromName(It.Value.Group->GetGroupName()))
+			];
+
+		HeaderRow.AddCustomContextMenuAction(FUIAction(
+			FExecuteAction::CreateSP(this, &FMaterialInstanceParameterDetails::EnableGroupParameters, It.Value.AffectedGroupsIndices, true)),
+			LOCTEXT("ToggleParametersEnable", "Enable All Parameters"),
+			LOCTEXT("ToggleParametersEnableTooltip", "Enable All Parameters in group"),
+			FSlateIcon());
+
+		HeaderRow.AddCustomContextMenuAction(FUIAction(
+			FExecuteAction::CreateSP(this, &FMaterialInstanceParameterDetails::EnableGroupParameters, It.Value.AffectedGroupsIndices, false)),
+			LOCTEXT("ToggleParametersDisable", "Disable All Parameters"),
+			LOCTEXT("ToggleParametersDisableTooltip", "Disable All Parameters in group"),
+			FSlateIcon());
+	}
+
+
+
+
 	if (bShowSaveButtons)
 	{
 		FDetailWidgetRow& SaveInstanceRow = GroupsCategory.AddCustomRow(LOCTEXT("SaveInstances", "Save Instances"));
@@ -617,13 +590,27 @@ void FMaterialInstanceParameterDetails::CreateParameterCollectionOverrideCategor
 	}
 }
 
-void FMaterialInstanceParameterDetails::EnableGroupParameters(FEditorParameterGroup& ParameterGroup, bool ShouldEnable)
+void FMaterialInstanceParameterDetails::EnableGroupParameters(const TArray<int32> ParameterGroupIndices, const bool bShouldEnable)
 {
+	if (!MaterialEditorInstance)
+	{
+		return;
+	}
+
 	// loop through each parameter in the group and toggle to enable/disable them all
-	for (int32 ParamIdx = 0; ParamIdx < ParameterGroup.Parameters.Num(); ++ParamIdx)
+	for (const int32 ParameterGroupIndex : ParameterGroupIndices)
 	{
-		UDEditorParameterValue* Parameter = ParameterGroup.Parameters[ParamIdx];
-		Parameter->bOverride = ShouldEnable;
+		if (!MaterialEditorInstance->ParameterGroups.IsValidIndex(ParameterGroupIndex))
+		{
+			continue;
+		}
+
+		FEditorParameterGroup& ParameterGroup = MaterialEditorInstance->ParameterGroups[ParameterGroupIndex];
+		for (int32 ParamIdx = 0; ParamIdx < ParameterGroup.Parameters.Num(); ++ParamIdx)
+		{
+			UDEditorParameterValue* Parameter = ParameterGroup.Parameters[ParamIdx];
+			Parameter->bOverride = bShouldEnable;
+		}
 	}
 }
 
@@ -745,45 +732,53 @@ EVisibility FMaterialInstanceParameterDetails::ShouldShowToonProfile() const
 	return bHasProfiles ? EVisibility::Visible : EVisibility::Collapsed;
 }
 
-void FMaterialInstanceParameterDetails::OnCopyParameterValues(int32 ParameterGroupIndex)
+void FMaterialInstanceParameterDetails::OnCopyParameterValues(const TArray<int32> ParameterGroupIndices)
 {
-	if (!MaterialEditorInstance || !MaterialEditorInstance->ParameterGroups.IsValidIndex(ParameterGroupIndex))
+	if (!MaterialEditorInstance)
 	{
 		return;
 	}
-	FEditorParameterGroup& ParameterGroup = MaterialEditorInstance->ParameterGroups[ParameterGroupIndex];
 
 	TStringBuilder<4096> CombinedValue;
-
-	const int32 NumParams = ParameterGroup.Parameters.Num();
-	for (int32 ParamIdx = 0; ParamIdx < NumParams; ++ParamIdx)
+	for (const int32 ParameterGroupIndex : ParameterGroupIndices)
 	{
-		UDEditorParameterValue* Parameter = ParameterGroup.Parameters[ParamIdx];
-		const FName ParamName = Parameter->ParameterInfo.Name;
+		if (!MaterialEditorInstance->ParameterGroups.IsValidIndex(ParameterGroupIndex))
+		{
+			continue;
+		}
 
-		const TCHAR* Prefix = (ParamIdx == 0) ? TEXT("") : TEXT(",");
+		FEditorParameterGroup& ParameterGroup = MaterialEditorInstance->ParameterGroups[ParameterGroupIndex];
 
-		// Include the value in the result entry only if the parameter is overridden.
-		const bool bOverride = FMaterialPropertyHelpers::IsOverriddenExpression(Parameter);
-		if (bOverride)
+		const int32 NumParams = ParameterGroup.Parameters.Num();
+		for (int32 ParamIdx = 0; ParamIdx < NumParams; ++ParamIdx)
 		{
-			FProperty* ParameterValueProperty = Parameter->GetClass()->FindPropertyByName("ParameterValue");
-			if (ParameterValueProperty != nullptr)
+			UDEditorParameterValue* Parameter = ParameterGroup.Parameters[ParamIdx];
+			const FName ParamName = Parameter->ParameterInfo.Name;
+
+			const TCHAR* Prefix = (CombinedValue.Len() == 0) ? TEXT("") : TEXT(",");
+
+			// Include the value in the result entry only if the parameter is overridden.
+			const bool bOverride = FMaterialPropertyHelpers::IsOverriddenExpression(Parameter);
+			if (bOverride)
 			{
-				FString ParameterValueString;
-				if (ParameterValueProperty->ExportText_InContainer(0, ParameterValueString, Parameter, Parameter, Parameter, PPF_Copy))
+				FProperty* ParameterValueProperty = Parameter->GetClass()->FindPropertyByName("ParameterValue");
+				if (ParameterValueProperty != nullptr)
 				{
-					CombinedValue.Appendf(TEXT("%s%s.Override=True,%s.Value=\"%s\""),
-						Prefix,
-						*(ParamName.ToString()),
-						*(ParamName.ToString()),
-						*(ParameterValueString.ReplaceCharWithEscapedChar()));
+					FString ParameterValueString;
+					if (ParameterValueProperty->ExportText_InContainer(0, ParameterValueString, Parameter, Parameter, Parameter, PPF_Copy))
+					{
+						CombinedValue.Appendf(TEXT("%s%s.Override=True,%s.Value=\"%s\""),
+							Prefix,
+							*(ParamName.ToString()),
+							*(ParamName.ToString()),
+							*(ParameterValueString.ReplaceCharWithEscapedChar()));
+					}
 				}
 			}
-		}
-		else
-		{
-			CombinedValue.Appendf(TEXT("%s%s.Override=False"), Prefix, *(ParamName.ToString()));
+			else
+			{
+				CombinedValue.Appendf(TEXT("%s%s.Override=False"), Prefix, *(ParamName.ToString()));
+			}
 		}
 	}
 
@@ -794,30 +789,60 @@ void FMaterialInstanceParameterDetails::OnCopyParameterValues(int32 ParameterGro
 	}
 }
 
-bool FMaterialInstanceParameterDetails::CanCopyParameterValues(int32 ParameterGroupIndex)
+bool FMaterialInstanceParameterDetails::CanCopyParameterValues(const TArray<int32> ParameterGroupIndices)
 {
-	return MaterialEditorInstance && MaterialEditorInstance->ParameterGroups.IsValidIndex(ParameterGroupIndex)
-		&& (MaterialEditorInstance->ParameterGroups[ParameterGroupIndex].Parameters.Num() > 0);
+	if (!MaterialEditorInstance)
+	{
+		return false;
+	}
+
+	for (const int32 ParameterGroupIndex : ParameterGroupIndices)
+	{
+		if (MaterialEditorInstance->ParameterGroups.IsValidIndex(ParameterGroupIndex) &&
+			MaterialEditorInstance->ParameterGroups[ParameterGroupIndex].Parameters.Num() > 0)
+		{
+			return true;
+		}
+	}
+
+	return false;
 }
 
-void FMaterialInstanceParameterDetails::OnPasteParameterValues(int32 ParameterGroupIndex)
+void FMaterialInstanceParameterDetails::OnPasteParameterValues(const TArray<int32> ParameterGroupIndices)
 {
-	if (!MaterialEditorInstance || !MaterialEditorInstance->ParameterGroups.IsValidIndex(ParameterGroupIndex))
+	if (!MaterialEditorInstance)
 	{
 		return;
 	}
 
 	FString ClipboardContent;
 	FPlatformApplicationMisc::ClipboardPaste(ClipboardContent);
-	if (!ClipboardContent.IsEmpty())
+	if (ClipboardContent.IsEmpty())
+	{
+		return;
+	}
+
+	FScopedTransaction Transaction(LOCTEXT("PasteMaterialInstanceParameters", "Paste Material Instance Parameters"));
+
+	bool bIsValidTransaction = false;
+
+	for (const int32 ParameterGroupIndex : ParameterGroupIndices)
 	{
-		FScopedTransaction Transaction(LOCTEXT("PasteMaterialInstanceParameters", "Paste Material Instance Parameters"));
+		if (!MaterialEditorInstance->ParameterGroups.IsValidIndex(ParameterGroupIndex))
+		{
+			continue;
+		}
 
-		MaterialEditorInstance->Modify();
+		if (!bIsValidTransaction)
+		{
+			bIsValidTransaction = true;
+			MaterialEditorInstance->Modify();
+		}
 
-		for (int32 ParamIdx = 0; ParamIdx < MaterialEditorInstance->ParameterGroups[ParameterGroupIndex].Parameters.Num(); ++ParamIdx)
+		FEditorParameterGroup& ParameterGroup = MaterialEditorInstance->ParameterGroups[ParameterGroupIndex];
+		for (int32 ParamIdx = 0; ParamIdx < ParameterGroup.Parameters.Num(); ++ParamIdx)
 		{
-			UDEditorParameterValue* Parameter = MaterialEditorInstance->ParameterGroups[ParameterGroupIndex].Parameters[ParamIdx];
+			UDEditorParameterValue* Parameter = ParameterGroup.Parameters[ParamIdx];
 			Parameter->Modify();
 
 			const FName ParamName = Parameter->ParameterInfo.Name;
@@ -845,15 +870,22 @@ void FMaterialInstanceParameterDetails::OnPasteParameterValues(int32 ParameterGr
 			}
 		}
 
-		MaterialEditorInstance->PostEditChange();
-		FEditorSupportDelegates::RedrawAllViewports.Broadcast();
 	}
+
+	if (!bIsValidTransaction)
+	{
+		Transaction.Cancel();
+		return;
+	}
+
+	MaterialEditorInstance->PostEditChange();
+	FEditorSupportDelegates::RedrawAllViewports.Broadcast();
 }
 
-bool FMaterialInstanceParameterDetails::CanPasteParameterValues(int32 ParameterGroupIndex)
+bool FMaterialInstanceParameterDetails::CanPasteParameterValues(const TArray<int32> ParameterGroupIndices)
 {
 	// First check the same criteria as copying.
-	if (!CanCopyParameterValues(ParameterGroupIndex))
+	if (!CanCopyParameterValues(ParameterGroupIndices))
 	{
 		return false;
 	}
diff --git a/Engine/Source/Editor/MaterialEditor/Private/MaterialEditorInstanceDetailCustomization.h b/Engine/Source/Editor/MaterialEditor/Private/MaterialEditorInstanceDetailCustomization.h
--- a/Engine/Source/Editor/MaterialEditor/Private/MaterialEditorInstanceDetailCustomization.h
+++ b/Engine/Source/Editor/MaterialEditor/Private/MaterialEditorInstanceDetailCustomization.h
@@ -48,7 +48,7 @@ private:
 	void CreateSingleGroupWidget(struct FEditorParameterGroup& ParameterGroup, TSharedPtr<IPropertyHandle> ParameterGroupProperty, class IDetailGroup& DetailGroup, int32 GroupIndex = -1, bool bForceShowParam = false);
 
 	/** Enable/Disable all parameter properties in a group */
-	static void EnableGroupParameters(struct FEditorParameterGroup& ParameterGroup, bool ShouldEnable);
+	void EnableGroupParameters(TArray<int32> ParameterGroupIndices, bool bShouldEnable);
 
 	/** Called to check if an asset can be set as a parent */
 	bool OnShouldSetAsset(const FAssetData& InAssetData) const;
@@ -74,22 +74,22 @@ private:
 	 * Copy all parameter values in a parameter group to the clipboard
 	 * in the format "Param1=\"Value1\"\nParams2=\"Value2\"..."
 	 */
-	void OnCopyParameterValues(int32 ParameterGroupIndex);
+	void OnCopyParameterValues(TArray<int32> ParameterGroupIndices);
 
 	/** Whether it is possible to copy parameter values for the given parent group index */
-	bool CanCopyParameterValues(int32 ParameterGroupIndex);
+	bool CanCopyParameterValues(TArray<int32> ParameterGroupIndices);
 
 	/**
 	 * Paste parameter values from the clipboard, assumed to be
 	 * in the format copied by CopyParameterValues.
 	 */
-	void OnPasteParameterValues(int32 ParameterGroupIndex);
+	void OnPasteParameterValues(TArray<int32> ParameterGroupIndices);
 
 	/**
 	 * Whether it is possible to paste parameter values onto the given group index,
 	 * and if there is anything on the clipboard to paste from.
 	 */
-	bool CanPasteParameterValues(int32 ParameterGroupIndex);
+	bool CanPasteParameterValues(TArray<int32> ParameterGroupIndices);
 
 	/** Creates all the lightmass property override widgets. */
 	void CreateLightmassOverrideWidgets(IDetailLayoutBuilder& DetailLayout);
