EditorUsability : MaterialEditor
- Added shaderball as a mesh preview in material editor
- Save the last used preview mesh in config to load it again on material & material instance editor opening
- Materials will use their own preview mesh property if set, otherwise fallback on the last used editor preview mesh, when last used preview mesh can't be loaded, falling back to the sphere mesh

#virtualized
#Synced-CL 50297267

diff --git a/Engine/Source/Editor/EditorStyle/Private/StarshipStyle.cpp b/Engine/Source/Editor/EditorStyle/Private/StarshipStyle.cpp
--- a/Engine/Source/Editor/EditorStyle/Private/StarshipStyle.cpp
+++ b/Engine/Source/Editor/EditorStyle/Private/StarshipStyle.cpp
@@ -7521,6 +7521,7 @@ void FStarshipEditorStyle::FStyle::SetupToolkitStyles()
 
 		Set( "MaterialEditor.SetCylinderPreview", new IMAGE_BRUSH_SVG("Icons/AssetIcons/ModelingCylinder_16", Icon16x16));
 		Set( "MaterialEditor.SetSpherePreview", new IMAGE_BRUSH_SVG("Icons/AssetIcons/ModelingSphereShaded_16", Icon16x16));
+		Set( "MaterialEditor.SetShaderBallPreview", new IMAGE_BRUSH_SVG("Icons/AssetIcons/ModelingShaderBall_16", Icon16x16));
 		Set( "MaterialEditor.SetPlanePreview", new IMAGE_BRUSH_SVG( "Icons/AssetIcons/ModelingPlane_16", Icon16x16) );
 		Set( "MaterialEditor.SetCubePreview", new IMAGE_BRUSH_SVG( "Icons/AssetIcons/ModelingBox_16", Icon16x16) );
 		Set( "MaterialEditor.SetPreviewMeshFromSelection", new IMAGE_BRUSH_SVG( "Starship/AssetIcons/StaticMesh_16", Icon16x16 ) );
diff --git a/Engine/Source/Editor/MaterialEditor/Private/MaterialEditor.cpp b/Engine/Source/Editor/MaterialEditor/Private/MaterialEditor.cpp
--- a/Engine/Source/Editor/MaterialEditor/Private/MaterialEditor.cpp
+++ b/Engine/Source/Editor/MaterialEditor/Private/MaterialEditor.cpp
@@ -687,13 +687,6 @@ void FMaterialEditor::InitMaterialEditor( const EToolkitMode::Type Mode, const T
 
 	// Load editor settings from disk.
 	LoadEditorSettings();
-	
-	// Set the preview mesh for the material.  This call must occur after the toolbar is initialized.
-	if (!SetPreviewAssetByName(*Material->PreviewMesh.ToString()))
-	{
-		// The material preview mesh couldn't be found or isn't loaded.  Default to the one of the primitive types.
-		SetPreviewAsset( GUnrealEd->GetThumbnailManager()->EditorSphere );
-	}
 
 	// Initialize expression previews.
 	if (MaterialFunction)
@@ -2472,7 +2465,7 @@ void FMaterialEditor::RefreshPreviewViewport()
 
 void FMaterialEditor::LoadEditorSettings()
 {
-	EditorOptions = NewObject<UMaterialEditorOptions>();
+	EditorOptions = UMaterialEditorOptions::Get();
 	
 	if (EditorOptions->bHideUnusedConnectorsSetting) {OnHideConnectors();}
 	if (bLivePreview != EditorOptions->bLivePreviewUpdate)
@@ -2495,11 +2488,11 @@ void FMaterialEditor::LoadEditorSettings()
 		ToggleHideUnrelatedNodes();
 	}
 
-	// Primitive type
-	int32 PrimType;
-	if(GConfig->GetInt(TEXT("MaterialEditor"), TEXT("PrimType"), PrimType, GEditorPerProjectIni))
+	// Set the preview mesh for the material.  This call must occur after the toolbar is initialized.
+	if (!SetPreviewAssetByName(*Material->PreviewMesh.ToString()))
 	{
-		PreviewViewport->OnSetPreviewPrimitive((EThumbnailPrimType)PrimType);
+		// The material preview mesh couldn't be found or isn't loaded. Default to the saved editor preview mesh.
+		SetPreviewAsset(EditorOptions->GetPreviewMesh());
 	}
 }
 
@@ -2518,8 +2511,6 @@ void FMaterialEditor::SaveEditorSettings()
 		EditorOptions->bHideUnrelatedNodes          = bHideUnrelatedNodes;
 		EditorOptions->SaveConfig();
 	}
-
-	GConfig->SetInt(TEXT("MaterialEditor"), TEXT("PrimType"), PreviewViewport->PreviewPrimType, GEditorPerProjectIni);
 }
 
 void FMaterialEditor::RegenerateCodeView(bool bForce)
@@ -7244,10 +7235,13 @@ void FMaterialEditor::NotifyPostChange( const FPropertyChangedEvent& PropertyCha
 			(NameOfPropertyThatChanged == GET_MEMBER_NAME_CHECKED(UMaterial, bUsedWithSkeletalMesh)))
 		PRAGMA_ENABLE_DEPRECATION_WARNINGS
 		{
-			// Revert to a sphere static mesh preview when usage isn't supported.
+			// Revert to editor preview mesh when usage isn't supported.
 			if (!SetPreviewAssetByName(*Material->PreviewMesh.ToString()))
 			{
-				SetPreviewAsset(GUnrealEd->GetThumbnailManager()->EditorSphere);
+				if (EditorOptions)
+				{
+					SetPreviewAsset(EditorOptions->GetPreviewMesh());
+				}
 			}
 		}
 		else if( NameOfPropertyThatChanged == GET_MEMBER_NAME_CHECKED(UMaterial, MaterialDomain) ||
diff --git a/Engine/Source/Editor/MaterialEditor/Private/MaterialEditorActions.cpp b/Engine/Source/Editor/MaterialEditor/Private/MaterialEditorActions.cpp
--- a/Engine/Source/Editor/MaterialEditor/Private/MaterialEditorActions.cpp
+++ b/Engine/Source/Editor/MaterialEditor/Private/MaterialEditorActions.cpp
@@ -37,6 +37,7 @@ void FMaterialEditorCommands::RegisterCommands()
 	UI_COMMAND( SetSpherePreview, "Sphere", "Sets the preview mesh to a sphere primitive.", EUserInterfaceActionType::RadioButton, FInputChord() );
 	UI_COMMAND( SetPlanePreview, "Plane", "Sets the preview mesh to a plane primitive.", EUserInterfaceActionType::RadioButton, FInputChord() );
 	UI_COMMAND( SetCubePreview, "Cube", "Sets the preview mesh to a cube primitive.", EUserInterfaceActionType::RadioButton, FInputChord() );
+	UI_COMMAND( SetShaderBallPreview, "ShaderBall", "Sets the preview mesh to a shader ball primitive.", EUserInterfaceActionType::RadioButton, FInputChord() );
 	UI_COMMAND( SetPreviewMeshFromSelection, "Mesh", "Sets the preview mesh based on the current content browser selection.", EUserInterfaceActionType::RadioButton, FInputChord() );
 	UI_COMMAND( TogglePreviewBackground, "Background", "Toggles the preview pane's background.", EUserInterfaceActionType::ToggleButton, FInputChord() );
 
diff --git a/Engine/Source/Editor/MaterialEditor/Private/MaterialEditorViewportToolbarSections.cpp b/Engine/Source/Editor/MaterialEditor/Private/MaterialEditorViewportToolbarSections.cpp
--- a/Engine/Source/Editor/MaterialEditor/Private/MaterialEditorViewportToolbarSections.cpp
+++ b/Engine/Source/Editor/MaterialEditor/Private/MaterialEditorViewportToolbarSections.cpp
@@ -102,13 +102,18 @@ void UE::MaterialEditor::ExtendPreviewSceneSettingsSubmenu(FName InSubmenuName)
 	FToolMenuEntryToolBarData ToolbarData;
 	ToolbarData.BlockGroupName = "PreviewMeshOptions";
 	ToolbarData.ResizeParams.ClippingPriority = 2000;
-
-	for (const TSharedPtr<FUICommandInfo>& Command : {
-			 FMaterialEditorCommands::Get().SetSpherePreview,
-			 FMaterialEditorCommands::Get().SetCylinderPreview,
-			 FMaterialEditorCommands::Get().SetPlanePreview,
-			 FMaterialEditorCommands::Get().SetCubePreview,
-		 })
+	
+	const FMaterialEditorCommands& MaterialEditorCommands = FMaterialEditorCommands::Get();
+	const TArray<TSharedPtr<FUICommandInfo>> PreviewMeshCommands
+	{
+		MaterialEditorCommands.SetSpherePreview,
+		MaterialEditorCommands.SetCylinderPreview,
+		MaterialEditorCommands.SetPlanePreview,
+		MaterialEditorCommands.SetCubePreview,
+		MaterialEditorCommands.SetShaderBallPreview
+	};
+
+	for (const TSharedPtr<FUICommandInfo>& Command : PreviewMeshCommands)
 	{
 		FToolMenuEntry& Entry = PreviewMeshSection.AddMenuEntry(Command);
 		Entry.SetShowInToolbarTopLevel(true);
@@ -116,7 +121,7 @@ void UE::MaterialEditor::ExtendPreviewSceneSettingsSubmenu(FName InSubmenuName)
 	}
 
 	FToolMenuEntry& MeshFromSelectionEntry =
-		PreviewMeshSection.AddMenuEntry(FMaterialEditorCommands::Get().SetPreviewMeshFromSelection);
+		PreviewMeshSection.AddMenuEntry(MaterialEditorCommands.SetPreviewMeshFromSelection);
 	MeshFromSelectionEntry.Label = LOCTEXT("SetPreviewMeshFromSelectionLabel", "Static Mesh in Content Browser");
 	MeshFromSelectionEntry.Icon = FSlateIconFinder::FindIconForClass(UStaticMesh::StaticClass());
 	MeshFromSelectionEntry.SetShowInToolbarTopLevel(true);
diff --git a/Engine/Source/Editor/MaterialEditor/Private/MaterialExpressionClasses.cpp b/Engine/Source/Editor/MaterialEditor/Private/MaterialExpressionClasses.cpp
--- a/Engine/Source/Editor/MaterialEditor/Private/MaterialExpressionClasses.cpp
+++ b/Engine/Source/Editor/MaterialEditor/Private/MaterialExpressionClasses.cpp
@@ -84,7 +84,7 @@ void MaterialExpressionClasses::InitMaterialExpressionClasses()
 		static const auto CVarMaterialTranslatorEnableNew = IConsoleManager::Get().FindTConsoleVariableDataBool(TEXT("r.Material.Translator.EnableNew"));
 		const bool bMaterialTranslatorEnableNew = CVarMaterialTranslatorEnableNew->GetValueOnAnyThread();
 
-		UMaterialEditorOptions* TempEditorOptions = NewObject<UMaterialEditorOptions>();
+		UMaterialEditorOptions* EditorOptions = UMaterialEditorOptions::Get();
 		UClass* BaseType = UMaterialExpression::StaticClass();
 		if( BaseType )
 		{
@@ -157,9 +157,9 @@ void MaterialExpressionClasses::InitMaterialExpressionClasses()
 							}
 
 							// See if it is in the favorites array...
-							for (int32 FavoriteIndex = 0; FavoriteIndex < TempEditorOptions->FavoriteExpressions.Num(); FavoriteIndex++)
+							for (int32 FavoriteIndex = 0; FavoriteIndex < EditorOptions->FavoriteExpressions.Num(); FavoriteIndex++)
 							{
-								if (Class->GetName() == TempEditorOptions->FavoriteExpressions[FavoriteIndex])
+								if (Class->GetName() == EditorOptions->FavoriteExpressions[FavoriteIndex])
 								{
 									FavoriteExpressionClasses.AddUnique(MaterialExpression);
 								}
diff --git a/Engine/Source/Editor/MaterialEditor/Private/MaterialInstanceEditor.cpp b/Engine/Source/Editor/MaterialEditor/Private/MaterialInstanceEditor.cpp
--- a/Engine/Source/Editor/MaterialEditor/Private/MaterialInstanceEditor.cpp
+++ b/Engine/Source/Editor/MaterialEditor/Private/MaterialInstanceEditor.cpp
@@ -73,6 +73,7 @@
 #include "ViewportToolbar/UnrealEdViewportToolbar.h"
 
 #include "DiffTool/Widgets/SMaterialRevisionMenu.h"
+#include "Preferences/MaterialEditorOptions.h"
 
 #define LOCTEXT_NAMESPACE "MaterialInstanceEditor"
 
@@ -525,11 +526,10 @@ void FMaterialInstanceEditor::InitMaterialInstanceEditor( const EToolkitMode::Ty
 
 		if ( InstanceConstant->Parent == nullptr || !SetPreviewAssetByName( *InstanceConstant->Parent->PreviewMesh.ToString() ) )
 		{
-			USceneThumbnailInfoWithPrimitive* ThumbnailInfoWithPrim = Cast<USceneThumbnailInfoWithPrimitive>( InstanceConstant->ThumbnailInfo );
-
-			if ( ThumbnailInfoWithPrim != nullptr )
+			// Fallback on the default material editor preview mesh
+			if (const UMaterialEditorOptions* MaterialEditorOptions = UMaterialEditorOptions::Get())
 			{
-				SetPreviewAssetByName( *ThumbnailInfoWithPrim->PreviewMesh.ToString() );
+				SetPreviewAsset(MaterialEditorOptions->GetPreviewMesh());
 			}
 		}
 	}
@@ -1366,12 +1366,20 @@ void FMaterialInstanceEditor::RefreshPreviewAsset()
 	UObject* PreviewAsset = MaterialEditorInstance->SourceInstance->PreviewMesh.TryLoad();
 	if (!PreviewAsset)
 	{
-		// Attempt to use the parent material's preview mesh if the instance's preview mesh is invalid, and use a default
-		//	sphere instead if the parent's mesh is also invalid
+		// Attempt to use the parent material's preview mesh if the instance's preview mesh is invalid, and use the default
+		//	preview mesh instead if the parent's mesh is also invalid
 		UMaterialInterface* ParentMaterial = MaterialEditorInstance->SourceInstance->Parent;
 
-		UObject* ParentPreview = ParentMaterial != nullptr ? ParentMaterial->PreviewMesh.TryLoad() : nullptr;
-		PreviewAsset = ParentPreview != nullptr ? ParentPreview : ToRawPtr(GUnrealEd->GetThumbnailManager()->EditorSphere);
+		PreviewAsset = ParentMaterial != nullptr ? ParentMaterial->PreviewMesh.TryLoad() : nullptr;
+
+		if (!PreviewAsset)
+		{
+			// Fallback on the default material editor preview mesh
+			if (const UMaterialEditorOptions* MaterialEditorOptions = UMaterialEditorOptions::Get())
+			{
+				PreviewAsset = MaterialEditorOptions->GetPreviewMesh();
+			}
+		}
 
 		USceneThumbnailInfoWithPrimitive* ThumbnailInfo = Cast<USceneThumbnailInfoWithPrimitive>(MaterialEditorInstance->SourceInstance->ThumbnailInfo);
 		if (ThumbnailInfo)
@@ -1752,15 +1760,17 @@ void FMaterialInstanceEditor::SaveAssetAs_Execute()
 void FMaterialInstanceEditor::SaveSettings()
 {
 	GConfig->SetBool(TEXT("MaterialInstanceEditor"), TEXT("bDrawGrid"), PreviewVC->IsRealtime(), GEditorPerProjectIni);
-	GConfig->SetInt(TEXT("MaterialInstanceEditor"), TEXT("PrimType"), PreviewVC->PreviewPrimType, GEditorPerProjectIni);
+
+	if (UMaterialEditorOptions* MaterialEditorOptions = UMaterialEditorOptions::Get())
+	{
+		MaterialEditorOptions->SaveConfig();
+	}
 }
 
 void FMaterialInstanceEditor::LoadSettings()
 {
 	bool bRealtime=false;
-	int32 PrimType=static_cast<EThumbnailPrimType>( TPT_Sphere );
 	GConfig->GetBool(TEXT("MaterialInstanceEditor"), TEXT("bDrawGrid"), bRealtime, GEditorPerProjectIni);
-	GConfig->GetInt(TEXT("MaterialInstanceEditor"), TEXT("PrimType"), PrimType, GEditorPerProjectIni);
 
 	if(PreviewVC.IsValid())
 	{
@@ -1768,8 +1778,6 @@ void FMaterialInstanceEditor::LoadSettings()
 		{
 			PreviewVC->OnToggleRealtime();
 		}
-
-		PreviewVC->OnSetPreviewPrimitive( static_cast<EThumbnailPrimType>(PrimType), true);
 	}
 }
 
diff --git a/Engine/Source/Editor/MaterialEditor/Private/SMaterialEditorViewport.cpp b/Engine/Source/Editor/MaterialEditor/Private/SMaterialEditorViewport.cpp
--- a/Engine/Source/Editor/MaterialEditor/Private/SMaterialEditorViewport.cpp
+++ b/Engine/Source/Editor/MaterialEditor/Private/SMaterialEditorViewport.cpp
@@ -53,6 +53,7 @@
 #include "PreviewProfileController.h"
 #include "SMaterialEditorTopologyWidget.h"
 #include "UnrealWidget.h"
+#include "Preferences/MaterialEditorOptions.h"
 
 /** Viewport Client for the preview viewport */
 class FMaterialEditorViewportClient : public FEditorViewportClient
@@ -281,7 +282,11 @@ void SMaterialEditor3DPreviewViewport::Construct(const FArguments& InArgs)
 			SetPreviewMaterial(Material);
 		}
 
-	SetPreviewAsset( GUnrealEd->GetThumbnailManager()->EditorSphere );
+	// Set default preview mesh
+	if (const UMaterialEditorOptions* MaterialEditorOptions = UMaterialEditorOptions::Get())
+	{
+		SetPreviewAsset(MaterialEditorOptions->GetPreviewMesh());
+	}
 
 	OnPropertyChangedHandle = FCoreUObjectDelegates::FOnObjectPropertyChanged::FDelegate::CreateRaw(this, &SMaterialEditor3DPreviewViewport::OnPropertyChanged);
 	OnPropertyChangedHandleDelegateHandle = FCoreUObjectDelegates::OnObjectPropertyChanged.Add(OnPropertyChangedHandle);
@@ -351,12 +356,19 @@ bool SMaterialEditor3DPreviewViewport::SetPreviewAsset(UObject* InAsset)
 		return false;
 	}
 
+	// This should never occur since we are comparing a mesh against a component...
 	if (PreviewMeshComponent == InAsset)
 	{
 		// No change
 		return true;
 	}
 
+	if (PreviewMeshComponent && FComponentAssetBrokerage::GetAssetFromComponent(PreviewMeshComponent) == InAsset)
+	{
+		// No change
+		return true;
+	}
+
 	// Unregister the current component
 	if (PreviewMeshComponent != nullptr)
 	{
@@ -390,17 +402,17 @@ bool SMaterialEditor3DPreviewViewport::SetPreviewAsset(UObject* InAsset)
 		else if (StaticMesh == GUnrealEd->GetThumbnailManager()->EditorPlane)
 		{
 			PreviewPrimType = TPT_Plane;
+			// Update the rotation of the plane mesh so that it is front facing to the viewport camera's default forward view.
+			const FRotator PlaneRotation(0.0f, 180.0f, 0.0f);
+			Transform.SetRotation(FQuat(PlaneRotation));
 		}
-		else
+		else if (StaticMesh == GUnrealEd->GetThumbnailManager()->EditorShaderBall)
 		{
-			PreviewPrimType = TPT_None;
+			PreviewPrimType = TPT_ShaderBall;
 		}
-
-		// Update the rotation of the plane mesh so that it is front facing to the viewport camera's default forward view.
-		if (PreviewPrimType == TPT_Plane)
+		else
 		{
-			const FRotator PlaneRotation(0.0f, 180.0f, 0.0f);
-			Transform.SetRotation(FQuat(PlaneRotation));
+			PreviewPrimType = TPT_None;
 		}
 	}
 	else if (InAsset != nullptr)
@@ -731,34 +743,41 @@ void SMaterialEditor3DPreviewViewport::BindCommands()
 
 	if(MaterialEditorPtr.IsValid())
 	{
-	CommandList->Append(MaterialEditorPtr.Pin()->GetToolkitCommands());
+		CommandList->Append(MaterialEditorPtr.Pin()->GetToolkitCommands());
 	}
 
 	// Add the commands to the toolkit command list so that the toolbar buttons can find them
+	constexpr bool bInitialLoad = false;
 	CommandList->MapAction(
 		Commands.SetCylinderPreview,
-		FExecuteAction::CreateSP( this, &SMaterialEditor3DPreviewViewport::OnSetPreviewPrimitive, TPT_Cylinder, false ),
+		FExecuteAction::CreateSP( this, &SMaterialEditor3DPreviewViewport::OnSetPreviewPrimitive, TPT_Cylinder, bInitialLoad ),
 		FCanExecuteAction(),
 		FIsActionChecked::CreateSP( this, &SMaterialEditor3DPreviewViewport::IsPreviewPrimitiveChecked, TPT_Cylinder ) );
 
 	CommandList->MapAction(
 		Commands.SetSpherePreview,
-		FExecuteAction::CreateSP( this, &SMaterialEditor3DPreviewViewport::OnSetPreviewPrimitive, TPT_Sphere, false ),
+		FExecuteAction::CreateSP( this, &SMaterialEditor3DPreviewViewport::OnSetPreviewPrimitive, TPT_Sphere, bInitialLoad ),
 		FCanExecuteAction(),
 		FIsActionChecked::CreateSP( this, &SMaterialEditor3DPreviewViewport::IsPreviewPrimitiveChecked, TPT_Sphere ) );
 
 	CommandList->MapAction(
 		Commands.SetPlanePreview,
-		FExecuteAction::CreateSP( this, &SMaterialEditor3DPreviewViewport::OnSetPreviewPrimitive, TPT_Plane, false ),
+		FExecuteAction::CreateSP( this, &SMaterialEditor3DPreviewViewport::OnSetPreviewPrimitive, TPT_Plane, bInitialLoad ),
 		FCanExecuteAction(),
 		FIsActionChecked::CreateSP( this, &SMaterialEditor3DPreviewViewport::IsPreviewPrimitiveChecked, TPT_Plane ) );
 
 	CommandList->MapAction(
 		Commands.SetCubePreview,
-		FExecuteAction::CreateSP( this, &SMaterialEditor3DPreviewViewport::OnSetPreviewPrimitive, TPT_Cube, false ),
+		FExecuteAction::CreateSP( this, &SMaterialEditor3DPreviewViewport::OnSetPreviewPrimitive, TPT_Cube, bInitialLoad ),
 		FCanExecuteAction(),
 		FIsActionChecked::CreateSP( this, &SMaterialEditor3DPreviewViewport::IsPreviewPrimitiveChecked, TPT_Cube ) );
 
+	CommandList->MapAction(
+    	Commands.SetShaderBallPreview,
+    	FExecuteAction::CreateSP( this, &SMaterialEditor3DPreviewViewport::OnSetPreviewPrimitive, TPT_ShaderBall, bInitialLoad ),
+    	FCanExecuteAction(),
+    	FIsActionChecked::CreateSP( this, &SMaterialEditor3DPreviewViewport::IsPreviewPrimitiveChecked, TPT_ShaderBall ) );
+
 	CommandList->MapAction(
 		Commands.SetPreviewMeshFromSelection,
 		FExecuteAction::CreateSP( this, &SMaterialEditor3DPreviewViewport::OnSetPreviewMeshFromSelection ),
@@ -789,19 +808,39 @@ void SMaterialEditor3DPreviewViewport::OnSetPreviewPrimitive(EThumbnailPrimType
 {
 	if (SceneViewport.IsValid())
 	{
-		UStaticMesh* Primitive = nullptr;
+		UObject* Primitive = nullptr;
 		switch (PrimType)
 		{
-		case TPT_Cylinder: Primitive = GUnrealEd->GetThumbnailManager()->EditorCylinder; break;
+		case TPT_None:
+			{
+				if (const UMaterialEditorOptions* MaterialEditorOptions = UMaterialEditorOptions::Get())
+				{
+					Primitive = MaterialEditorOptions->GetPreviewMesh();
+					
+					// If invalid, take the next preview mesh otherwise break
+					if (!!Primitive)
+					{
+						break;
+					}
+				}
+			}
 		case TPT_Sphere: Primitive = GUnrealEd->GetThumbnailManager()->EditorSphere; break;
-		case TPT_Plane: Primitive = GUnrealEd->GetThumbnailManager()->EditorPlane; break;
 		case TPT_Cube: Primitive = GUnrealEd->GetThumbnailManager()->EditorCube; break;
+		case TPT_Plane: Primitive = GUnrealEd->GetThumbnailManager()->EditorPlane; break;
+		case TPT_Cylinder: Primitive = GUnrealEd->GetThumbnailManager()->EditorCylinder; break;
+		case TPT_ShaderBall: Primitive = GUnrealEd->GetThumbnailManager()->EditorShaderBall; break;
 		}
 
 		if (Primitive != nullptr)
 		{
 			SetPreviewAsset(Primitive);
 			
+			// Save the preview mesh as default
+			if (UMaterialEditorOptions* MaterialEditorOptions = UMaterialEditorOptions::Get())
+			{
+				MaterialEditorOptions->SetPreviewMesh(Primitive);
+			}
+			
 			// Clear the thumbnail preview mesh
 			if (UMaterialInterface* MaterialInterface = GetPreviewingMaterialInterface())
 			{
@@ -842,6 +881,13 @@ void SMaterialEditor3DPreviewViewport::OnSetPreviewMeshFromSelection()
 				{
 					SetPreviewAsset(TestAsset);
 					MaterialInterface->PreviewMesh = TestAsset->GetPathName();
+
+					// Save the preview mesh as default
+					if (UMaterialEditorOptions* MaterialEditorOptions = UMaterialEditorOptions::Get())
+					{
+						MaterialEditorOptions->SetPreviewMesh(TestAsset);
+					}
+
 					bFoundPreviewMesh = true;
 				}
 			}
diff --git a/Engine/Source/Editor/MaterialEditor/Public/MaterialEditorActions.h b/Engine/Source/Editor/MaterialEditor/Public/MaterialEditorActions.h
--- a/Engine/Source/Editor/MaterialEditor/Public/MaterialEditorActions.h
+++ b/Engine/Source/Editor/MaterialEditor/Public/MaterialEditorActions.h
@@ -73,6 +73,9 @@ public:
 	/** Sets the preview mesh to a cube */
 	TSharedPtr< FUICommandInfo > SetCubePreview;
 	
+	/** Sets the preview mesh to a shader ball */
+	TSharedPtr< FUICommandInfo > SetShaderBallPreview;
+	
 	/** Sets the preview mesh to the current selection in the level editor */
 	TSharedPtr< FUICommandInfo > SetPreviewMeshFromSelection;
 	
diff --git a/Engine/Source/Editor/UnrealEd/Classes/Preferences/MaterialEditorOptions.h b/Engine/Source/Editor/UnrealEd/Classes/Preferences/MaterialEditorOptions.h
--- a/Engine/Source/Editor/UnrealEd/Classes/Preferences/MaterialEditorOptions.h
+++ b/Engine/Source/Editor/UnrealEd/Classes/Preferences/MaterialEditorOptions.h
@@ -1,18 +1,17 @@
 // Copyright Epic Games, Inc. All Rights Reserved.
 
-/**
- *
- * A configuration class used by the UMaterial Editor to save editor
- * settings across sessions.
- */
-
 #pragma once
 
 #include "CoreMinimal.h"
 #include "UObject/ObjectMacros.h"
 #include "UObject/Object.h"
+#include "UObject/SoftObjectPath.h"
 #include "MaterialEditorOptions.generated.h"
 
+/**
+ * A configuration class used by the UMaterial Editor to save editor
+ * settings across sessions.
+ */
 UCLASS(hidecategories=Object, config=EditorPerProjectUserSettings, MinimalAPI)
 class UMaterialEditorOptions : public UObject
 {
@@ -49,5 +48,18 @@ class UMaterialEditorOptions : public UObject
 	/** The users favorite material expressions. */
 	UPROPERTY(EditAnywhere, config, Category=Options)
 	TArray<FString> FavoriteExpressions;
+	
+	/** Preview mesh used in material editor viewport */
+	UPROPERTY(EditAnywhere, config, Category=Options, meta=(AllowedClasses="/Script/Engine.StaticMesh,/Script/Engine.SkeletalMesh", ExactClass="true"))
+	FSoftObjectPath PreviewMesh;
+
+public:
+	UNREALED_API static UMaterialEditorOptions* Get();
+
+	/** Loads and returns the preview mesh or default preview mesh if none */
+	UNREALED_API UObject* GetPreviewMesh() const;
+
+	/** Changes the preview mesh */
+	UNREALED_API void SetPreviewMesh(UObject* InPreviewMesh);
 };
 
diff --git a/Engine/Source/Editor/UnrealEd/Classes/ThumbnailRendering/ThumbnailManager.h b/Engine/Source/Editor/UnrealEd/Classes/ThumbnailRendering/ThumbnailManager.h
--- a/Engine/Source/Editor/UnrealEd/Classes/ThumbnailRendering/ThumbnailManager.h
+++ b/Engine/Source/Editor/UnrealEd/Classes/ThumbnailRendering/ThumbnailManager.h
@@ -31,6 +31,7 @@ enum EThumbnailPrimType : int
 	TPT_Cube,
 	TPT_Plane,
 	TPT_Cylinder,
+	TPT_ShaderBall,
 	TPT_MAX,
 };
 
@@ -112,6 +113,9 @@ public:
 
 	UPROPERTY(Transient)
 	TObjectPtr<class UStaticMesh> EditorPlane;
+	
+	UPROPERTY(Transient)
+	TObjectPtr<class UStaticMesh> EditorShaderBall;
 
 	UPROPERTY(Transient)
 	TObjectPtr<class UStaticMesh> EditorSkySphere;
diff --git a/Engine/Source/Editor/UnrealEd/Private/Preferences/MaterialEditorOptions.cpp b/Engine/Source/Editor/UnrealEd/Private/Preferences/MaterialEditorOptions.cpp
new file mode 100644
--- /dev/null
+++ b/Engine/Source/Editor/UnrealEd/Private/Preferences/MaterialEditorOptions.cpp
@@ -0,0 +1,36 @@
+// Copyright Epic Games, Inc. All Rights Reserved.
+
+#include "Preferences/MaterialEditorOptions.h"
+
+#include "Editor/UnrealEdEngine.h"
+#include "Engine/StaticMesh.h"
+#include "ThumbnailRendering/ThumbnailManager.h"
+#include "UnrealEdGlobals.h"
+
+UMaterialEditorOptions* UMaterialEditorOptions::Get()
+{
+	return GetMutableDefault<UMaterialEditorOptions>();
+}
+
+UObject* UMaterialEditorOptions::GetPreviewMesh() const
+{
+	UObject* PreviewMeshObject = PreviewMesh.TryLoad();
+
+	// Fallback to legacy preview static mesh : sphere
+	if (!PreviewMeshObject)
+	{
+		PreviewMeshObject = GUnrealEd->GetThumbnailManager()->EditorSphere;
+	}
+
+	return PreviewMeshObject;
+}
+
+void UMaterialEditorOptions::SetPreviewMesh(UObject* InPreviewMesh)
+{
+	if (!InPreviewMesh || PreviewMesh == InPreviewMesh)
+	{
+		return;
+	}
+
+	PreviewMesh = InPreviewMesh->GetPathName();
+}
diff --git a/Engine/Source/Editor/UnrealEd/Private/ThumbnailHelpers.cpp b/Engine/Source/Editor/UnrealEd/Private/ThumbnailHelpers.cpp
--- a/Engine/Source/Editor/UnrealEd/Private/ThumbnailHelpers.cpp
+++ b/Engine/Source/Editor/UnrealEd/Private/ThumbnailHelpers.cpp
@@ -383,6 +383,10 @@ void FMaterialThumbnailScene::SetMaterialInterface(UMaterialInterface* InMateria
 			Transform.SetRotation(FQuat(FRotator(0, -90, 0)));
 			PreviewActor->GetStaticMeshComponent()->SetStaticMesh(GUnrealEd->GetThumbnailManager()->EditorPlane);
 			break;
+			
+		case TPT_ShaderBall:
+			PreviewActor->GetStaticMeshComponent()->SetStaticMesh(GUnrealEd->GetThumbnailManager()->EditorShaderBall);
+			break;
 
 		default:
 			check(0);
diff --git a/Engine/Source/Editor/UnrealEd/Private/ThumbnailManager.cpp b/Engine/Source/Editor/UnrealEd/Private/ThumbnailManager.cpp
--- a/Engine/Source/Editor/UnrealEd/Private/ThumbnailManager.cpp
+++ b/Engine/Source/Editor/UnrealEd/Private/ThumbnailManager.cpp
@@ -42,6 +42,7 @@ UThumbnailManager::UThumbnailManager(const FObjectInitializer& ObjectInitializer
 			ConstructorHelpers::FObjectFinder<UStaticMesh> EditorSphereMesh;
 			ConstructorHelpers::FObjectFinder<UStaticMesh> EditorCylinderMesh;
 			ConstructorHelpers::FObjectFinder<UStaticMesh> EditorPlaneMesh;
+			ConstructorHelpers::FObjectFinder<UStaticMesh> EditorShaderBallMesh;
 			ConstructorHelpers::FObjectFinder<UStaticMesh> EditorSkySphereMesh;
 			ConstructorHelpers::FObjectFinder<UMaterial> FloorPlaneMaterial;
 			ConstructorHelpers::FObjectFinder<UTextureCube> DaylightAmbientCubemap;
@@ -51,6 +52,7 @@ UThumbnailManager::UThumbnailManager(const FObjectInitializer& ObjectInitializer
 				, EditorSphereMesh(TEXT("/Engine/EditorMeshes/EditorSphere"))
 				, EditorCylinderMesh(TEXT("/Engine/EditorMeshes/EditorCylinder"))
 				, EditorPlaneMesh(TEXT("/Engine/EditorMeshes/EditorPlane"))
+				, EditorShaderBallMesh(TEXT("/Engine/EditorMeshes/EditorShaderBall"))
 				, EditorSkySphereMesh(TEXT("/Engine/EditorMeshes/EditorSkySphere"))
 				, FloorPlaneMaterial(TEXT("/Engine/EditorMaterials/Thumbnails/FloorPlaneMaterial"))
 				, DaylightAmbientCubemap(TEXT("/Engine/MapTemplates/Sky/DaylightAmbientCubemap"))
@@ -63,6 +65,7 @@ UThumbnailManager::UThumbnailManager(const FObjectInitializer& ObjectInitializer
 		EditorSphere = ConstructorStatics.EditorSphereMesh.Object;
 		EditorCylinder = ConstructorStatics.EditorCylinderMesh.Object;
 		EditorPlane = ConstructorStatics.EditorPlaneMesh.Object;
+		EditorShaderBall = ConstructorStatics.EditorShaderBallMesh.Object;
 		EditorSkySphere = ConstructorStatics.EditorSkySphereMesh.Object;
 		FloorPlaneMaterial = ConstructorStatics.FloorPlaneMaterial.Object;
 		AmbientCubemap = ConstructorStatics.DaylightAmbientCubemap.Object;
