Add the possibility to gather Parameter Categories for sorting from used functions, not just base material.

To handle the increased number of elements, add customization for the Groups array.
This adds a Search Box for filtering and modifies the row to show the Group Name and Priority directly, without need to expand.

Convert the Group Sourting gathering to happen in a Set, to avoid having the same category multiple times, with different sorting values.
Add GetSortingGroup which returns sorting values for existing groups. This is used instead of the previous AttemptInsertNewGroupName, which was also doing an add operation, needlessly, as the array was replaced at the end.

Add GetTypeHash for TSet compatibility.
Replace == operator operation to include only the string. This was not used before and it now makes sure there's no duplicate sorting groups with different sorting values.
Sort the Groups array.

After profiling, this feature adds 20ms on opening some of the heaviest materials.
[RN] minor | Rendering Architecture

#Synced-CL 50837534

diff --git a/Engine/Source/Editor/MaterialEditor/Private/MaterialEditor.cpp b/Engine/Source/Editor/MaterialEditor/Private/MaterialEditor.cpp
--- a/Engine/Source/Editor/MaterialEditor/Private/MaterialEditor.cpp
+++ b/Engine/Source/Editor/MaterialEditor/Private/MaterialEditor.cpp
@@ -204,6 +204,12 @@ static TAutoConsoleVariable<int32> CVarMaterialEdMaxDerivedMaterialInstances(
 	-1,
 	TEXT("Limits amount of derived material instance shown in platform stats. Use negative number to disable the limit. Material editor must be re-opened if changed at runtime."));
  
+static TAutoConsoleVariable<bool> CVarMaterialEdGroupSortIncludeFunctions(
+	TEXT("r.MaterialEditor.GroupSortIncludeFunctions"),
+	true,
+	TEXT("Whether Groups inside used Material Functions should be included in the sorting data."),
+	ECVF_Default);
+
 static void InvalidateSubstrateConversionVersion(UMaterial* Material)
 {
 	// The material has been updated, invalid the AutoConversion version
@@ -538,7 +544,7 @@ void FMaterialEditor::InitEditorForMaterial(UMaterial* InMaterial)
 	check(Material->GetExpressions().Find(nullptr) == INDEX_NONE);
 
 	TArray<FString> Groups;
-	GetAllMaterialExpressionGroups(&Groups);
+	UpdateAndGetMaterialExpressionGroups(&Groups);
 }
 
 void FMaterialEditor::InitEditorForMaterialFunction(UMaterialFunction* InMaterialFunction)
@@ -569,7 +575,7 @@ void FMaterialEditor::InitEditorForMaterialFunction(UMaterialFunction* InMateria
 	OriginalMaterial = Material;
 
 	TArray<FString> Groups;
-	GetAllMaterialExpressionGroups(&Groups);
+	UpdateAndGetMaterialExpressionGroups(&Groups);
 }
 
 void FMaterialEditor::InitMaterialEditor( const EToolkitMode::Type Mode, const TSharedPtr< class IToolkitHost >& InitToolkitHost, UObject* ObjectToEdit )
@@ -984,7 +990,7 @@ FMaterialEditor::~FMaterialEditor()
 	MaterialEditorInstance = NULL;
 }
 
-void FMaterialEditor::GetAllMaterialExpressionGroups(TArray<FString>* OutGroups)
+void FMaterialEditor::UpdateAndGetMaterialExpressionGroups(TArray<FString>* OutGroups)
 {
 	UMaterialEditorOnlyData* EditorOnlyData = Material->GetEditorOnlyData();
 	if (!EditorOnlyData)
@@ -992,29 +998,74 @@ void FMaterialEditor::GetAllMaterialExpressionGroups(TArray<FString>* OutGroups)
 		return;
 	}
 
-	TArray<FParameterGroupData> UpdatedGroups;
-	for (const UMaterialExpression* MaterialExpression : Material->GetExpressions())
+	UpdateMaterialExpressionGroups();
+	for (const FParameterGroupData& GroupData : EditorOnlyData->ParameterGroupData)
+	{
+		OutGroups->Add(GroupData.GroupName);
+	}
+}
+
+void FMaterialEditor::UpdateMaterialExpressionGroups()
+{
+	TRACE_CPUPROFILER_EVENT_SCOPE(FMaterialInstanceEditor::UpdateMaterialExpressionGroups)
+
+	UMaterialEditorOnlyData* EditorOnlyData = Material->GetEditorOnlyData();
+	if (!EditorOnlyData)
+	{
+		return;
+	}
+
+	TArray<UMaterialExpression*> ExpressionsToParse = TArray<UMaterialExpression*>(Material->GetExpressions());
+
+	if (CVarMaterialEdGroupSortIncludeFunctions.GetValueOnGameThread())
+	{
+		TSet<UMaterialExpression*> FunctionsExpressions;
+
+		for (const UMaterialExpression* MaterialExpression : Material->GetExpressions())
+		{
+			if (const UMaterialExpressionMaterialFunctionCall* FunctionCall = Cast<UMaterialExpressionMaterialFunctionCall>(MaterialExpression))
+			{
+				if (FunctionCall->MaterialFunction)
+				{
+					TArray<UMaterialExpression*> FunctionExpressions;
+					FunctionCall->MaterialFunction->GetAllExpressionsOfType<UMaterialExpression>(FunctionExpressions);
+					FunctionsExpressions.Append(FunctionExpressions);
+				}
+			}
+		}
+
+		ExpressionsToParse.Append(FunctionsExpressions.Array());
+	}
+
+	TSet<FParameterGroupData> UpdatedGroups;
+	for (const UMaterialExpression* MaterialExpression : ExpressionsToParse)
 	{
 		FMaterialParameterMetadata ParameterMeta;
 		if (MaterialExpression->GetParameterValue(ParameterMeta))
 		{
+			if (ParameterMeta.Group.IsNone())
+			{
+				continue;
+			}
+
 			const FString GroupName = ParameterMeta.Group.ToString();
-			OutGroups->AddUnique(GroupName);
-			if (Material->AttemptInsertNewGroupName(GroupName))
+
+			if (const FParameterGroupData* GroupData = Material->GetSortingGroup(GroupName))
 			{
-				UpdatedGroups.Add(FParameterGroupData(GroupName, 0));
+				UpdatedGroups.Add(*GroupData);
 			}
 			else
 			{
-				FParameterGroupData* ParameterGroupDataElement = EditorOnlyData->ParameterGroupData.FindByPredicate([&GroupName](const FParameterGroupData& DataElement)
-				{
-					return GroupName == DataElement.GroupName;
-				});
-				UpdatedGroups.AddUnique(FParameterGroupData(GroupName, ParameterGroupDataElement->GroupSortPriority));
+				UpdatedGroups.Add(FParameterGroupData(GroupName, 0));
 			}
 		}
 	}
-	EditorOnlyData->ParameterGroupData = UpdatedGroups;
+
+	TArray<FParameterGroupData> UpdatedGroupsArray = UpdatedGroups.Array();
+	UpdatedGroupsArray.Sort(); // Sort groups based on name
+
+	// This indirectly gets rid of any groups which might have been removed
+	EditorOnlyData->ParameterGroupData = UpdatedGroupsArray;
 }
 
 void FMaterialEditor::UpdatePreviewViewportsVisibility()
@@ -1573,7 +1624,7 @@ void FMaterialEditor::CreateInternalWidgets()
 	MaterialDetailsView = PropertyEditorModule.CreateDetailView( DetailsViewArgs );
 
 	FOnGetDetailCustomizationInstance LayoutExpressionParameterDetails = FOnGetDetailCustomizationInstance::CreateStatic(
-		&FMaterialExpressionParameterDetails::MakeInstance, FOnCollectParameterGroups::CreateSP(this, &FMaterialEditor::GetAllMaterialExpressionGroups) );
+		&FMaterialExpressionParameterDetails::MakeInstance, FOnCollectParameterGroups::CreateSP(this, &FMaterialEditor::UpdateAndGetMaterialExpressionGroups) );
 
 	MaterialDetailsView->RegisterInstancedCustomPropertyLayout( 
 		UMaterialExpressionParameter::StaticClass(), 
@@ -1601,7 +1652,7 @@ void FMaterialEditor::CreateInternalWidgets()
 		);
 
 	FOnGetDetailCustomizationInstance LayoutLayerExpressionParameterDetails = FOnGetDetailCustomizationInstance::CreateStatic(
-		&FMaterialExpressionLayersParameterDetails::MakeInstance, FOnCollectParameterGroups::CreateSP(this, &FMaterialEditor::GetAllMaterialExpressionGroups));
+		&FMaterialExpressionLayersParameterDetails::MakeInstance, FOnCollectParameterGroups::CreateSP(this, &FMaterialEditor::UpdateAndGetMaterialExpressionGroups));
 
 	MaterialDetailsView->RegisterInstancedCustomPropertyLayout(
 		UMaterialExpressionMaterialAttributeLayers::StaticClass(),
@@ -8265,6 +8316,12 @@ void FMaterialEditor::OnSelectedNodesChanged(const TSet<class UObject*>& NewSele
 	GetDetailView()->SetObjects( SelectedObjects, true );
 	FocusDetailsPanel();
 	
+	if (SelectedObjects.Num() == 1 && SelectedObjects[0] == Material)
+	{
+		// Update when we display the material details, in case any MF parameter was modified
+		UpdateMaterialExpressionGroups();
+	}
+
 	if (bHideUnrelatedNodes)
 	{
 		if (!bLockNodeFadeState)
diff --git a/Engine/Source/Editor/MaterialEditor/Private/MaterialEditor.h b/Engine/Source/Editor/MaterialEditor/Private/MaterialEditor.h
--- a/Engine/Source/Editor/MaterialEditor/Private/MaterialEditor.h
+++ b/Engine/Source/Editor/MaterialEditor/Private/MaterialEditor.h
@@ -644,8 +644,11 @@ private:
 	/** Creates all internal widgets for the tabs to point at */
 	void CreateInternalWidgets();
 	
-	/** Collects all groups for all material expressions */
-	void GetAllMaterialExpressionGroups(TArray<FString>* OutGroups);
+	/** Updates and returns all groups for all material expressions */
+	void UpdateAndGetMaterialExpressionGroups(TArray<FString>* OutGroups);
+
+	/** Updates all groups for all material expressions*/
+	void UpdateMaterialExpressionGroups();
 
 	/** Updates the 3D and UI preview viewport visibility based on material domain */
 	void UpdatePreviewViewportsVisibility();
diff --git a/Engine/Source/Editor/MaterialEditor/Private/MaterialEditorDetailCustomization.cpp b/Engine/Source/Editor/MaterialEditor/Private/MaterialEditorDetailCustomization.cpp
--- a/Engine/Source/Editor/MaterialEditor/Private/MaterialEditorDetailCustomization.cpp
+++ b/Engine/Source/Editor/MaterialEditor/Private/MaterialEditorDetailCustomization.cpp
@@ -51,6 +51,13 @@
 #include "MaterialShared.h"
 #include "SColorGradientEditor.h"
 #include "SWarningOrErrorBox.h"
+#include "IDetailChildrenBuilder.h"
+
+static TAutoConsoleVariable<bool> CVarMaterialEdUseGroupCustomUI(
+	TEXT("r.MaterialEditor.UseGroupsCustomUI"),
+	true,
+	TEXT("Whether Custom UI should be used for Groups Priorities."),
+	ECVF_Default);
 
 
 #define LOCTEXT_NAMESPACE "MaterialEditor"
@@ -877,6 +884,112 @@ static void HandleNewMaterialTranslatorToggleProperty(IDetailLayoutBuilder& Deta
 		];
 }
 
+class FDetailGroupsPrioritiesBuilder : public FDetailArrayBuilder, public TSharedFromThis<FDetailGroupsPrioritiesBuilder>
+{
+public:
+
+	FDetailGroupsPrioritiesBuilder(TSharedRef<IPropertyHandle> InBaseProperty, bool InGenerateHeader = true)
+		: FDetailArrayBuilder(InBaseProperty, InGenerateHeader, false /* InDisplayResetToDefault */, false /* InDisplayElementNum */)
+	{}
+
+	virtual void GenerateHeaderRowContent(FDetailWidgetRow& NodeRow) override
+	{
+		if (GetShouldGenerateHeader())
+		{
+			TSharedPtr<SHorizontalBox> ContentHorizontalBox;
+			SAssignNew(ContentHorizontalBox, SHorizontalBox);
+
+			ContentHorizontalBox->AddSlot()
+			[
+				SNew(SSearchBox)
+				.OnTextChanged(this, &FDetailGroupsPrioritiesBuilder::OnFilterTextChanged)
+			];
+
+			NodeRow
+			.FilterString(LOCTEXT("ParameterGroupsName", "Parameter Groups Priorities"))
+			.NameContent()
+			[
+				GetPropertyHandle()->CreatePropertyNameWidget(LOCTEXT("ParameterGroupsName", "Parameter Groups Priorities"), FText::GetEmpty())
+			]
+			.ValueContent()
+			[
+				ContentHorizontalBox.ToSharedRef()
+			];
+		}
+	}
+
+	void OnFilterTextChanged(const FText& InFilterText)
+	{
+		CurrentFilterText = InFilterText.ToString();
+		RefreshChildren();
+	}
+
+	virtual void GenerateChildContent(IDetailChildrenBuilder& ChildrenBuilder) override
+	{
+		TSharedPtr<IPropertyHandleArray> LocalArrayProperty = GetPropertyHandle()->AsArray();
+
+		uint32 NumChildren = 0;
+		LocalArrayProperty->GetNumElements(NumChildren);
+
+		for (uint32 ChildIndex = 0; ChildIndex < NumChildren; ++ChildIndex)
+		{
+			TSharedRef<IPropertyHandle> ElementHandle = LocalArrayProperty->GetElement(ChildIndex);
+
+			GenerateSortingGroupsWidget(ElementHandle, ChildIndex, ChildrenBuilder);
+		}
+	}
+
+	void GenerateSortingGroupsWidget(TSharedRef<IPropertyHandle> PropertyHandle, int32 ArrayIndex, IDetailChildrenBuilder& ChildrenBuilder)
+	{
+		const TSharedPtr<IPropertyHandle> GroupNameProp = PropertyHandle->GetChildHandle(GET_MEMBER_NAME_CHECKED(FParameterGroupData, GroupName));
+
+		FString PropDisplayValue;
+		GroupNameProp->GetValueAsDisplayString(PropDisplayValue);
+
+		if (!CurrentFilterText.IsEmpty())
+		{
+			if (!PropDisplayValue.Contains(CurrentFilterText))
+			{
+				return;
+			}
+		}
+
+		IDetailPropertyRow& PropertyRow = ChildrenBuilder.AddProperty(PropertyHandle);
+
+		const TSharedPtr<IPropertyHandle> GroupPrioProp = PropertyHandle->GetChildHandle(GET_MEMBER_NAME_CHECKED(FParameterGroupData, GroupSortPriority));
+
+		FUIAction CopyAction;
+		FUIAction PasteAction;
+		GroupPrioProp->CreateDefaultPropertyCopyPasteActions(CopyAction, PasteAction);
+
+		PropertyRow.CustomWidget(false)
+			.OverrideResetToDefault(FResetToDefaultOverride::Create(
+			true,
+			FSimpleDelegate::CreateLambda([GroupPrioProp]()
+			{
+				if (GroupPrioProp)
+				{
+					GroupPrioProp->ResetToDefault();
+				}
+			})))
+			.NameContent()
+			[
+				SNew(STextBlock)
+					.Text(FText::FromString(PropDisplayValue))
+			]
+			.ValueContent()
+			[
+				GroupPrioProp->CreatePropertyValueWidget()
+			]
+			.CopyAction(CopyAction)
+			.PasteAction(PasteAction);
+	}
+
+	private:
+	FString CurrentFilterText;
+	TSharedPtr<SSearchBox> GroupSearchBox;
+};
+
 void FMaterialDetailCustomization::CustomizeDetails( IDetailLayoutBuilder& DetailLayout )
 {	
 	TArray<TWeakObjectPtr<UObject> > Objects;
@@ -1046,6 +1159,17 @@ void FMaterialDetailCustomization::CustomizeDetails( IDetailLayoutBuilder& Detai
 			}
 		}
 	}
+
+	if (CVarMaterialEdUseGroupCustomUI.GetValueOnGameThread())
+	{
+		IDetailCategoryBuilder& GroupSortingCategory = DetailLayout.EditCategory("Group Sorting");
+
+		TSharedPtr<IPropertyHandle> SortingGroupDataProperty = DetailLayout.GetProperty(GET_MEMBER_NAME_CHECKED(UMaterialEditorOnlyData, ParameterGroupData), UMaterialEditorOnlyData::StaticClass());
+		check(SortingGroupDataProperty);
+
+		TSharedRef<FDetailGroupsPrioritiesBuilder> SortingGroupDataBuilder = MakeShared<FDetailGroupsPrioritiesBuilder>(SortingGroupDataProperty.ToSharedRef());
+		GroupSortingCategory.AddCustomBuilder(SortingGroupDataBuilder);
+	}
 }
 
 TSharedRef<class IDetailCustomization> FMaterialFunctionDetailCustomization::MakeInstance()
diff --git a/Engine/Source/Runtime/Engine/Private/Materials/Material.cpp b/Engine/Source/Runtime/Engine/Private/Materials/Material.cpp
--- a/Engine/Source/Runtime/Engine/Private/Materials/Material.cpp
+++ b/Engine/Source/Runtime/Engine/Private/Materials/Material.cpp
@@ -2915,17 +2915,25 @@ void UMaterial::ReleaseResourcesAndMutateDDCKey(const FGuid& TransformationId)
 	}
 }
 
-bool UMaterial::AttemptInsertNewGroupName(const FString & InNewName)
+const FParameterGroupData* UMaterial::GetSortingGroup(const FString& InGroupName)
 {
 #if WITH_EDITOR
 	UMaterialEditorOnlyData* LocalData = GetEditorOnlyData();
 	if (LocalData)
 	{
-		FParameterGroupData* ParameterGroupDataElement = LocalData->ParameterGroupData.FindByPredicate([&InNewName](const FParameterGroupData& DataElement)
-		{
-			return InNewName == DataElement.GroupName;
-		});
+		return LocalData->ParameterGroupData.FindByKey(InGroupName);
+	}
+#endif
+	return nullptr;
+}
 
+bool UMaterial::AttemptInsertNewGroupName(const FString& InNewName)
+{
+#if WITH_EDITOR
+	UMaterialEditorOnlyData* LocalData = GetEditorOnlyData();
+	if (LocalData)
+	{
+		FParameterGroupData* ParameterGroupDataElement = LocalData->ParameterGroupData.FindByKey(InNewName);
 		if (ParameterGroupDataElement == nullptr)
 		{
 			FParameterGroupData NewGroupData;
diff --git a/Engine/Source/Runtime/Engine/Public/Materials/Material.h b/Engine/Source/Runtime/Engine/Public/Materials/Material.h
--- a/Engine/Source/Runtime/Engine/Public/Materials/Material.h
+++ b/Engine/Source/Runtime/Engine/Public/Materials/Material.h
@@ -293,15 +293,31 @@ struct FParameterGroupData
 		GroupSortPriority = 0;
 	}
 
+	FParameterGroupData(const FString& InString)
+		: GroupName(InString),
+		GroupSortPriority(0)
+	{
+	}
+
 	FParameterGroupData(const FString& InString, int32 InSortPriority)
 	:	GroupName(InString),
 		GroupSortPriority(InSortPriority)
 	{
 	}
 
+	bool operator<(const FParameterGroupData& Other) const
+	{
+		return GroupName < Other.GroupName;
+	}
+
+	friend uint32 GetTypeHash(const FParameterGroupData& InGroupData)
+	{
+		return GetTypeHash(InGroupData.GroupName);
+	}
+
 	bool operator==(const FParameterGroupData& Other) const
 	{
-		return GroupName == Other.GroupName && GroupSortPriority == Other.GroupSortPriority;
+		return GroupName == Other.GroupName;
 	}
 };
 
@@ -1681,6 +1697,9 @@ public:
 	ENGINE_API void UpdateCachedExpressionData();
 #endif
 
+	/** Try to get an existing sorting group with the given name. Nullptr if none is found. */
+	ENGINE_API const FParameterGroupData* GetSortingGroup(const FString& InGroupName);
+
 	/** Attempts to add a new group name to the Group Data struct. True if new name was added. */
 	ENGINE_API bool AttemptInsertNewGroupName(const FString& InNewName);
 
