Fix crash happening when attempting to use a Custom Preview Mesh which has Nanite data, in a Material Function Editor when no Material Expressions are selected for previewing.
The crash was happening because when ShouldCreateNaniteProxy() attempts to set Nanite usage flag on the preview material, the Material Editor(which is used for Material Function editing) thinks that the main material which is being edited has had a Usage Flag modified by a user manually. This happens because in the Material Editor, when editing an MF, the Material used for Preview and the OriginalMaterial are the same, given that there is no actual base material.

This leads to it attempting to update the preview material, which should not happen in this call at all when the preview material is modified, but especially in a call coming from EndOfFrame, which causes a crash.
The fix is simply to take into consideration when the Material Editor receives a UsageFlag modify callback, that we are editing a Material Function.

[RN] minor | Rendering Architecture
#Synced-CL 51008691

diff --git a/Engine/Source/Editor/MaterialEditor/Private/MaterialEditor.cpp b/Engine/Source/Editor/MaterialEditor/Private/MaterialEditor.cpp
--- a/Engine/Source/Editor/MaterialEditor/Private/MaterialEditor.cpp
+++ b/Engine/Source/Editor/MaterialEditor/Private/MaterialEditor.cpp
@@ -5342,7 +5342,11 @@ void FMaterialEditor::RenameAssetFromRegistry(const FAssetData& InAddedAssetData
 void FMaterialEditor::OnMaterialUsageFlagsChanged(UMaterial* MaterialThatChanged, int32 FlagThatChanged)
 {
 	EMaterialUsage Flag = static_cast<EMaterialUsage>(FlagThatChanged);
-	if(MaterialThatChanged == OriginalMaterial)
+
+	// Update only if a modification has been done to the original material by the user and guard when Editor
+	// is opened for MF, as we don't have an original material there, thus OriginalMaterial == Material 
+	// and Material can be modified and received here, depending on which Preview Mesh is selected
+	if((MaterialThatChanged == OriginalMaterial) && !MaterialFunction)
 	{
 		Material->SetUsageByFlag(Flag, OriginalMaterial->GetUsageByFlag(Flag));
 		UpdatePreviewMaterial();
