Modify the material diff tool to work on Material Functions.

- The code in SMaterialRevisionMenu::MakeDiffMenu explicitly only looked for objects that inherit from UMaterialInterface which won't include Material Functions
- I've just changed the logic to consider any valid UMaterialInterface or UMaterialFunctions that are not transient.

#Synced-CL 51020455

diff --git a/Engine/Source/Editor/MaterialEditor/Private/DiffTool/Widgets/SMaterialRevisionMenu.cpp b/Engine/Source/Editor/MaterialEditor/Private/DiffTool/Widgets/SMaterialRevisionMenu.cpp
--- a/Engine/Source/Editor/MaterialEditor/Private/DiffTool/Widgets/SMaterialRevisionMenu.cpp
+++ b/Engine/Source/Editor/MaterialEditor/Private/DiffTool/Widgets/SMaterialRevisionMenu.cpp
@@ -182,26 +182,22 @@ void SMaterialRevisionMenu::MakeDiffMenu(UToolMenu* Menu)
 					TSharedPtr<IMaterialEditor> MaterialEditorStrong = MaterialEditorWeak.Pin();
 					if (MaterialEditorStrong && ISourceControlModule::Get().IsEnabled() && ISourceControlModule::Get().GetProvider().IsAvailable())
 					{
-						UMaterialInterface* MaterialObj = MaterialEditorStrong->GetMaterialInterface();
-						if( MaterialObj == nullptr || 
-							(MaterialObj->HasAnyFlags(RF_Transient) || MaterialObj->GetPackage() == GetTransientPackage()))
+						const TArray<UObject*>* ObjectsBeingEdited = MaterialEditorStrong->GetObjectsCurrentlyBeingEdited();
+						UObject* MaterialObj = nullptr;
+						if (ObjectsBeingEdited)
 						{
-							const TArray< UObject* >* ObjectsBeingEdited = MaterialEditorStrong->GetObjectsCurrentlyBeingEdited();
-							MaterialObj = nullptr;
-							if(ObjectsBeingEdited)
+							for (UObject* Obj : *ObjectsBeingEdited)
 							{
-								for(UObject* Obj: *ObjectsBeingEdited)
+								if (Obj &&
+									!(Obj->HasAnyFlags(RF_Transient) || Obj->GetPackage() == GetTransientPackage()) &&
+									(Obj->IsA<UMaterial>() || Obj->IsA<UMaterialFunction>()))
 								{
-									UMaterialInterface* PotentialMaterialObj = Cast<UMaterialInterface>(Obj);
-									if(PotentialMaterialObj && 
-										!(PotentialMaterialObj->HasAnyFlags(RF_Transient) || PotentialMaterialObj->GetPackage() == GetTransientPackage()))
-									{
-										MaterialObj = PotentialMaterialObj;
-										break;
-									}
+									MaterialObj = Obj;
+									break;
 								}
 							}
 						}
+
 						if (MaterialObj)
 						{
 							TWeakObjectPtr<UObject> WeakMaterialObj = MaterialObj;
