EditorUsability : MaterialEditor
- Added drag and drop support from content browser to material viewport to change preview mesh
- Dropped preview mesh will become default preview mesh when reopening asset editor

#Synced-CL 51208877

diff --git a/Engine/Source/Editor/MaterialEditor/MaterialEditor.Build.cs b/Engine/Source/Editor/MaterialEditor/MaterialEditor.Build.cs
--- a/Engine/Source/Editor/MaterialEditor/MaterialEditor.Build.cs
+++ b/Engine/Source/Editor/MaterialEditor/MaterialEditor.Build.cs
@@ -25,6 +25,7 @@ public class MaterialEditor : ModuleRules
 			new string[] {
                 "AppFramework",
 				"ContentBrowser",
+				"ContentBrowserData",
 				"Core",
 				"CoreUObject",
 				"ApplicationCore",
diff --git a/Engine/Source/Editor/MaterialEditor/Private/SMaterialEditorViewport.cpp b/Engine/Source/Editor/MaterialEditor/Private/SMaterialEditorViewport.cpp
--- a/Engine/Source/Editor/MaterialEditor/Private/SMaterialEditorViewport.cpp
+++ b/Engine/Source/Editor/MaterialEditor/Private/SMaterialEditorViewport.cpp
@@ -49,6 +49,7 @@
 
 #define LOCTEXT_NAMESPACE "MaterialEditor"
 #include "AdvancedPreviewSceneMenus.h"
+#include "ContentBrowserDataDragDropOp.h"
 #include "MaterialEditorTabs.h"
 #include "PreviewProfileController.h"
 #include "SMaterialEditorTopologyWidget.h"
@@ -799,6 +800,47 @@ void SMaterialEditor3DPreviewViewport::OnFocusViewportToSelection()
 	}
 }
 
+FReply SMaterialEditor3DPreviewViewport::OnDrop(const FGeometry& InGeometry, const FDragDropEvent& InEvent)
+{
+	if (const TSharedPtr<FContentBrowserDataDragDropOp> ContentBrowserAsset = InEvent.GetOperationAs<FContentBrowserDataDragDropOp>())
+	{
+		if (!ContentBrowserAsset->GetDraggedFiles().IsEmpty())
+		{
+			for (const FContentBrowserItem& Item : ContentBrowserAsset->GetDraggedFiles())
+			{
+				FAssetData AssetData;
+				if (Item.Legacy_TryGetAssetData(AssetData))
+				{
+					// Only load class first to validate
+					if (const UClass* Class = AssetData.GetClass(EResolveClass::Yes))
+					{
+						if (Class->IsChildOf<UStaticMesh>() || Class->IsChildOf<USkeletalMesh>())
+						{
+							// Load asset now
+							if (UObject* Asset = AssetData.GetAsset())
+							{
+								SetPreviewAsset(Asset);
+								
+								// Save the preview mesh as default
+								if (UMaterialEditorOptions* MaterialEditorOptions = UMaterialEditorOptions::Get())
+								{
+									MaterialEditorOptions->SetPreviewMesh(Asset);
+								}
+
+								// Notify the change
+								RefreshViewport();
+
+								return FReply::Handled();
+							}
+						}
+					}
+				}
+			}
+		}
+	}
+	return FReply::Unhandled();
+}
+
 UMaterialInterface* SMaterialEditor3DPreviewViewport::GetPreviewingMaterialInterface()
 {
 	return MaterialEditorPtr.IsValid() ? MaterialEditorPtr.Pin()->GetMaterialInterface() : PreviewMaterialPtr.Get();
diff --git a/Engine/Source/Editor/MaterialEditor/Private/SMaterialEditorViewport.h b/Engine/Source/Editor/MaterialEditor/Private/SMaterialEditorViewport.h
--- a/Engine/Source/Editor/MaterialEditor/Private/SMaterialEditorViewport.h
+++ b/Engine/Source/Editor/MaterialEditor/Private/SMaterialEditorViewport.h
@@ -101,6 +101,11 @@ protected:
 	virtual EVisibility OnGetViewportContentVisibility() const override;
 	virtual void BindCommands() override;
 	virtual void OnFocusViewportToSelection() override;
+	
+	//~ Begin SWidget
+	virtual FReply OnDrop(const FGeometry& InGeometry, const FDragDropEvent& InEvent) override;
+	//~ End SWidget
+	
 private:
 	/** If we have a MaterialEditor, it will look up the preview material from that. Otherwise it will fallback to previewing this material */
 	UMaterialInterface* GetPreviewingMaterialInterface();
