EditorUsability : MaterialEditor
- Removing unused private material instance code

#Synced-CL 51539177

diff --git a/Engine/Source/Editor/MaterialEditor/Private/MaterialEditorTabs.cpp b/Engine/Source/Editor/MaterialEditor/Private/MaterialEditorTabs.cpp
--- a/Engine/Source/Editor/MaterialEditor/Private/MaterialEditorTabs.cpp
+++ b/Engine/Source/Editor/MaterialEditor/Private/MaterialEditorTabs.cpp
@@ -22,5 +22,3 @@ const FName FMaterialInstanceEditorTabs::PreviewTabId( TEXT( "MaterialInstanceEd
 const FName FMaterialInstanceEditorTabs::PropertiesTabId( TEXT( "MaterialInstanceEditor_MaterialProperties" ) );
 const FName FMaterialInstanceEditorTabs::LayerPropertiesTabId(TEXT("MaterialInstanceEditor_MaterialLayerProperties"));
 const FName FMaterialInstanceEditorTabs::PreviewSettingsTabId(TEXT("MaterialInstanceEditor_PreviewSettings"));
-const FName FMaterialInstanceEditorTabs::AssetBrowserTabId(TEXT("MaterialInstanceEditor_AssetBrowser"));
-const FName FMaterialInstanceEditorTabs::SubstrateTabId(TEXT("MaterialInstanceEditor_Substrate"));
diff --git a/Engine/Source/Editor/MaterialEditor/Private/MaterialEditorTabs.h b/Engine/Source/Editor/MaterialEditor/Private/MaterialEditorTabs.h
--- a/Engine/Source/Editor/MaterialEditor/Private/MaterialEditorTabs.h
+++ b/Engine/Source/Editor/MaterialEditor/Private/MaterialEditorTabs.h
@@ -31,8 +31,6 @@ struct FMaterialInstanceEditorTabs
 	static const FName PropertiesTabId;
 	static const FName LayerPropertiesTabId;
 	static const FName PreviewSettingsTabId;
-	static const FName AssetBrowserTabId;
-	static const FName SubstrateTabId;
 
 	// Disable default constructor
 	FMaterialInstanceEditorTabs() = delete;
diff --git a/Engine/Source/Editor/MaterialEditor/Private/MaterialInstanceEditor.cpp b/Engine/Source/Editor/MaterialEditor/Private/MaterialInstanceEditor.cpp
--- a/Engine/Source/Editor/MaterialEditor/Private/MaterialInstanceEditor.cpp
+++ b/Engine/Source/Editor/MaterialEditor/Private/MaterialInstanceEditor.cpp
@@ -1134,100 +1134,6 @@ TSharedRef<SDockTab> FMaterialInstanceEditor::SpawnTab_PreviewSettings(const FSp
 	return SpawnedTab;
 }
 
-TSharedRef<SDockTab> FMaterialInstanceEditor::SpawnTab_AssetBrowser(const FSpawnTabArgs& Args)
-{
-	check(Args.GetTabId() == FMaterialInstanceEditorTabs::AssetBrowserTabId);
-
-	// TSharedRef<SWidget> InWidget = SNullWidget::NullWidget;
-	FContentBrowserModule& ContentBrowserModule = FModuleManager::Get().LoadModuleChecked<FContentBrowserModule>(TEXT("ContentBrowser"));
-
-	// Configure filter for asset picker
-	FAssetPickerConfig Config;
-	Config.SelectionMode = ESelectionMode::Single;
-	Config.Filter.ClassPaths.Add(UMaterialFunctionMaterialLayer::StaticClass()->GetClassPathName());
-	Config.Filter.ClassPaths.Add(UMaterialFunctionMaterialLayerInstance::StaticClass()->GetClassPathName());
-	Config.Filter.ClassPaths.Add(UMaterialFunctionMaterialLayerBlend::StaticClass()->GetClassPathName());
-	Config.Filter.ClassPaths.Add(UMaterialFunctionMaterialLayerBlendInstance::StaticClass()->GetClassPathName());
-	Config.bAddFilterUI = true;
-	Config.ThumbnailScale = 0.4f;
-	Config.InitialThumbnailSize = EThumbnailSize::Small;
-	Config.InitialAssetViewType = EAssetViewType::Tile;
-	Config.OnAssetDoubleClicked = FOnAssetDoubleClicked::CreateSP(this, &FMaterialInstanceEditor::OnAssetDoubleClicked);
-	Config.OnGetAssetContextMenu = FOnGetAssetContextMenu::CreateSP(this, &FMaterialInstanceEditor::OnGetAssetContextMenu);
-	Config.bForceShowEngineContent = true;
-	TSharedRef<SDockTab> SpawnedTab = SNew(SDockTab)
-		.Label(LOCTEXT("AssetBrowserTab", "Asset Browser"))
-		[
-			SNew(SBox)
-			[
-				ContentBrowserModule.Get().CreateAssetPicker(Config)
-			]
-		];
-
-	return SpawnedTab;
-}
-TSharedRef<SDockTab> FMaterialInstanceEditor::SpawnTab_Substrate(const FSpawnTabArgs& Args)
-{
-	check(Args.GetTabId() == FMaterialInstanceEditorTabs::SubstrateTabId);
-
-	TSharedRef<SDockTab> SubstrateTab = SNew(SDockTab)
-		.Label(LOCTEXT("MaterialSubstrateTabTitle", "Substrate"))
-		[
-			SNew(SBox)
-			[
-				SubstrateWidget.ToSharedRef()
-			]
-		];
-
-	return SubstrateTab;
-}
-TSharedPtr<SWidget> FMaterialInstanceEditor::OnGetAssetContextMenu(const TArray<FAssetData>& SelectedAssets) const
-{
-	if (SelectedAssets.Num() <= 0)
-	{
-		return nullptr;
-	}
-
-	UObject* SelectedAsset = SelectedAssets[0].GetAsset();
-	if (SelectedAsset == nullptr)
-	{
-		return nullptr;
-	}
-	
-	FMenuBuilder MenuBuilder(true, MakeShared<FUICommandList>());
-
-	MenuBuilder.BeginSection(TEXT("Asset"), LOCTEXT("AssetSectionLabel", "Asset"));
-	{
-		MenuBuilder.AddMenuEntry(
-			LOCTEXT("Browse", "Browse to Asset"),
-			LOCTEXT("BrowseTooltip", "Browses to the associated asset and selects it in the most recently used Content Browser (summoning one if necessary)"),
-			FSlateIcon(FAppStyle::GetAppStyleSetName(), "SystemWideCommands.FindInContentBrowser.Small"),
-			FUIAction(
-				FExecuteAction::CreateLambda([SelectedAsset] ()
-				{
-					if (SelectedAsset)
-					{
-						const TArray<FAssetData>& Assets = { SelectedAsset };
-						const FContentBrowserModule& ContentBrowserModule = FModuleManager::Get().LoadModuleChecked<FContentBrowserModule>("ContentBrowser");
-						ContentBrowserModule.Get().SyncBrowserToAssets(Assets);
-					}
-				}),
-				FCanExecuteAction::CreateLambda([] () { return true; })
-			)
-		);
-	}
-	MenuBuilder.EndSection();
-
-	return MenuBuilder.MakeWidget();
-}
-
-void FMaterialInstanceEditor::OnAssetDoubleClicked(const FAssetData& AssetData)
-{
-	if (UAssetEditorSubsystem* EditorSubsystem = GEditor->GetEditorSubsystem<UAssetEditorSubsystem>())
-	{
-		EditorSubsystem->OpenEditorForAsset(AssetData.ToSoftObjectPath());
-	}
-}
 void FMaterialInstanceEditor::AddToSpawnedToolPanels(const FName& TabIdentifier, const TSharedRef<SDockTab>& SpawnedTab)
 {
 	TWeakPtr<SDockTab>* TabSpot = SpawnedToolPanels.Find(TabIdentifier);
diff --git a/Engine/Source/Editor/MaterialEditor/Private/MaterialInstanceEditor.h b/Engine/Source/Editor/MaterialEditor/Private/MaterialInstanceEditor.h
--- a/Engine/Source/Editor/MaterialEditor/Private/MaterialInstanceEditor.h
+++ b/Engine/Source/Editor/MaterialEditor/Private/MaterialInstanceEditor.h
@@ -190,21 +190,9 @@ private:
 	/**	Spawns the properties tab */
 	TSharedRef<SDockTab> SpawnTab_LayerProperties(const FSpawnTabArgs& Args);
 
-	/**	Spawns the parents tab */
-	TSharedRef<SDockTab> SpawnTab_Parents( const FSpawnTabArgs& Args );
-
 	/** Spawns the advanced preview settings tab */
 	TSharedRef<SDockTab> SpawnTab_PreviewSettings(const FSpawnTabArgs& Args);
 	
-	/** Spawns the asset browser tab */
-	TSharedRef<SDockTab> SpawnTab_AssetBrowser(const FSpawnTabArgs& Args);
-
-	/** Spawns the substrate details tab */
-	TSharedRef<SDockTab> SpawnTab_Substrate(const FSpawnTabArgs& Args);
-	
-	void OnAssetDoubleClicked(const FAssetData& AssetData);
-	TSharedPtr<SWidget> OnGetAssetContextMenu(const TArray<FAssetData>& SelectedAssets) const;
-	
 	/**	Caches the specified tab for later retrieval */
 	void AddToSpawnedToolPanels( const FName& TabIdentifier, const TSharedRef<SDockTab>& SpawnedTab );
 
