Material instance UI fix.
Override usage flags now respects when the parent material has disabled the creation of material pemutations.
#rnx
#Synced-CL 51537754

diff --git a/Engine/Source/Editor/MaterialEditor/Private/MaterialEditorInstanceDetailCustomization.cpp b/Engine/Source/Editor/MaterialEditor/Private/MaterialEditorInstanceDetailCustomization.cpp
--- a/Engine/Source/Editor/MaterialEditor/Private/MaterialEditorInstanceDetailCustomization.cpp
+++ b/Engine/Source/Editor/MaterialEditor/Private/MaterialEditorInstanceDetailCustomization.cpp
@@ -1106,7 +1106,7 @@ void FMaterialInstanceParameterDetails::CreateBasePropertyOverrideWidgets(IDetai
 	TSharedRef<IPropertyHandle> BasePropertyOverrideProperty = DetailLayout.GetProperty("BasePropertyOverrides");
 	FMaterialInstanceBasePropertyOverrides& BaseOverrides = MaterialEditorInstance->BasePropertyOverrides;
 	TObjectPtr<UMaterialInterface> ParentMat = MaterialEditorInstance->Parent;
-	const FText ParameterDisabledToolTipString = FText::FromString(TEXT("This material instance parent restricts the creation of new shader permutations. Overriding this parameter would result in the generation of additional shader permutations."));
+	const FText ParameterDisabledToolTipString = LOCTEXT("DisableParameterOverride", "This material instance parent restricts the creation of new shader permutations. Overriding this parameter would result in the generation of additional shader permutations.");
 	const bool bStaticParametersOverrideDisabled = MaterialEditorInstance->SourceInstance->bDisallowStaticParameterPermutations;
 
 	auto CreateBaseOverrideRow = [&] (
@@ -1322,14 +1322,19 @@ void FMaterialInstanceParameterDetails::CreateUsageFlagOverrideWidgets(IDetailLa
 	TSharedPtr<IPropertyHandle> OverrideProperty = ParentProperty->GetChildHandle(GET_MEMBER_NAME_CHECKED(FMaterialInstanceBasePropertyOverrides, bOverride_UsageFlags));
 	TSharedPtr<IPropertyHandle> UsageProperty = ParentProperty->GetChildHandle(GET_MEMBER_NAME_CHECKED(FMaterialInstanceBasePropertyOverrides, UsageFlags));
 	UEnum* UsageEnum = StaticEnum<EMaterialUsage>();
+	
+	// Disable overrides for case where parent material instance is locked.
+	const bool bIsWidgetDisabled = MaterialEditorInstance->SourceInstance->bDisallowStaticParameterPermutations;
+	const FText DisabledToolTipText = LOCTEXT("DisableUsageOverride", "This material instance parent restricts the creation of new shader permutations. Overriding this usage flag would result in the generation of additional shader permutations.");
 
 	for (int32 UsageIndex = 0; UsageIndex < MATUSAGE_MAX; UsageIndex++)
 	{
 		const FText UsageName = FText::FromString(UMaterialInterface::GetUsageName((EMaterialUsage)UsageIndex));
-		const FText UsageToolTip = FText::FromString(UsageEnum->GetMetaData(TEXT("ToolTip"), UsageIndex));
+		const FText UsageToolTip = bIsWidgetDisabled ? DisabledToolTipText : FText::FromString(UsageEnum->GetMetaData(TEXT("ToolTip"), UsageIndex));
 
 		Group.AddWidgetRow()
 		.PropertyHandleList({UsageProperty})
+		.IsEnabled(!bIsWidgetDisabled)
 		.NameContent()
 		[
 			SNew(STextBlock)
@@ -1369,8 +1374,14 @@ void FMaterialInstanceParameterDetails::CreateUsageFlagOverrideWidgets(IDetailLa
 				OverrideProperty->GetValue(Value);
 				return (Value & (1u << UsageIndex)) != 0;
 			}),
-			FOnBooleanValueChanged::CreateLambda([OverrideProperty, UsageIndex](bool bNewValue)
+			FOnBooleanValueChanged::CreateLambda([OverrideProperty, UsageIndex, bIsWidgetDisabled](bool bNewValue)
 			{
+				// Edit condition doesn't have a disabled flag and doesn't respect the row setting. So ignore changes when we are disabled.
+				if (bIsWidgetDisabled)
+				{
+					return;
+				}
+
 				uint32 Value = 0;
 				OverrideProperty->GetValue(Value);
 				
