Fix performance issues cause by [REDACTED_PERSON] validation happening every tick.

Static switches trigger the UI to be recreated so trigger the validation once, when the UI is built.
Add FMaterialInstanceBasePropertyOverrides as a Property Class which triggers UI to be recreated, so Usage Flags behave similarly to static switches.

[RN] minor | Rendering Architecture
#Synced-CL 51584014

diff --git a/Engine/Source/Editor/MaterialEditor/Private/MaterialEditorInstanceDetailCustomization.cpp b/Engine/Source/Editor/MaterialEditor/Private/MaterialEditorInstanceDetailCustomization.cpp
--- a/Engine/Source/Editor/MaterialEditor/Private/MaterialEditorInstanceDetailCustomization.cpp
+++ b/Engine/Source/Editor/MaterialEditor/Private/MaterialEditorInstanceDetailCustomization.cpp
@@ -60,7 +60,6 @@ FMaterialInstanceParameterDetails::FMaterialInstanceParameterDetails(UMaterialEd
 	: MaterialEditorInstance(MaterialInstance)
 	, MaterialLayersFunctionsInstance(InMaterialLayersFunctionsInstance) 
 	, ShowHiddenDelegate(InShowHiddenDelegate)
-	, ValidationContext(new FMaterialEditorValidatorContext)
 {
 }
 
@@ -87,38 +86,23 @@ void FMaterialInstanceParameterDetails::CustomizeDetails(IDetailLayoutBuilder& D
 	FName GroupsCategoryName = TEXT("ParameterGroups");
 	IDetailCategoryBuilder& GroupsCategory = DetailLayout.EditCategory(GroupsCategoryName, LOCTEXT("MICParamGroupsTitle", "Parameter Groups"));
 	
-	TWeakObjectPtr<UMaterialInterface> MaterialWeakPtr(MaterialEditorInstance->GetMaterialInterface());
-	TWeakPtr<FMaterialEditorValidatorContext> ValidationContextWeakPtr = ValidationContext;
-
-	// Insert optional validation warning row at the top of the category.
-	GroupsCategory.AddCustomRow(LOCTEXT("ValidationWarningRow", "Validation Warning"))
-	.WholeRowContent()
-	[
-		SNew(SWarningOrErrorBox)
-		.Visibility_Lambda([MaterialWeakPtr, ValidationContextWeakPtr]() -> EVisibility
-		{
-			bool bWarn = false;
-			if (MaterialWeakPtr.IsValid())
-			{
-				if (UMaterialEditorValidationSubsystem* Subsystem = GEditor->GetEditorSubsystem<UMaterialEditorValidationSubsystem>())
-				{
-					if (TSharedPtr<FMaterialEditorValidatorContext> ValidationContext = ValidationContextWeakPtr.Pin())
-					{
-						bWarn = !Subsystem->Validate(MaterialWeakPtr.Get(), *ValidationContext.Get());
-					}
-				}
-			}
-			return bWarn ? EVisibility::Visible : EVisibility::Collapsed;
-		})
-		.Message_Lambda([ValidationContextWeakPtr]() -> FText
+	FMaterialEditorValidatorContext ValidationContext;
+	// Properties are updated when a static switch is modified so need to run once per update
+	if (UMaterialEditorValidationSubsystem* Subsystem = GEditor->GetEditorSubsystem<UMaterialEditorValidationSubsystem>())
+	{
+		const bool bWarn = !Subsystem->Validate(MaterialEditorInstance->GetMaterialInterface(), ValidationContext);
+		if (bWarn)
 		{
-			if (TSharedPtr<FMaterialEditorValidatorContext> ValidationContext = ValidationContextWeakPtr.Pin())
-			{
-				return ValidationContext->Message;
-			}
-			return {};
-		})
-	];
+			// Insert optional validation warning row at the top of the category.
+			GroupsCategory.AddCustomRow(LOCTEXT("ValidationWarningRow", "Validation Warning"))
+			.WholeRowContent()
+			[
+				SNew(SWarningOrErrorBox)
+				.Visibility(EVisibility::Visible)
+				.Message(ValidationContext.Message)
+			];
+		}
+	}
 	
 	TSharedRef<IPropertyHandle> ParameterGroupsProperty = DetailLayout.GetProperty("ParameterGroups");
 	CreateGroupsWidget(ParameterGroupsProperty, GroupsCategory);
diff --git a/Engine/Source/Editor/MaterialEditor/Private/MaterialEditorInstanceDetailCustomization.h b/Engine/Source/Editor/MaterialEditor/Private/MaterialEditorInstanceDetailCustomization.h
--- a/Engine/Source/Editor/MaterialEditor/Private/MaterialEditorInstanceDetailCustomization.h
+++ b/Engine/Source/Editor/MaterialEditor/Private/MaterialEditorInstanceDetailCustomization.h
@@ -141,8 +141,5 @@ private:
 
 	/** Associated FMaterialInstance utilities */
 	TWeakPtr<class IPropertyUtilities> PropertyUtilities;
-
-	/** Storage for validation context. */
-	TSharedPtr<struct FMaterialEditorValidatorContext> ValidationContext;
 };
 
diff --git a/Engine/Source/Editor/MaterialEditor/Private/MaterialInstanceEditor.cpp b/Engine/Source/Editor/MaterialEditor/Private/MaterialInstanceEditor.cpp
--- a/Engine/Source/Editor/MaterialEditor/Private/MaterialInstanceEditor.cpp
+++ b/Engine/Source/Editor/MaterialEditor/Private/MaterialInstanceEditor.cpp
@@ -1235,7 +1235,9 @@ void FMaterialInstanceEditor::NotifyPostChange( const FPropertyChangedEvent& Pro
 	// Rebuild the property window to account for the possibility that  
 	// the item changed was a static switch or function call parameter
 	UObject* PropertyClass = PropertyThatChanged->GetOwner<UObject>();
-	if(PropertyClass && (PropertyClass == UDEditorStaticSwitchParameterValue::StaticClass() || PropertyClass == FEditorParameterGroup::StaticStruct())
+	if(PropertyClass && (PropertyClass == UDEditorStaticSwitchParameterValue::StaticClass() || 
+		PropertyClass == FEditorParameterGroup::StaticStruct() || 
+		PropertyClass == FMaterialInstanceBasePropertyOverrides::StaticStruct())
 		&& MaterialEditorInstance->Parent && MaterialEditorInstance->SourceInstance )
 	{
 		// TODO: We need to hit this on MaterialLayerParam updates but only get notifications for their array elements changing, hence the overly generic test above
