Fix bug where static parameters, component mask parameters were not showing up in the Material Instance Diff Tool.

- This was happening because the Diff tool works by finding visible UPROPERTIES to diff.  Static parameters and such are not visible anywhere so were never considered for the diff.
- The solution is to construct and diff UMaterialEditorInstanceConstant which is the same uobject we display in the Material Instance Editor to display material parameters.
- This will have the correct static switches and has the added side benefit of presenting the same UI for diffing as users expect from the MI Editor.
- Material parameter groups and sub groups now show up in the diff tool.
- Usag flags are no longer shown as a integer, but with the custom detail to display each usage flag.
- We register the same property DetailCustomization for the diff tool as the MI editor does.
- Delete unused lambda CreateDetailsView

[RN] minor | Rendering - Materials and Shaders
#Synced-CL 51543259

diff --git a/Engine/Source/Editor/MaterialEditor/Internal/DiffTool/MaterialInstanceDiffPanel.h b/Engine/Source/Editor/MaterialEditor/Internal/DiffTool/MaterialInstanceDiffPanel.h
--- a/Engine/Source/Editor/MaterialEditor/Internal/DiffTool/MaterialInstanceDiffPanel.h
+++ b/Engine/Source/Editor/MaterialEditor/Internal/DiffTool/MaterialInstanceDiffPanel.h
@@ -13,6 +13,7 @@ class SMaterialEditor3DPreviewViewport;
 class UMaterial;
 class UMaterialExpression;
 class UMaterialInstance;
+class UMaterialEditorInstanceConstant;
 
 struct FRevisionInfo;
 
@@ -39,6 +40,9 @@ struct FMaterialInstanceDiffPanel : public FGCObject
 
 	TObjectPtr<UMaterialInstance> MaterialInstance = nullptr;
 
+	/** Editor proxy used by the diff tool to display all parameters (including static switches) */
+	TObjectPtr<UMaterialEditorInstanceConstant> EditorProxy = nullptr;
+
 	/** Revision information for this MaterialInstance */
 	FRevisionInfo RevisionInfo;
 	
diff --git a/Engine/Source/Editor/MaterialEditor/Private/DiffTool/MaterialInstanceDiffPanel.cpp b/Engine/Source/Editor/MaterialEditor/Private/DiffTool/MaterialInstanceDiffPanel.cpp
--- a/Engine/Source/Editor/MaterialEditor/Private/DiffTool/MaterialInstanceDiffPanel.cpp
+++ b/Engine/Source/Editor/MaterialEditor/Private/DiffTool/MaterialInstanceDiffPanel.cpp
@@ -166,6 +166,7 @@ void FMaterialInstanceDiffPanel::AddReferencedObjects(FReferenceCollector& Colle
 	Collector.AddReferencedObject(ExpressionPreviewMaterial);
 	Collector.AddReferencedObject(PreviewExpression);
 	Collector.AddReferencedObject(MaterialInstance);
+	Collector.AddReferencedObject(EditorProxy);
 }
 
 FString FMaterialInstanceDiffPanel::GetReferencerName() const
diff --git a/Engine/Source/Editor/MaterialEditor/Private/DiffTool/Widgets/SMaterialInstanceDiff.cpp b/Engine/Source/Editor/MaterialEditor/Private/DiffTool/Widgets/SMaterialInstanceDiff.cpp
--- a/Engine/Source/Editor/MaterialEditor/Private/DiffTool/Widgets/SMaterialInstanceDiff.cpp
+++ b/Engine/Source/Editor/MaterialEditor/Private/DiffTool/Widgets/SMaterialInstanceDiff.cpp
@@ -27,6 +27,9 @@
 #include "Widgets/Layout/SWidgetSwitcher.h"
 #include "Widgets/SOverlay.h"
 #include "Materials/MaterialInstance.h"
+#include "MaterialEditor/MaterialEditorInstanceConstant.h"
+#include "MaterialEditorInstanceDetailCustomization.h"
+#include "Materials/MaterialInstanceConstant.h"
 
 #define LOCTEXT_NAMESPACE "SMaterialGraphDiff"
 
@@ -295,20 +298,6 @@ void SMaterialInstanceDiff::GenerateDifferencesList()
 	PrimaryDifferencesList.Empty();
 	RealDifferences.Empty();
 
-	const auto CreateDetailsView = []()
-	{
-		FDetailsViewArgs DetailsViewArgs;
-		DetailsViewArgs.bHideSelectionTip = true;
-		DetailsViewArgs.DefaultsOnlyVisibility = EEditDefaultsOnlyNodeVisibility::Hide;
-		DetailsViewArgs.bShowSectionSelector = true;
-		DetailsViewArgs.bAllowFavoriteSystem = true;
-		DetailsViewArgs.bShowScrollBar = true;
-		DetailsViewArgs.NameAreaSettings = FDetailsViewArgs::HideNameArea;
-
-		FPropertyEditorModule& PropertyModule = FModuleManager::LoadModuleChecked<FPropertyEditorModule>("PropertyEditor");
-		return PropertyModule.CreateDetailView(DetailsViewArgs);
-	};
-
 	ModePanels.Add(DetailsMode, GenerateMaterialInstancePanel());
 
 	DifferencesTreeView->RebuildList();
@@ -335,10 +324,53 @@ void SMaterialInstanceDiff::OnCloseAssetEditor(UObject* Asset, EAssetEditorClose
 
 SMaterialInstanceDiff::FMaterialInstanceDiffControl SMaterialInstanceDiff::GenerateMaterialInstancePanel()
 {
-	const TSharedPtr<FDetailsDiffControl> NewDiffControl = MakeShared<FDetailsDiffControl>(PanelOld.MaterialInstance,
-																						   PanelNew.MaterialInstance,
+	// Create UMaterialEditorInstanceConstant proxies for each side so that ALL parameter types
+	// (including static switches and component masks) are visible in the details view.
+	// This is the same proxy object the Material Instance Editor uses to display parameters.
+	auto CreateEditorProxy = [](FMaterialInstanceDiffPanel& Panel) -> UMaterialEditorInstanceConstant*
+	{
+		UMaterialInstanceConstant* MIC = Cast<UMaterialInstanceConstant>(Panel.MaterialInstance);
+		if (!MIC)
+		{
+			return nullptr;
+		}
+
+		UMaterialEditorInstanceConstant* Proxy = NewObject<UMaterialEditorInstanceConstant>(GetTransientPackage(), NAME_None, RF_Transactional);
+		Proxy->SetSourceInstance(MIC);
+		Panel.EditorProxy = Proxy;
+		return Proxy;
+	};
+
+	UMaterialEditorInstanceConstant* OldProxy = PanelOld.MaterialInstance ? CreateEditorProxy(PanelOld) : nullptr;
+	UMaterialEditorInstanceConstant* NewProxy = PanelNew.MaterialInstance ? CreateEditorProxy(PanelNew) : nullptr;
+
+	// Determine what objects to diff: use proxies when available, fall back to raw MaterialInstance
+	const UObject* OldDiffObject = OldProxy ? static_cast<UObject*>(OldProxy) : static_cast<UObject*>(PanelOld.MaterialInstance);
+	const UObject* NewDiffObject = NewProxy ? static_cast<UObject*>(NewProxy) : static_cast<UObject*>(PanelNew.MaterialInstance);
+
+	const TSharedPtr<FDetailsDiffControl> NewDiffControl = MakeShared<FDetailsDiffControl>(OldDiffObject,
+																						   NewDiffObject,
 																						   FOnDiffEntryFocused::CreateRaw(this, &SMaterialInstanceDiff::SetCurrentMode, DetailsMode),
 																						   true);
+
+	if (OldProxy && NewProxy)
+	{
+		// Register FMaterialInstanceParameterDetails on each details view so that parameters
+		// render with the same grouped UI as the Material Instance Editor.
+		auto RegisterMICCustomization = [](const TSharedRef<IDetailsView>& DetailsView, UMaterialEditorInstanceConstant* Proxy)
+		{
+			FOnGetDetailCustomizationInstance LayoutMICDetails = FOnGetDetailCustomizationInstance::CreateStatic(
+				&FMaterialInstanceParameterDetails::MakeInstance,
+				Proxy,
+				static_cast<SMaterialLayersFunctionsInstanceWrapper*>(nullptr),
+				FGetShowHiddenParameters());
+			DetailsView->RegisterInstancedCustomPropertyLayout(UMaterialEditorInstanceConstant::StaticClass(), LayoutMICDetails);
+			DetailsView->ForceRefresh();
+		};
+		RegisterMICCustomization(NewDiffControl->GetDetailsWidget(OldDiffObject), OldProxy);
+		RegisterMICCustomization(NewDiffControl->GetDetailsWidget(NewDiffObject), NewProxy);
+	}
+
 	NewDiffControl->EnableComments(DifferencesTreeView.ToWeakPtr());
 	NewDiffControl->GenerateTreeEntries(PrimaryDifferencesList, RealDifferences);
 
@@ -349,24 +381,22 @@ SMaterialInstanceDiff::FMaterialInstanceDiffControl SMaterialInstanceDiff::Gener
 
 	DiffDetailSplitter = SNew(SDetailsSplitter);
 
-	if (PanelOld.MaterialInstance)
+	if (OldDiffObject)
 	{
-		const UObject* OldMaterialInstance = PanelOld.MaterialInstance;
 		DiffDetailSplitter->AddSlot(
 			SDetailsSplitter::Slot()
 			.Value(0.5f)
-			.DetailsView(NewDiffControl->GetDetailsWidget(PanelOld.MaterialInstance))
-			.DifferencesWithRightPanel(NewDiffControl.ToSharedRef(), &FDetailsDiffControl::GetDifferencesWithRight, OldMaterialInstance)
+			.DetailsView(NewDiffControl->GetDetailsWidget(OldDiffObject))
+			.DifferencesWithRightPanel(NewDiffControl.ToSharedRef(), &FDetailsDiffControl::GetDifferencesWithRight, OldDiffObject)
 		);
 	}
-	if (PanelNew.MaterialInstance)
+	if (NewDiffObject)
 	{
-		const UObject* NewMaterialInstance = PanelNew.MaterialInstance;
 		DiffDetailSplitter->AddSlot(
 			SDetailsSplitter::Slot()
 			.Value(0.5f)
-			.DetailsView(NewDiffControl->GetDetailsWidget(PanelNew.MaterialInstance))
-			.DifferencesWithRightPanel(NewDiffControl.ToSharedRef(), &FDetailsDiffControl::GetDifferencesWithRight, NewMaterialInstance)
+			.DetailsView(NewDiffControl->GetDetailsWidget(NewDiffObject))
+			.DifferencesWithRightPanel(NewDiffControl.ToSharedRef(), &FDetailsDiffControl::GetDifferencesWithRight, NewDiffObject)
 		);
 	}
 
