Fix bug where I was accidently excluding MaterialInstances from diffing revisions.

- Revert back to checking for UMaterialInterface (which will include UMaterial and UMaterialInstanceConstant)

#rnx
#Synced-CL 51543259

diff --git a/Engine/Source/Editor/MaterialEditor/Private/DiffTool/Widgets/SMaterialRevisionMenu.cpp b/Engine/Source/Editor/MaterialEditor/Private/DiffTool/Widgets/SMaterialRevisionMenu.cpp
--- a/Engine/Source/Editor/MaterialEditor/Private/DiffTool/Widgets/SMaterialRevisionMenu.cpp
+++ b/Engine/Source/Editor/MaterialEditor/Private/DiffTool/Widgets/SMaterialRevisionMenu.cpp
@@ -190,7 +190,7 @@ void SMaterialRevisionMenu::MakeDiffMenu(UToolMenu* Menu)
 							{
 								if (Obj &&
 									!(Obj->HasAnyFlags(RF_Transient) || Obj->GetPackage() == GetTransientPackage()) &&
-									(Obj->IsA<UMaterial>() || Obj->IsA<UMaterialFunction>()))
+									(Obj->IsA<UMaterialInterface>() || Obj->IsA<UMaterialFunction>()))
 								{
 									MaterialObj = Obj;
 									break;
