Fix a bug where the preview meshes in the Material Viewport were not showing up in the Material and Material Instance diff tool.

- Previously the logic allowed for null MaterialEditorPtr which meant we got a 3D viewport with the preview mesh set in situations where we don't have a material editor (the Material and MaterialInstance diff tools).
- Guard against null assets to prevent deferencing a nullptr later.

#rnx
#Synced-CL 51543259

diff --git a/Engine/Source/Editor/MaterialEditor/Private/SMaterialEditorViewport.cpp b/Engine/Source/Editor/MaterialEditor/Private/SMaterialEditorViewport.cpp
--- a/Engine/Source/Editor/MaterialEditor/Private/SMaterialEditorViewport.cpp
+++ b/Engine/Source/Editor/MaterialEditor/Private/SMaterialEditorViewport.cpp
@@ -355,12 +355,7 @@ void SMaterialEditor3DPreviewViewport::RefreshViewport()
 
 bool SMaterialEditor3DPreviewViewport::SetPreviewAsset(UObject* InAsset)
 {
-	if (!MaterialEditorPtr.IsValid() || !InAsset)
-	{
-		return false;
-	}
-
-	if (!MaterialEditorPtr.Pin()->ApproveSetPreviewAsset(InAsset))
+	if (MaterialEditorPtr.IsValid() && !MaterialEditorPtr.Pin()->ApproveSetPreviewAsset(InAsset))
 	{
 		//If an invalid object is selected, revert usage to default asset.
 		if (UMaterialEditorOptions* MaterialEditorOptions = UMaterialEditorOptions::Get())
@@ -373,6 +368,12 @@ bool SMaterialEditor3DPreviewViewport::SetPreviewAsset(UObject* InAsset)
 		}
 	}
 
+	if (InAsset == nullptr)
+	{
+		// No change
+		return true;
+	}
+
 	// This should never occur since we are comparing a mesh against a component...
 	if (PreviewMeshComponent == InAsset)
 	{
