Drastically simplify the SMaterialMaterialInstanceDiff code to use SDetailsDiff directly instead of a bunch of copy/paste.

- SDetailsDiff already handles taking to UObjects, generating a difference list, and management for the difference tree view, the details panel comparison, and difference navigation buttons.
- EditorProxy creation now happens in SMaterialInstanceDiff::Construct.
- No need to make the RevisionInfoWidget, the DifferencesTreeView, GenerateDifferencesList, Diff navigation, or any of the code in FMaterialInstanceDiffControl.  All of this is now handled internally for free by the SDetailsDiff class.
- Use .OnCustomizeDetailsWidget_Raw(this, &SMaterialInstanceDiff::AddMaterialInstanceDetailCustomization) on the SDetailsDiff to control detail customization of the EditorProxy objects.
- This fixes an issue where the different rows were not being highlighted to show that they are different.  The differences are easier to see for the user now.

#rnx
#Synced-CL 51543259

diff --git a/Engine/Source/Editor/MaterialEditor/Internal/DiffTool/Widgets/SMaterialInstanceDiff.h b/Engine/Source/Editor/MaterialEditor/Internal/DiffTool/Widgets/SMaterialInstanceDiff.h
--- a/Engine/Source/Editor/MaterialEditor/Internal/DiffTool/Widgets/SMaterialInstanceDiff.h
+++ b/Engine/Source/Editor/MaterialEditor/Internal/DiffTool/Widgets/SMaterialInstanceDiff.h
@@ -4,18 +4,14 @@
 
 #include "AssetDefinition.h"
 #include "CoreMinimal.h"
-#include "DiffTool/Widgets/SMaterialDiff.h"
 #include "DiffControl.h"
 #include "Materials/MaterialInstance.h"
-#include "SBlueprintDiff.h"
 #include "DiffTool/MaterialInstanceDiffPanel.h"
 #include "Widgets/SCompoundWidget.h"
 
 #define UE_API MATERIALEDITOR_API
 
-class SDetailsSplitter;
-class FBlueprintDifferenceTreeEntry;
-
+class IDetailsView;
 enum class EAssetEditorCloseReason : uint8;
 
 /* Visual Diff between two Material Instances */
@@ -33,9 +29,6 @@ public:
 	UE_API void Construct(const FArguments& InArgs);
 	UE_API virtual ~SMaterialInstanceDiff() override;
 
-	/** Called when user clicks on an entry in the listview of differences */
-	UE_API void OnDiffListSelectionChanged(TSharedPtr<FMaterialDiffResultItem> TheDiff);
-
 	/** Helper function for generating an empty widget */
 	static UE_API TSharedRef<SWidget> DefaultEmptyPanel();
 
@@ -49,25 +42,12 @@ public:
 	FOnWindowClosedEvent OnWindowClosedEvent;
 
 protected:
-	/** Called when user clicks button to go to next difference */
-	UE_API void NextDiff();
-
-	/** Called when user clicks button to go to prev difference */
-	UE_API void PrevDiff();
-
-	/** Called to determine whether we have a list of differences to cycle through */
-	UE_API bool HasNextDiff() const;
-	UE_API bool HasPrevDiff() const;
-
 	/** User toggles the option to show/hide preview viewports */
 	UE_API void ToggleViewport();
 
 	/** Get the image to show for the toggle viewport option */
 	UE_API FSlateIcon GetViewportImage() const;
 
-	/** Function used to generate the list of differences and the widgets needed to calculate that list */
-	UE_API void GenerateDifferencesList();
-
 	/** Called when editor may need to be closed */
 	UE_API void OnCloseAssetEditor(UObject* Asset, EAssetEditorCloseReason CloseReason);
 
@@ -83,8 +63,6 @@ protected:
 		TSharedPtr<IDiffControl> DiffControl;
 	};
 
-	UE_API FMaterialInstanceDiffControl GenerateMaterialInstancePanel();
-
 	/** Accessor and event handler for toggling between diff view modes - only details mode for now: */
 	UE_API void SetCurrentMode(FName NewMode);
 
@@ -102,19 +80,6 @@ protected:
 	/** Contents widget that we swap when mode changes (defaults, components, etc) */
 	TSharedPtr<SBox> ModeContents;
 
-	TSharedPtr<SSplitter> TopRevisionInfoWidget;
-
-	TSharedPtr<SDetailsSplitter> DiffDetailSplitter;
-
-	/** Tree of differences collected across all panels: */
-	TArray<TSharedPtr<FBlueprintDifferenceTreeEntry>> PrimaryDifferencesList;
-
-	/** List of all differences, cached so that we can iterate only the differences and not labels, etc: */
-	TArray<TSharedPtr<FBlueprintDifferenceTreeEntry>> RealDifferences;
-
-	/** Tree view that displays the differences, cached for the buttons that iterate the differences: */
-	TSharedPtr<STreeView<TSharedPtr<FBlueprintDifferenceTreeEntry>>> DifferencesTreeView;
-
 	/** Stored references to widgets used to display various parts of a blueprint, from the mode name */
 	TMap<FName, FMaterialInstanceDiffControl> ModePanels;
 
@@ -122,6 +87,9 @@ protected:
 	TWeakPtr<SWindow> WeakParentWindow;
 
 	FDelegateHandle AssetEditorCloseDelegate;
+
+private:
+	void AddMaterialInstanceDetailCustomization(const TSharedRef<IDetailsView>& DetailsView);
 };
 
 #undef UE_API
diff --git a/Engine/Source/Editor/MaterialEditor/Private/DiffTool/Widgets/SMaterialInstanceDiff.cpp b/Engine/Source/Editor/MaterialEditor/Private/DiffTool/Widgets/SMaterialInstanceDiff.cpp
--- a/Engine/Source/Editor/MaterialEditor/Private/DiffTool/Widgets/SMaterialInstanceDiff.cpp
+++ b/Engine/Source/Editor/MaterialEditor/Private/DiffTool/Widgets/SMaterialInstanceDiff.cpp
@@ -6,25 +6,13 @@
 #include "Framework/Application/SlateApplication.h"
 #include "Framework/MultiBox/MultiBoxBuilder.h"
 #include "Framework/MultiBox/MultiBoxDefs.h"
-#include "GraphDiffControl.h"
 #include "IDetailsView.h"
-#include "MaterialGraph/MaterialGraph.h"
-#include "Modules/ModuleManager.h"
-#include "PropertyEditorModule.h"
 #include "SBlueprintDiff.h"
 #include "SDetailsDiff.h"
-#include "SDetailsSplitter.h"
 #include "SlateOptMacros.h"
-#include "SMyBlueprint.h"
-#include "SourceControlHelpers.h"
 #include "Styling/AppStyle.h"
 #include "Subsystems/AssetEditorSubsystem.h"
-#include "DiffTool/MaterialToDiff.h"
-#include "Widgets/Docking/SDockTab.h"
-#include "Widgets/Input/SButton.h"
-#include "Widgets/Layout/SSpacer.h"
 #include "Widgets/Layout/SSplitter.h"
-#include "Widgets/Layout/SWidgetSwitcher.h"
 #include "Widgets/SOverlay.h"
 #include "Materials/MaterialInstance.h"
 #include "MaterialEditor/MaterialEditorInstanceConstant.h"
@@ -40,8 +28,26 @@ void SMaterialInstanceDiff::Construct(const FArguments& InArgs)
 {
 	check(InArgs._OldMaterialInstance || InArgs._NewMaterialInstance);
 
+	// Create UMaterialEditorInstanceConstant proxies for each side so that ALL parameter types
+	// (including static switches and component masks) are visible in the details view.
+	// This is the same proxy object the Material Instance Editor uses to display parameters.
+	auto CreateEditorProxy = [](UMaterialInstance* MI) -> UMaterialEditorInstanceConstant*
+	{
+		UMaterialInstanceConstant* MIC = Cast<UMaterialInstanceConstant>(MI);
+		if (!MIC)
+		{
+			return nullptr;
+		}
+
+		UMaterialEditorInstanceConstant* Proxy = NewObject<UMaterialEditorInstanceConstant>(GetTransientPackage(), NAME_None, RF_Transactional);
+		Proxy->SetSourceInstance(MIC);
+		return Proxy;
+	};
+
 	PanelOld.MaterialInstance = InArgs._OldMaterialInstance;
 	PanelNew.MaterialInstance = InArgs._NewMaterialInstance;
+	PanelOld.EditorProxy = CreateEditorProxy(PanelOld.MaterialInstance);
+	PanelNew.EditorProxy = CreateEditorProxy(PanelNew.MaterialInstance);
 	PanelOld.RevisionInfo = InArgs._OldRevision;
 	PanelNew.RevisionInfo = InArgs._NewRevision;
 
@@ -53,24 +59,6 @@ void SMaterialInstanceDiff::Construct(const FArguments& InArgs)
 	}
 
 	FToolBarBuilder NavToolBarBuilder(TSharedPtr<const FUICommandList>(), FMultiBoxCustomization::None);
-	NavToolBarBuilder.AddToolBarButton(
-		FUIAction(
-			FExecuteAction::CreateSP(this, &SMaterialInstanceDiff::PrevDiff),
-			FCanExecuteAction::CreateSP(this, &SMaterialInstanceDiff::HasPrevDiff)),
-		NAME_None,
-		LOCTEXT("PrevDiffLabel", "Prev"),
-		LOCTEXT("PrevDiffTooltip", "Go to previous difference"),
-		FSlateIcon(FAppStyle::GetAppStyleSetName(), "BlueprintDif.PrevDiff")
-	);
-	NavToolBarBuilder.AddToolBarButton(
-		FUIAction(
-			FExecuteAction::CreateSP(this, &SMaterialInstanceDiff::NextDiff),
-			FCanExecuteAction::CreateSP(this, &SMaterialInstanceDiff::HasNextDiff)),
-		NAME_None,
-		LOCTEXT("NextDiffLabel", "Next"),
-		LOCTEXT("NextDiffTooltip", "Go to next difference"),
-		FSlateIcon(FAppStyle::GetAppStyleSetName(), "BlueprintDif.NextDiff")
-	);
 	NavToolBarBuilder.AddToolBarButton(
 		FUIAction(
 			FExecuteAction::CreateSP(this, &SMaterialInstanceDiff::ToggleViewport)),
@@ -80,48 +68,8 @@ void SMaterialInstanceDiff::Construct(const FArguments& InArgs)
 		TAttribute<FSlateIcon>(this, &SMaterialInstanceDiff::GetViewportImage)
 	);
 
-	DifferencesTreeView = DiffTreeView::CreateTreeView(&PrimaryDifferencesList);
-
-	GenerateDifferencesList();
-
-	const auto TextBlock = [](FText Text) -> TSharedRef<SWidget>
-	{
-		return SNew(SBox)
-		.Padding(FMargin(4.0f, 10.0f))
-		.VAlign(VAlign_Center)
-		.HAlign(HAlign_Left)
-		[
-			SNew(STextBlock)
-			.Visibility(EVisibility::HitTestInvisible)
-			.TextStyle(FAppStyle::Get(), "DetailsView.CategoryTextStyle")
-			.Text(Text)
-		];
-	};
-
-	TopRevisionInfoWidget =
-		SNew(SSplitter)
-		.Visibility(EVisibility::HitTestInvisible)
-		+ SSplitter::Slot()
-		.Value(.2f)
-		[
-			SNew(SBox)
-		]
-		+ SSplitter::Slot()
-		.Value(.8f)
-		[
-			SNew(SSplitter)
-			.PhysicalSplitterHandleSize(10.0f)
-			+ SSplitter::Slot()
-			.Value(.5f)
-			[
-				TextBlock(DiffViewUtils::GetPanelLabel(PanelOld.MaterialInstance, PanelOld.RevisionInfo, FText()))
-			]
-			+ SSplitter::Slot()
-			.Value(.5f)
-			[
-				TextBlock(DiffViewUtils::GetPanelLabel(PanelNew.MaterialInstance, PanelNew.RevisionInfo, FText()))
-			]
-		];
+	PanelOld.SetViewportToDisplay();
+	PanelNew.SetViewportToDisplay();
 
 	// Clang-format off
 	ChildSlot
@@ -129,21 +77,7 @@ void SMaterialInstanceDiff::Construct(const FArguments& InArgs)
 		SNew(SBorder)
 		.BorderImage(FAppStyle::GetBrush("Docking.Tab", ".ContentAreaBrush"))
 		[
-			SNew(SOverlay)
-			+SOverlay::Slot()
-			.VAlign(VAlign_Top)
-			[
-				TopRevisionInfoWidget.ToSharedRef()		
-			]
-			+ SOverlay::Slot()
-			.Padding(0.f, 6.f, 0.f, 4.f)
-			[
-				SNew(SOverlay)
-				+SOverlay::Slot()
-				.VAlign(VAlign_Top)
-				.Padding(0.0f, 6.0f, 0.0f, 4.0f)
-				[
-					SNew(SVerticalBox)
+			SNew(SVerticalBox)
 					+SVerticalBox::Slot()
 					.AutoHeight()
 					.Padding(0.0f, 2.0f, 0.0f, 2.0f)
@@ -160,33 +94,50 @@ void SMaterialInstanceDiff::Construct(const FArguments& InArgs)
 							SNew(SSpacer)
 						]
 					]
-					+SVerticalBox::Slot()
-					.FillHeight(1.f)
+			+ SVerticalBox::Slot()
+			.FillHeight(1.f)
+			[
+				SNew(SHorizontalBox)
+				+ SHorizontalBox::Slot()
+				.FillWidth(1.f)
+				[
+					SNew(SSplitter)
+					.Orientation(Orient_Vertical)
+					+ SSplitter::Slot()
+					.Value(0.6f)
+					[
+						SNew(SDetailsDiff)
+							.OldAsset(PanelOld.EditorProxy)
+							.NewAsset(PanelNew.EditorProxy)
+							.OldRevision(PanelOld.RevisionInfo)
+							.NewRevision(PanelNew.RevisionInfo)
+							.ShowAssetNames(true)
+							.ParentWindow(WeakParentWindow.Pin())
+							.OnCustomizeDetailsWidget_Raw(this, &SMaterialInstanceDiff::AddMaterialInstanceDetailCustomization)
+					]
+					+ SSplitter::Slot()
+					.Value(0.4f)
 					[
 						SNew(SSplitter)
-						+SSplitter::Slot()
-						.Value(.2f)
+						.Visibility_Lambda([this]()
+						{
+							return bShowViewport ? EVisibility::Visible : EVisibility::Collapsed;
+						})
+						.Orientation(Orient_Horizontal)
+						+ SSplitter::Slot()
 						[
-							SNew(SBorder)
-							.BorderImage(FAppStyle::GetBrush("ToolPanel.GroupBorder"))
-							[
-								DifferencesTreeView.ToSharedRef()
-							]
+							PanelOld.GetViewportToDisplay()
 						]
-						+SSplitter::Slot()
-						.Value(.8f)
+						+ SSplitter::Slot()
 						[
-							SAssignNew(ModeContents, SBox)
+							PanelNew.GetViewportToDisplay()
 						]
 					]
 				]
-				
 			]
 		]
 	];
 	// clang-format on
-
-	SetCurrentMode(DetailsMode);
 }
 END_SLATE_FUNCTION_BUILD_OPTIMIZATION
 
@@ -198,10 +149,6 @@ SMaterialInstanceDiff::~SMaterialInstanceDiff()
 	}
 }
 
-void SMaterialInstanceDiff::OnDiffListSelectionChanged(TSharedPtr<FMaterialDiffResultItem> TheDiff)
-{
-}
-
 TSharedRef<SWidget> SMaterialInstanceDiff::DefaultEmptyPanel()
 {
 	return SNew(SHorizontalBox)
@@ -263,26 +210,6 @@ TSharedPtr<SWindow> SMaterialInstanceDiff::CreateDiffWindow(TObjectPtr<UMaterial
 	return CreateDiffWindow(WindowTitle, OldMaterialInstance, NewMaterialInstance, OldRevision, NewRevision);
 }
 
-void SMaterialInstanceDiff::NextDiff()
-{
-	DiffTreeView::HighlightNextDifference(DifferencesTreeView.ToSharedRef(), RealDifferences, PrimaryDifferencesList);
-}
-
-void SMaterialInstanceDiff::PrevDiff()
-{
-	DiffTreeView::HighlightPrevDifference(DifferencesTreeView.ToSharedRef(), RealDifferences, PrimaryDifferencesList);
-}
-
-bool SMaterialInstanceDiff::HasNextDiff() const
-{
-	return DiffTreeView::HasNextDifference(DifferencesTreeView.ToSharedRef(), RealDifferences);
-}
-
-bool SMaterialInstanceDiff::HasPrevDiff() const
-{
-	return DiffTreeView::HasPrevDifference(DifferencesTreeView.ToSharedRef(), RealDifferences);
-}
-
 void SMaterialInstanceDiff::ToggleViewport()
 {
 	bShowViewport = !bShowViewport;
@@ -292,17 +219,6 @@ FSlateIcon SMaterialInstanceDiff::GetViewportImage() const
 {
 	return FSlateIcon(FAppStyle::GetAppStyleSetName(), "LevelEditor.Tabs.Viewports");
 }
- 
-void SMaterialInstanceDiff::GenerateDifferencesList()
-{
-	PrimaryDifferencesList.Empty();
-	RealDifferences.Empty();
-
-	ModePanels.Add(DetailsMode, GenerateMaterialInstancePanel());
-
-	DifferencesTreeView->RebuildList();
-}
-
 void SMaterialInstanceDiff::OnCloseAssetEditor(UObject* Asset, EAssetEditorCloseReason CloseReason)
 {
 	if (CloseReason == EAssetEditorCloseReason::CloseAllAssetEditors)
@@ -322,121 +238,38 @@ void SMaterialInstanceDiff::OnCloseAssetEditor(UObject* Asset, EAssetEditorClose
 	}
 }
 
-SMaterialInstanceDiff::FMaterialInstanceDiffControl SMaterialInstanceDiff::GenerateMaterialInstancePanel()
+void SMaterialInstanceDiff::AddMaterialInstanceDetailCustomization(const TSharedRef<IDetailsView>& DetailsView)
 {
-	// Create UMaterialEditorInstanceConstant proxies for each side so that ALL parameter types
-	// (including static switches and component masks) are visible in the details view.
-	// This is the same proxy object the Material Instance Editor uses to display parameters.
-	auto CreateEditorProxy = [](FMaterialInstanceDiffPanel& Panel) -> UMaterialEditorInstanceConstant*
+	// Determine which panel this details view belongs to by checking which EditorProxy it is displaying.
+	// SDetailsDiff calls this callback twice: once for the Old side and once for the New side.
+	UMaterialEditorInstanceConstant* ProxyForThisView = nullptr;
+	const TArray<TWeakObjectPtr<UObject>>& SelectedObjects = DetailsView->GetSelectedObjects();
+	if (SelectedObjects.Num() > 0 && SelectedObjects[0].IsValid())
 	{
-		UMaterialInstanceConstant* MIC = Cast<UMaterialInstanceConstant>(Panel.MaterialInstance);
-		if (!MIC)
+		UObject* DisplayedObject = SelectedObjects[0].Get();
+		if (DisplayedObject == PanelOld.EditorProxy)
 		{
-			return nullptr;
+			ProxyForThisView = PanelOld.EditorProxy;
 		}
-
-		UMaterialEditorInstanceConstant* Proxy = NewObject<UMaterialEditorInstanceConstant>(GetTransientPackage(), NAME_None, RF_Transactional);
-		Proxy->SetSourceInstance(MIC);
-		Panel.EditorProxy = Proxy;
-		return Proxy;
-	};
-
-	UMaterialEditorInstanceConstant* OldProxy = PanelOld.MaterialInstance ? CreateEditorProxy(PanelOld) : nullptr;
-	UMaterialEditorInstanceConstant* NewProxy = PanelNew.MaterialInstance ? CreateEditorProxy(PanelNew) : nullptr;
-
-	// Determine what objects to diff: use proxies when available, fall back to raw MaterialInstance
-	const UObject* OldDiffObject = OldProxy ? static_cast<UObject*>(OldProxy) : static_cast<UObject*>(PanelOld.MaterialInstance);
-	const UObject* NewDiffObject = NewProxy ? static_cast<UObject*>(NewProxy) : static_cast<UObject*>(PanelNew.MaterialInstance);
-
-	const TSharedPtr<FDetailsDiffControl> NewDiffControl = MakeShared<FDetailsDiffControl>(OldDiffObject,
-																						   NewDiffObject,
-																						   FOnDiffEntryFocused::CreateRaw(this, &SMaterialInstanceDiff::SetCurrentMode, DetailsMode),
-																						   true);
-
-	if (OldProxy && NewProxy)
-	{
-		// Register FMaterialInstanceParameterDetails on each details view so that parameters
-		// render with the same grouped UI as the Material Instance Editor.
-		auto RegisterMICCustomization = [](const TSharedRef<IDetailsView>& DetailsView, UMaterialEditorInstanceConstant* Proxy)
+		else if (DisplayedObject == PanelNew.EditorProxy)
 		{
-			FOnGetDetailCustomizationInstance LayoutMICDetails = FOnGetDetailCustomizationInstance::CreateStatic(
-				&FMaterialInstanceParameterDetails::MakeInstance,
-				Proxy,
-				static_cast<SMaterialLayersFunctionsInstanceWrapper*>(nullptr),
-				FGetShowHiddenParameters());
-			DetailsView->RegisterInstancedCustomPropertyLayout(UMaterialEditorInstanceConstant::StaticClass(), LayoutMICDetails);
-			DetailsView->ForceRefresh();
-		};
-		RegisterMICCustomization(NewDiffControl->GetDetailsWidget(OldDiffObject), OldProxy);
-		RegisterMICCustomization(NewDiffControl->GetDetailsWidget(NewDiffObject), NewProxy);
+			ProxyForThisView = PanelNew.EditorProxy;
+		}
 	}
 
-	NewDiffControl->EnableComments(DifferencesTreeView.ToWeakPtr());
-	NewDiffControl->GenerateTreeEntries(PrimaryDifferencesList, RealDifferences);
-
-	FMaterialInstanceDiffControl MaterialInstancePanel{};
-
-	PanelOld.SetViewportToDisplay();
-	PanelNew.SetViewportToDisplay();
-
-	DiffDetailSplitter = SNew(SDetailsSplitter);
-
-	if (OldDiffObject)
-	{
-		DiffDetailSplitter->AddSlot(
-			SDetailsSplitter::Slot()
-			.Value(0.5f)
-			.DetailsView(NewDiffControl->GetDetailsWidget(OldDiffObject))
-			.DifferencesWithRightPanel(NewDiffControl.ToSharedRef(), &FDetailsDiffControl::GetDifferencesWithRight, OldDiffObject)
-		);
-	}
-	if (NewDiffObject)
+	if (!ProxyForThisView)
 	{
-		DiffDetailSplitter->AddSlot(
-			SDetailsSplitter::Slot()
-			.Value(0.5f)
-			.DetailsView(NewDiffControl->GetDetailsWidget(NewDiffObject))
-			.DifferencesWithRightPanel(NewDiffControl.ToSharedRef(), &FDetailsDiffControl::GetDifferencesWithRight, NewDiffObject)
-		);
+		return;
 	}
 
-	MaterialInstancePanel.DiffControl = NewDiffControl;
-	MaterialInstancePanel.Widget = SNew(SVerticalBox)
-	+SVerticalBox::Slot()
-	.FillHeight(1.0f)
-	[
-		SNew(SHorizontalBox)
-		+SHorizontalBox::Slot()
-		.FillWidth(1.0f)
-		[
-			SNew(SSplitter)
-			.Orientation(Orient_Vertical)
-			+SSplitter::Slot()
-			.Value(0.6f)
-			[
-				DiffDetailSplitter.ToSharedRef()
-			]
-			+SSplitter::Slot()
-			.Value(0.4f)
-			[
-				SNew(SSplitter)
-				.Orientation(Orient_Horizontal).Visibility_Lambda([this]()
-				{
-					return bShowViewport ? EVisibility::Visible : EVisibility::Collapsed;
-				})
-				+ SSplitter::Slot()
-				[
-					PanelOld.GetViewportToDisplay()
-				]
-				+SSplitter::Slot()
-				[
-					PanelNew.GetViewportToDisplay()
-				]
-			]
-		]
-	];
+	FOnGetDetailCustomizationInstance LayoutMICDetails = FOnGetDetailCustomizationInstance::CreateStatic(
+		&FMaterialInstanceParameterDetails::MakeInstance,
+		ProxyForThisView,
+		static_cast<SMaterialLayersFunctionsInstanceWrapper*>(nullptr),
+		FGetShowHiddenParameters());
 
-	return MaterialInstancePanel;
+	DetailsView->RegisterInstancedCustomPropertyLayout(UMaterialEditorInstanceConstant::StaticClass(), LayoutMICDetails);
+	DetailsView->ForceRefresh();
 }
 
 void SMaterialInstanceDiff::SetCurrentMode(FName NewMode)
