Focus the 3D Material Viewports on the mesh in the preview scene when in the Material Diff tool.

- This way you don't have to manually zoom in/out each time you open the viewport to see the full meshes.

#rnx
#Synced-CL 51679461

diff --git a/Engine/Source/Editor/MaterialEditor/Private/DiffTool/MaterialDiffPanel.cpp b/Engine/Source/Editor/MaterialEditor/Private/DiffTool/MaterialDiffPanel.cpp
--- a/Engine/Source/Editor/MaterialEditor/Private/DiffTool/MaterialDiffPanel.cpp
+++ b/Engine/Source/Editor/MaterialEditor/Private/DiffTool/MaterialDiffPanel.cpp
@@ -2,6 +2,7 @@
 
 #include "DiffTool/MaterialDiffPanel.h"
 
+#include "Components/MeshComponent.h"
 #include "EdGraphUtilities.h"
 #include "Framework/Commands/GenericCommands.h"
 #include "GraphDiffControl.h"
@@ -380,6 +381,11 @@ void FMaterialDiffPanel::SetViewportToDisplay()
 	SAssignNew(Preview3DViewport, SMaterialEditor3DPreviewViewport)
 		.PreviewMaterial(MaterialGraph->Material);
 
+	if (Preview3DViewport->PreviewMeshComponent != nullptr)
+	{
+		Preview3DViewport->GetViewportClient()->FocusViewportOnBox(Preview3DViewport->PreviewMeshComponent->Bounds.GetBox(), true);
+	}
+
 	// See FMaterialEditor::InitMaterialEditor
 	if (MaterialGraph->MaterialFunction && MaterialGraph->Material)
 	{
