Make Physical Material be displayed as an overrideable parameter.

Because this parameter behaves like an overrideable one with parent being chosen if current has "None", but nothing was being displayed to showscase that, it was creating confusion.

Now, the Physical Material has an overrideable flag which starts as false for a new Material Instance, mimicking the other properties.
Materials with existing Physical Material override will be fixed up on Serialize() and their override flag will be enabled.

In case of an intentional override flag being enabled and value set to "None", default will be used, as requested by users.

Modify Build.cs to be able to include PhysicalMaterial in MaterialEditorInstanceDetailCustomization.cpp .

[RN] minor | Rendering Architecture

#Synced-CL 51712365

diff --git a/Engine/Source/Editor/MaterialEditor/MaterialEditor.Build.cs b/Engine/Source/Editor/MaterialEditor/MaterialEditor.Build.cs
--- a/Engine/Source/Editor/MaterialEditor/MaterialEditor.Build.cs
+++ b/Engine/Source/Editor/MaterialEditor/MaterialEditor.Build.cs
@@ -50,9 +50,10 @@ public class MaterialEditor : ModuleRules
 				"Landscape",
 				"Kismet", // for rich diffing machinery which was mostly implemented in the kismet module
 				"SourceControl",
-				"EditorSubsystem"
-			}
-		);
+				"EditorSubsystem",
+				"PhysicsCore"
+            }
+        );
 
 		DynamicallyLoadedModuleNames.AddRange(
 			new string[] {
diff --git a/Engine/Source/Editor/MaterialEditor/Private/MaterialEditorInstanceDetailCustomization.cpp b/Engine/Source/Editor/MaterialEditor/Private/MaterialEditorInstanceDetailCustomization.cpp
--- a/Engine/Source/Editor/MaterialEditor/Private/MaterialEditorInstanceDetailCustomization.cpp
+++ b/Engine/Source/Editor/MaterialEditor/Private/MaterialEditorInstanceDetailCustomization.cpp
@@ -47,6 +47,7 @@
 #include "SWarningOrErrorBox.h"
 #include "Widgets/Input/SCheckBox.h"
 #include "Widgets/SToolTip.h"
+#include "PhysicalMaterials/PhysicalMaterial.h"
 
 #define LOCTEXT_NAMESPACE "MaterialInstanceEditor"
 
@@ -184,9 +185,6 @@ void FMaterialInstanceParameterDetails::CustomizeDetails(IDetailLayoutBuilder& D
 		DetailLayout.HideProperty("PostProcessOverrides");
 		CreatePostProcessOverrideWidgets(DetailLayout);
 
-		// Add PhysMaterial property
-		DefaultCategory.AddProperty("PhysMaterial");
-
 		// Customize Parent property so we can check for recursively set parents
 		TSharedRef<IPropertyHandle> ParentPropertyHandle = DetailLayout.GetProperty("Parent");
 		IDetailPropertyRow& ParentPropertyRow = DefaultCategory.AddProperty(ParentPropertyHandle);
@@ -214,6 +212,36 @@ void FMaterialInstanceParameterDetails::CustomizeDetails(IDetailLayoutBuilder& D
 				ParentValueWidget
 			];
 
+		{
+			DetailLayout.HideProperty("PhysMaterial");
+			DetailLayout.HideProperty("bOverridePhysMaterial");
+
+			// Set up the override logic for PhysMaterial
+			{
+				TAttribute<bool> IsPhysMatOverrideEnabled = TAttribute<bool>::Create(TAttribute<bool>::FGetter::CreateLambda([this]() { return (bool)MaterialEditorInstance->bOverridePhysMaterial; }));
+				TSharedRef<IPropertyHandle> PhysMatPropHandle = DetailLayout.GetProperty("PhysMaterial");
+				IDetailPropertyRow& PhysMatPropertyRow = DefaultCategory.AddProperty(PhysMatPropHandle);
+
+				FResetToDefaultHandler ResetPhysMatPropertyHandler = FResetToDefaultHandler::CreateLambda([this](TSharedPtr<IPropertyHandle> InHandle) {
+					if (MaterialEditorInstance->Parent)
+					{
+						InHandle->SetValue(MaterialEditorInstance->Parent->GetPhysicalMaterial());
+					}
+				});
+
+				FResetToDefaultOverride ResetPhysMatPropertyOverride = FResetToDefaultOverride::Create(ResetPhysMatPropertyHandler);
+
+				PhysMatPropertyRow
+					.EditCondition(IsPhysMatOverrideEnabled,
+						FOnBooleanValueChanged::CreateLambda([this](bool NewValue) {
+							MaterialEditorInstance->bOverridePhysMaterial = (uint32)NewValue;
+							MaterialEditorInstance->PostEditChange();
+							FEditorSupportDelegates::RedrawAllViewports.Broadcast();
+							}))
+					.OverrideResetToDefault(ResetPhysMatPropertyOverride);
+			}
+		}
+
 		// Add/hide other properties
 		DetailLayout.HideProperty("LightmassSettings");
 		CreateLightmassOverrideWidgets(DetailLayout);
diff --git a/Engine/Source/Editor/UnrealEd/Classes/MaterialEditor/MaterialEditorInstanceConstant.h b/Engine/Source/Editor/UnrealEd/Classes/MaterialEditor/MaterialEditorInstanceConstant.h
--- a/Engine/Source/Editor/UnrealEd/Classes/MaterialEditor/MaterialEditorInstanceConstant.h
+++ b/Engine/Source/Editor/UnrealEd/Classes/MaterialEditor/MaterialEditorInstanceConstant.h
@@ -353,6 +353,10 @@ private:
 	UPROPERTY(transient, duplicatetransient)
 	uint32 bIsFunctionInstanceDirty : 1;
 	
+	/** Defines if PhysMaterial from this instance is used or if the parent one is used. */
+	UPROPERTY(EditAnywhere, Category=MaterialEditorInstanceConstant)
+	uint32 bOverridePhysMaterial : 1;
+	
 	UPROPERTY(EditAnywhere, Category=MaterialOverrides)
 	FMaterialInstanceBasePropertyOverrides BasePropertyOverrides;
 
diff --git a/Engine/Source/Editor/UnrealEd/Private/PreviewMaterial.cpp b/Engine/Source/Editor/UnrealEd/Private/PreviewMaterial.cpp
--- a/Engine/Source/Editor/UnrealEd/Private/PreviewMaterial.cpp
+++ b/Engine/Source/Editor/UnrealEd/Private/PreviewMaterial.cpp
@@ -942,6 +942,8 @@ void UMaterialEditorInstanceConstant::CopyToSourceInstance(const bool bForceStat
 
 		// Copy phys material back to source instance
 		SourceInstance->PhysMaterial = PhysMaterial;
+		SourceInstance->bOverridePhysMaterial = bOverridePhysMaterial;
+
 		SourceInstance->NaniteOverrideMaterial.bEnableOverride = bNaniteOverride;
 		SourceInstance->NaniteOverrideMaterial.OverrideMaterialEditor = NaniteOverrideMaterial;
 		SourceInstance->EnumerationObjects = EnumerationObjects;
@@ -1056,7 +1058,7 @@ void UMaterialEditorInstanceConstant::SetSourceInstance(UMaterialInstanceConstan
 	PreviewInstance->BasePropertyOverrides.UsageFlags = (1u << MATUSAGE_StaticMesh);
 
 	Parent = SourceInstance->Parent;
-	PhysMaterial = SourceInstance->PhysMaterial;
+
 	bNaniteOverride = SourceInstance->NaniteOverrideMaterial.bEnableOverride;
 	NaniteOverrideMaterial = SourceInstance->NaniteOverrideMaterial.OverrideMaterialEditor;
 	EnumerationObjects = SourceInstance->EnumerationObjects;
@@ -1184,6 +1186,11 @@ void UMaterialEditorInstanceConstant::CopyBasePropertiesFromParent()
 	{
 		BasePropertyOverrides.bCompatibleWithLumenCardSharing = SourceInstance->IsCompatibleWithLumenCardSharing();
 	}
+
+	bOverridePhysMaterial = SourceInstance->bOverridePhysMaterial;
+	// If overriden, always keep the assigned value, even if None
+	PhysMaterial = bOverridePhysMaterial ? SourceInstance->PhysMaterial : TObjectPtr<UPhysicalMaterial>(SourceInstance->GetPhysicalMaterial());
+
 	// Usage flags and overrides are a bitmask.	
 	BasePropertyOverrides.UsageFlags = (BasePropertyOverrides.bOverride_UsageFlags & BasePropertyOverrides.UsageFlags) | (~BasePropertyOverrides.bOverride_UsageFlags & SourceInstance->GetUsageFlags());
 
diff --git a/Engine/Source/Runtime/Core/Public/UObject/FortniteMainBranchObjectVersions.inl b/Engine/Source/Runtime/Core/Public/UObject/FortniteMainBranchObjectVersions.inl
--- a/Engine/Source/Runtime/Core/Public/UObject/FortniteMainBranchObjectVersions.inl
+++ b/Engine/Source/Runtime/Core/Public/UObject/FortniteMainBranchObjectVersions.inl
@@ -787,5 +787,8 @@ FORTNITE_MAIN_VERSION(DynamicMeshAttributesTriangleLabels, 257)
 // Remove RF_Public on entity components and subobjects
 FORTNITE_MAIN_VERSION(FixPublicEntityComponentsAndSubobjects, 258)
 
+// Added override flag for Physical Material
+FORTNITE_MAIN_VERSION(AddPhysMatOverrideFlag, 259)
+
 // -----<new versions can be added above this line>-------------------------------------------------
 #undef FORTNITE_MAIN_VERSION
diff --git a/Engine/Source/Runtime/Engine/Private/Materials/MaterialInstance.cpp b/Engine/Source/Runtime/Engine/Private/Materials/MaterialInstance.cpp
--- a/Engine/Source/Runtime/Engine/Private/Materials/MaterialInstance.cpp
+++ b/Engine/Source/Runtime/Engine/Private/Materials/MaterialInstance.cpp
@@ -2137,11 +2137,20 @@ UPhysicalMaterial* UMaterialInstance::GetPhysicalMaterial() const
 	}
 
 	FMICReentranceGuard	Guard(const_cast<UMaterialInstance*>(this));  // should not need this to determine loop
-	if(PhysMaterial)
+
+	if (bOverridePhysMaterial)
 	{
-		return PhysMaterial;
+		if (PhysMaterial)
+		{
+			return PhysMaterial;
+		}
+		else
+		{
+			check(GEngine->DefaultPhysMaterial != NULL);
+			return GEngine->DefaultPhysMaterial;
+		}
 	}
-	else if(Parent)
+	else if (Parent)
 	{
 		// If no physical material has been associated with this instance, simply use the parent's physical material.
 		return Parent->GetPhysicalMaterial();
@@ -3431,6 +3440,12 @@ void UMaterialInstance::Serialize(FArchive& Ar)
 	{
 		ReferencedDefaultTextures.Empty();
 	}
+
+	if (Ar.IsLoading() && Ar.CustomVer(FFortniteMainBranchObjectVersion::GUID) < FFortniteMainBranchObjectVersion::AddPhysMatOverrideFlag)
+	{
+		// Set flag on materials which had Phys Mat already overriden
+		bOverridePhysMaterial = bool(PhysMaterial);
+	}
 #endif
 }
 
diff --git a/Engine/Source/Runtime/Engine/Public/Materials/MaterialInstance.h b/Engine/Source/Runtime/Engine/Public/Materials/MaterialInstance.h
--- a/Engine/Source/Runtime/Engine/Public/Materials/MaterialInstance.h
+++ b/Engine/Source/Runtime/Engine/Public/Materials/MaterialInstance.h
@@ -700,6 +700,10 @@ class UMaterialInstance : public UMaterialInterface
 	UPROPERTY(EditAnywhere, BlueprintReadWrite, Category = MaterialInstance)
 	uint8 bOverrideBlendablePriority : 1;
 
+	/** Defines if PhysMaterial from this instance is used or if the parent one is used. */
+	UPROPERTY(EditAnywhere, BlueprintReadWrite, Category = MaterialInstance)
+	uint8 bOverridePhysMaterial : 1;
+
 	uint8 TwoSided : 1;
 	uint8 bIsThinSurface : 1;
 	uint8 DitheredLODTransition : 1;
