Add a Named Reroute palette in the Material Editor.

- Adds a SMaterialNamedReroutesPanel class and add it to the MaterialEditor.
- Need 3 separate locations in the code where reroute declarations can change and I needed to handle refreshing the reroute panel:
- Add a OnGraphChangedDelegateHandle to the MaterialGraph to handle changes to the node graph like duplication, add, or removal to catch when the reroute palette should update.  The graph changed delegate blindly refreshes the palette on any graph change.
- FMaterialEditor::OnNodeTitleCommitted needed special case code to refresh the palette when you change the reroute declaration's name on the node directly itself.
- FMaterialEditor::OnFinishedChangingProperties needed special case code to refresh the palette when you change the Name or NodeColor on the declaration.
- Similar to the "Named Reroutes" section in the material right click menu, this is a tab on the left side which is a list of all named reroute nodes.
- Refactor UMaterialGraphSchema::GetNamedRerouteActions into a public function so both the graph and the palette can discover all reroutes with a common function.
- Refactor common code into FMaterialEditor::SelectNamedRerouteUsages to find all reroute usages in the graph from a given declaration expression.  This is used in the context menu in the palette and in the graph.

[RN] major | Rendering - Materials and Shaders

#Synced-CL 51761382

diff --git a/Engine/Source/Editor/MaterialEditor/Private/MaterialEditor.cpp b/Engine/Source/Editor/MaterialEditor/Private/MaterialEditor.cpp
--- a/Engine/Source/Editor/MaterialEditor/Private/MaterialEditor.cpp
+++ b/Engine/Source/Editor/MaterialEditor/Private/MaterialEditor.cpp
@@ -19,6 +19,7 @@
 #include "Widgets/Notifications/SNotificationBackground.h"
 #include "Styling/AppStyle.h"
 #include "EdGraph/EdGraph.h"
+#include "GraphEditAction.h"
 #include "WorkflowOrientedApp/WorkflowUObjectDocuments.h"
 #include "MaterialGraph/MaterialGraph.h"
 #include "MaterialGraph/MaterialGraphNode_Comment.h"
@@ -135,6 +136,7 @@
 #include "SGraphPanel.h"
 #include "MaterialEditorUtilities.h"
 #include "SMaterialPalette.h"
+#include "SMaterialNamedReroutesPanel.h"
 #include "FindInMaterial.h"
 #include "Misc/FeedbackContext.h"
 #include "UObject/UObjectHash.h"
@@ -455,6 +457,11 @@ void FMaterialEditor::RegisterToolbarTab(const TSharedRef<class FTabManager>& In
 		.SetGroup( WorkspaceMenuCategoryRef )
 		.SetIcon(FSlateIcon(FAppStyle::GetAppStyleSetName(), "Kismet.Tabs.Palette"));
 
+	InTabManager->RegisterTabSpawner(FMaterialEditorTabs::NamedReroutesTabId, FOnSpawnTab::CreateSP(this, &FMaterialEditor::SpawnTab_NamedReroutes))
+		.SetDisplayName(LOCTEXT("NamedReroutesTab", "Named Reroutes"))
+		.SetGroup(WorkspaceMenuCategoryRef)
+		.SetIcon(FSlateIcon(FAppStyle::GetAppStyleSetName(), "Kismet.Tabs.Palette"));
+
 	InTabManager->RegisterTabSpawner(FMaterialEditorTabs::FindTabId, FOnSpawnTab::CreateSP(this, &FMaterialEditor::SpawnTab_Find))
 		.SetDisplayName(LOCTEXT("FindTab", "Find Results"))
 		.SetGroup( WorkspaceMenuCategoryRef )
@@ -508,6 +515,7 @@ void FMaterialEditor::UnregisterTabSpawners(const TSharedRef<class FTabManager>&
 	InTabManager->UnregisterTabSpawner(FMaterialEditorTabs::PreviewTabId);
 	InTabManager->UnregisterTabSpawner(FMaterialEditorTabs::PropertiesTabId);
 	InTabManager->UnregisterTabSpawner(FMaterialEditorTabs::PaletteTabId);
+	InTabManager->UnregisterTabSpawner(FMaterialEditorTabs::NamedReroutesTabId);
 	InTabManager->UnregisterTabSpawner(FMaterialEditorTabs::FindTabId);
 	InTabManager->UnregisterTabSpawner(FMaterialEditorTabs::PreviewSettingsTabId);
 	InTabManager->UnregisterTabSpawner(FMaterialEditorTabs::ParameterDefaultsTabId);
@@ -968,6 +976,12 @@ FMaterialEditor::~FMaterialEditor()
 		ClearNumericParameterDefaultOnDependentMaterials(It.Key, It.Value);
 	}
 
+	if (OnGraphChangedDelegateHandle.IsValid() && Material && Material->MaterialGraph)
+	{
+		Material->MaterialGraph->RemoveOnGraphChangedHandler(OnGraphChangedDelegateHandle);
+		OnGraphChangedDelegateHandle.Reset();
+	}
+
 	// Unregister this delegate
 	FEditorSupportDelegates::MaterialUsageFlagsChanged.RemoveAll(this);
 	FEditorSupportDelegates::NumericParameterDefaultChanged.RemoveAll(this);
@@ -1719,6 +1733,12 @@ void FMaterialEditor::CreateInternalWidgets()
 		.InGenerator(Generator);
 
 	Palette = SNew(SMaterialPalette, SharedThis(this));
+	NamedReroutesPanel = SNew(SMaterialNamedReroutesPanel, SharedThis(this));
+
+	if (Material && Material->MaterialGraph)
+	{
+		OnGraphChangedDelegateHandle = Material->MaterialGraph->AddOnGraphChangedHandler(FOnGraphChanged::FDelegate::CreateSP(this, &FMaterialEditor::OnGraphChanged));
+	}
 
 	FMessageLogModule& MessageLogModule = FModuleManager::LoadModuleChecked<FMessageLogModule>("MessageLog");
 	FMessageLogInitializationOptions LogOptions;
@@ -1796,6 +1816,28 @@ void FMaterialEditor::OnFinishedChangingProperties(const FPropertyChangedEvent&
 		UpdatePreviewMaterial();
 		SetEditorInstanceForWidgets();
 	}
+
+	// Refresh the Named Reroutes panel when a declaration's name or color changes
+	if (NamedReroutesPanel.IsValid())
+	{
+		for (int32 Index = 0; Index < PropertyChangedEvent.GetNumObjectsBeingEdited(); ++Index)
+		{
+			if (PropertyChangedEvent.GetObjectBeingEdited(Index)->IsA<UMaterialExpressionNamedRerouteDeclaration>())
+			{
+				NamedReroutesPanel->RequestRefresh();
+				break;
+			}
+		}
+	}
+}
+
+void FMaterialEditor::OnGraphChanged(const FEdGraphEditAction& Action)
+{
+	// This will handle adding nodes, removing nodes, and duplicating nodes.
+	if (NamedReroutesPanel.IsValid())
+	{
+		NamedReroutesPanel->RequestRefresh();
+	}
 }
 
 void FMaterialEditor::OnFinishedChangingParametersFromOverview(const FPropertyChangedEvent& PropertyChangedEvent)
@@ -4377,6 +4419,36 @@ void FMaterialEditor::OnSelectNamedRerouteDeclaration()
 	}
 }
 
+void FMaterialEditor::SelectNamedRerouteUsages(UMaterialExpressionNamedRerouteDeclaration* Declaration)
+{
+	if (!Declaration)
+	{
+		return;
+	}
+
+	TArray<UEdGraphNode*> FoundGraphNodes;
+	for (UMaterialExpression* Expression : Material->GetExpressions())
+	{
+		UMaterialExpressionNamedRerouteUsage* Usage = Cast<UMaterialExpressionNamedRerouteUsage>(Expression);
+		if (Usage && Usage->Declaration == Declaration)
+		{
+			if (UEdGraphNode* UsageGraphNode = Usage->GraphNode)
+			{
+				FoundGraphNodes.Add(UsageGraphNode);
+			}
+		}
+	}
+
+	if (FoundGraphNodes.Num())
+	{
+		// Spawn the tab in case the user doesn't have it open
+		TabManager->TryInvokeTab(FMaterialEditorTabs::FindTabId);
+
+		FindResults->PopulateSearchItems(FoundGraphNodes);
+		FindResults->FocusForUse();
+	}
+}
+
 void FMaterialEditor::OnSelectNamedRerouteUsages()
 {
 	if (TSharedPtr<SGraphEditor> FocusedGraphEd = FocusedGraphEdPtr.Pin())
@@ -4387,33 +4459,11 @@ void FMaterialEditor::OnSelectNamedRerouteUsages()
 			FocusedGraphEd->ClearSelectionSet();
 			for (FGraphPanelSelectionSet::TConstIterator NodeIt(SelectedNodes); NodeIt; ++NodeIt)
 			{
-				UMaterialGraphNode* GraphNode = Cast<UMaterialGraphNode>(*NodeIt);
-				if (GraphNode)
+				if (UMaterialGraphNode* GraphNode = Cast<UMaterialGraphNode>(*NodeIt))
 				{
-					TArray<UEdGraphNode*> FoundGraphNodes;
-
-					UMaterialExpression* CurrentSelectedExpression = GraphNode->MaterialExpression;
-					UMaterialExpressionNamedRerouteDeclaration* Declaration = Cast<UMaterialExpressionNamedRerouteDeclaration>(CurrentSelectedExpression);
-					for(UMaterialExpression* Expression : Material->GetExpressions())
-					{
-						auto* Usage = Cast<UMaterialExpressionNamedRerouteUsage>(Expression);
-						if (Usage && Usage->Declaration == Declaration)
-						{
-							UEdGraphNode* UsageGraphNode = Usage->GraphNode;
-							if (UsageGraphNode)
-							{
-								FoundGraphNodes.Add(UsageGraphNode);
-							}
-						}
-					}
-					
-					if (FoundGraphNodes.Num())
+					if (UMaterialExpressionNamedRerouteDeclaration* Declaration = Cast<UMaterialExpressionNamedRerouteDeclaration>(GraphNode->MaterialExpression))
 					{
-						// Spawn the tab in case the user doesn't have it open
-						TabManager->TryInvokeTab(FMaterialEditorTabs::FindTabId);
-
-						FindResults->PopulateSearchItems(FoundGraphNodes);
-						FindResults->FocusForUse();
+						SelectNamedRerouteUsages(Declaration);
 					}
 				}
 			}
@@ -5551,6 +5601,23 @@ TSharedRef<SDockTab> FMaterialEditor::SpawnTab_Palette(const FSpawnTabArgs& Args
 	return SpawnedTab;
 }
 
+TSharedRef<SDockTab> FMaterialEditor::SpawnTab_NamedReroutes(const FSpawnTabArgs& Args)
+{
+	check(Args.GetTabId() == FMaterialEditorTabs::NamedReroutesTabId);
+
+	TSharedRef<SDockTab> SpawnedTab = SNew(SDockTab)
+		.Label(LOCTEXT("MaterialNamedReroutesTitle", "Named Reroutes"))
+		[
+			SNew(SBox)
+			.AddMetaData<FTagMetaData>(FTagMetaData(TEXT("MaterialNamedReroutes")))
+			[
+				NamedReroutesPanel.ToSharedRef()
+			]
+		];
+
+	return SpawnedTab;
+}
+
 TSharedRef<SDockTab> FMaterialEditor::SpawnTab_Find(const FSpawnTabArgs& Args)
 {
 	check(Args.GetTabId() == FMaterialEditorTabs::FindTabId);
@@ -8514,6 +8581,11 @@ void FMaterialEditor::OnNodeTitleCommitted(const FText& NewText, ETextCommit::Ty
 		NodeBeingChanged->OnRenameNode(NewName);
 		UpdateGenerator();
 		SetEditorInstanceForWidgets();
+
+		if (NamedReroutesPanel.IsValid())
+		{
+			NamedReroutesPanel->RequestRefresh();
+		}
 	}
 }
 
diff --git a/Engine/Source/Editor/MaterialEditor/Private/MaterialEditor.h b/Engine/Source/Editor/MaterialEditor/Private/MaterialEditor.h
--- a/Engine/Source/Editor/MaterialEditor/Private/MaterialEditor.h
+++ b/Engine/Source/Editor/MaterialEditor/Private/MaterialEditor.h
@@ -26,6 +26,7 @@
 #include "UObject/WeakFieldPtr.h"
 
 struct FAssetData;
+struct FEdGraphEditAction;
 struct FToolMenuSection;
 class FCanvas;
 class FMaterialCompiler;
@@ -40,6 +41,7 @@ class UFactory;
 class UMaterialEditorOptions;
 class UMaterialExpressionComment;
 class UMaterialExpressionComposite;
+class UMaterialExpressionNamedRerouteDeclaration;
 class UMaterialInstance;
 class UMaterialGraphNode;
 struct FGraphAppearanceInfo;
@@ -439,6 +441,9 @@ public:
 	static void NotifyAllUserSceneTextureLoadOrUnload(UMaterialInterface* MaterialInterface);
 
 public:
+	/** Find all usages of the given Named Reroute Declaration and show them in the Find Results tab */
+	void SelectNamedRerouteUsages(UMaterialExpressionNamedRerouteDeclaration* Declaration);
+
 	/** Set to true when modifications have been made to the material */
 	bool bMaterialDirty;
 
@@ -1013,6 +1018,7 @@ private:
 	TSharedRef<SDockTab> SpawnTab_Preview(const FSpawnTabArgs& Args);
 	TSharedRef<SDockTab> SpawnTab_MaterialProperties(const FSpawnTabArgs& Args);
 	TSharedRef<SDockTab> SpawnTab_Palette(const FSpawnTabArgs& Args);
+	TSharedRef<SDockTab> SpawnTab_NamedReroutes(const FSpawnTabArgs& Args);
 	TSharedRef<SDockTab> SpawnTab_Find(const FSpawnTabArgs& Args);
 	TSharedRef<SDockTab> SpawnTab_PreviewSettings(const FSpawnTabArgs& Args);
 	TSharedRef<SDockTab> SpawnTab_ParameterDefaults(const FSpawnTabArgs& Args);
@@ -1022,6 +1028,7 @@ private:
 	TSharedRef<SDockTab> SpawnTab_MaterialExpressionsOverview(const FSpawnTabArgs& Args);
 
 	void OnFinishedChangingProperties(const FPropertyChangedEvent& PropertyChangedEvent);
+	void OnGraphChanged(const FEdGraphEditAction& Action);
 	void OnFinishedChangingParametersFromOverview(const FPropertyChangedEvent& PropertyChangedEvent);
 	void OnChangeBreadCrumbGraph(class UEdGraph* InGraph);
 	void GeneratorRowsRefreshed();
@@ -1046,6 +1053,12 @@ private:
 	/** Palette of Material Expressions and functions */
 	TSharedPtr<class SMaterialPalette> Palette;
 
+	/** Panel listing Named Reroute Declaration nodes */
+	TSharedPtr<class SMaterialNamedReroutesPanel> NamedReroutesPanel;
+
+	/** Handle for the material graph's OnGraphChanged delegate (used to refresh the Named Reroutes panel) */
+	FDelegateHandle OnGraphChangedDelegateHandle;
+
 	/** The Substrate control tab */
 	TSharedPtr<class SMaterialEditorSubstrateWidget> SubstrateWidget;
 
diff --git a/Engine/Source/Editor/MaterialEditor/Private/MaterialEditorModes.cpp b/Engine/Source/Editor/MaterialEditor/Private/MaterialEditorModes.cpp
--- a/Engine/Source/Editor/MaterialEditor/Private/MaterialEditorModes.cpp
+++ b/Engine/Source/Editor/MaterialEditor/Private/MaterialEditorModes.cpp
@@ -19,7 +19,7 @@ const FName FMaterialEditorApplicationModes::StandardMaterialEditorMode(TEXT("Gr
 
 TSharedPtr<FTabManager::FLayout> FMaterialEditorApplicationModes::GetDefaultEditorLayout(TSharedPtr<FMaterialEditor> InMaterialEditor)
 {
-	const TSharedRef<FTabManager::FLayout> StandaloneDefaultLayout = FTabManager::NewLayout("Standalone_MaterialEditor_Layout_v16")
+	const TSharedRef<FTabManager::FLayout> StandaloneDefaultLayout = FTabManager::NewLayout("Standalone_MaterialEditor_Layout_v17")
 	->AddArea
 	(
 		FTabManager::NewPrimaryArea() ->SetOrientation(Orient_Vertical)
@@ -75,6 +75,7 @@ TSharedPtr<FTabManager::FLayout> FMaterialEditorApplicationModes::GetDefaultEdit
 				(
 					FTabManager::NewStack()
 					->AddTab( FMaterialEditorTabs::PaletteTabId, ETabState::SidebarTab, ESidebarLocation::Right, 0.14f )
+					->AddTab( FMaterialEditorTabs::NamedReroutesTabId, ETabState::SidebarTab, ESidebarLocation::Right, 0.14f )
 				)
 			)
 		)
diff --git a/Engine/Source/Editor/MaterialEditor/Private/MaterialEditorTabs.cpp b/Engine/Source/Editor/MaterialEditor/Private/MaterialEditorTabs.cpp
--- a/Engine/Source/Editor/MaterialEditor/Private/MaterialEditorTabs.cpp
+++ b/Engine/Source/Editor/MaterialEditor/Private/MaterialEditorTabs.cpp
@@ -16,6 +16,7 @@ const FName FMaterialEditorTabs::CustomPrimitiveTabId(TEXT("MaterialEditor_Custo
 const FName FMaterialEditorTabs::LayerPropertiesTabId(TEXT("MaterialInstanceEditor_MaterialLayerProperties"));
 const FName FMaterialEditorTabs::SubstrateTabId(TEXT("MaterialEditor_Substrate"));
 const FName FMaterialEditorTabs::MaterialOverviewTabId(TEXT("MaterialEditor_MaterialExpressionOverview"));
+const FName FMaterialEditorTabs::NamedReroutesTabId(TEXT("MaterialEditor_NamedReroutes"));
 
 
 const FName FMaterialInstanceEditorTabs::PreviewTabId( TEXT( "MaterialInstanceEditor_Preview" ) );
diff --git a/Engine/Source/Editor/MaterialEditor/Private/MaterialEditorTabs.h b/Engine/Source/Editor/MaterialEditor/Private/MaterialEditorTabs.h
--- a/Engine/Source/Editor/MaterialEditor/Private/MaterialEditorTabs.h
+++ b/Engine/Source/Editor/MaterialEditor/Private/MaterialEditorTabs.h
@@ -19,6 +19,7 @@ struct FMaterialEditorTabs
 	static const FName LayerPropertiesTabId;
 	static const FName SubstrateTabId;
 	static const FName MaterialOverviewTabId;
+	static const FName NamedReroutesTabId;
 
 	// Disable default constructor
 	FMaterialEditorTabs() = delete;
diff --git a/Engine/Source/Editor/MaterialEditor/Private/SMaterialNamedReroutesPanel.cpp b/Engine/Source/Editor/MaterialEditor/Private/SMaterialNamedReroutesPanel.cpp
new file mode 100644
--- /dev/null
+++ b/Engine/Source/Editor/MaterialEditor/Private/SMaterialNamedReroutesPanel.cpp
@@ -0,0 +1,194 @@
+// Copyright Epic Games, Inc. All Rights Reserved.
+
+#include "SMaterialNamedReroutesPanel.h"
+#include "MaterialEditor.h"
+#include "MaterialGraph/MaterialGraph.h"
+#include "MaterialGraph/MaterialGraphSchema.h"
+#include "Materials/MaterialExpressionNamedReroute.h"
+#include "Styling/AppStyle.h"
+#include "Framework/MultiBox/MultiBoxBuilder.h"
+#include "Widgets/SBoxPanel.h"
+#include "Widgets/Colors/SColorBlock.h"
+
+#define LOCTEXT_NAMESPACE "MaterialNamedReroutesPanel"
+
+//////////////////////////////////////////////////////////////////////////
+// SMaterialNamedRerouteItem
+
+void SMaterialNamedRerouteItem::Construct(const FArguments& InArgs, FCreateWidgetForActionData* const InCreateData)
+{
+	check(InCreateData->Action.IsValid());
+
+	ActionPtr = InCreateData->Action;
+
+	// Try to get the node color from the declaration
+	FLinearColor SwatchColor = FLinearColor::White;
+	if (InCreateData->Action->GetTypeId() == FMaterialGraphSchemaAction_NewNamedRerouteUsage::StaticGetTypeId())
+	{
+		const TSharedPtr<FMaterialGraphSchemaAction_NewNamedRerouteUsage> RerouteAction = StaticCastSharedPtr<FMaterialGraphSchemaAction_NewNamedRerouteUsage>(InCreateData->Action);
+		if (RerouteAction->Declaration)
+		{
+			SwatchColor = RerouteAction->Declaration->NodeColor;
+		}
+	}
+
+	TSharedRef<SWidget> NameSlotWidget = CreateTextSlotWidget(InCreateData, false);
+
+	this->ChildSlot
+	[
+		SNew(SHorizontalBox)
+		// Color swatch
+		+SHorizontalBox::Slot()
+		.AutoWidth()
+		.VAlign(VAlign_Center)
+		.Padding(2, 0)
+		[
+			SNew(SColorBlock)
+			.Color(SwatchColor)
+			.Size(FVector2D(12.0f, 12.0f))
+			.CornerRadius(FVector4(2.0f, 2.0f, 2.0f, 2.0f))
+		]
+		// Name slot
+		+SHorizontalBox::Slot()
+		.FillWidth(1.f)
+		.VAlign(VAlign_Center)
+		.Padding(3, 0)
+		[
+			NameSlotWidget
+		]
+	];
+}
+
+FText SMaterialNamedRerouteItem::GetItemTooltip() const
+{
+	return ActionPtr.IsValid() ? ActionPtr.Pin()->GetTooltipDescription() : FText::GetEmpty();
+}
+
+//////////////////////////////////////////////////////////////////////////
+// SMaterialNamedReroutesPanel
+
+void SMaterialNamedReroutesPanel::Construct(const FArguments& InArgs, TWeakPtr<FMaterialEditor> InMaterialEditorPtr)
+{
+	MaterialEditorPtr = InMaterialEditorPtr;
+
+	this->ChildSlot
+	[
+		SNew(SBorder)
+		.Padding(2.0f)
+		.BorderImage(FAppStyle::GetBrush("ToolPanel.GroupBorder"))
+		[
+			SNew(SVerticalBox)
+			+SVerticalBox::Slot()
+			[
+				SAssignNew(GraphActionMenu, SGraphActionMenu)
+				.OnActionDragged(this, &SMaterialNamedReroutesPanel::OnActionDragged)
+				.OnCreateWidgetForAction(this, &SMaterialNamedReroutesPanel::OnCreateWidgetForAction)
+				.OnCollectAllActions(this, &SMaterialNamedReroutesPanel::CollectAllActions)
+				.OnActionDoubleClicked(this, &SMaterialNamedReroutesPanel::OnActionDoubleClicked)
+				.OnContextMenuOpening(this, &SMaterialNamedReroutesPanel::OnContextMenuOpening)
+				.AutoExpandActionMenu(true)
+			]
+		]
+	];
+}
+
+TSharedRef<SWidget> SMaterialNamedReroutesPanel::OnCreateWidgetForAction(FCreateWidgetForActionData* const InCreateData)
+{
+	return SNew(SMaterialNamedRerouteItem, InCreateData);
+}
+
+void SMaterialNamedReroutesPanel::CollectAllActions(FGraphActionListBuilderBase& OutAllActions)
+{
+	TSharedPtr<FMaterialEditor> MaterialEditor = MaterialEditorPtr.Pin();
+	if (!MaterialEditor.IsValid())
+	{
+		return;
+	}
+
+	UMaterial* Material = MaterialEditor->Material;
+	if (!Material || !Material->MaterialGraph)
+	{
+		return;
+	}
+
+	UMaterialGraphSchema::GetNamedRerouteActionsForGraph(OutAllActions, Material->MaterialGraph);
+}
+
+void SMaterialNamedReroutesPanel::OnActionDoubleClicked(const TArray<TSharedPtr<FEdGraphSchemaAction>>& InActions)
+{
+	TSharedPtr<FMaterialEditor> MaterialEditor = MaterialEditorPtr.Pin();
+	if (!MaterialEditor.IsValid())
+	{
+		return;
+	}
+
+	for (const TSharedPtr<FEdGraphSchemaAction>& Action : InActions)
+	{
+		if (Action.IsValid() && Action->GetTypeId() == FMaterialGraphSchemaAction_NewNamedRerouteUsage::StaticGetTypeId())
+		{
+			const TSharedPtr<FMaterialGraphSchemaAction_NewNamedRerouteUsage> RerouteAction = StaticCastSharedPtr<FMaterialGraphSchemaAction_NewNamedRerouteUsage>(Action);
+			if (RerouteAction->Declaration && RerouteAction->Declaration->GraphNode)
+			{
+				MaterialEditor->JumpToNode(RerouteAction->Declaration->GraphNode);
+			}
+		}
+	}
+}
+
+void SMaterialNamedReroutesPanel::RequestRefresh()
+{
+	RefreshActionsList(true);
+}
+
+TSharedPtr<SWidget> SMaterialNamedReroutesPanel::OnContextMenuOpening()
+{
+	TSharedPtr<FMaterialEditor> MaterialEditor = MaterialEditorPtr.Pin();
+	if (!MaterialEditor.IsValid())
+	{
+		return nullptr;
+	}
+
+	TArray<TSharedPtr<FEdGraphSchemaAction>> SelectedActions;
+	GraphActionMenu->GetSelectedActions(SelectedActions);
+
+	if (SelectedActions.Num() != 1 || !SelectedActions[0].IsValid())
+	{
+		return nullptr;
+	}
+
+	if (SelectedActions[0]->GetTypeId() != FMaterialGraphSchemaAction_NewNamedRerouteUsage::StaticGetTypeId())
+	{
+		return nullptr;
+	}
+
+	TSharedPtr<FMaterialGraphSchemaAction_NewNamedRerouteUsage> RerouteAction = StaticCastSharedPtr<FMaterialGraphSchemaAction_NewNamedRerouteUsage>(SelectedActions[0]);
+
+	if (!RerouteAction->Declaration)
+	{
+		return nullptr;
+	}
+
+	TWeakObjectPtr<UMaterialExpressionNamedRerouteDeclaration> WeakDeclaration = RerouteAction->Declaration;
+	TWeakPtr<FMaterialEditor> WeakEditor = MaterialEditorPtr;
+
+	FMenuBuilder MenuBuilder(/*bInShouldCloseWindowAfterMenuSelection=*/ true, /*InCommandList=*/ nullptr);
+	MenuBuilder.AddMenuEntry(
+		LOCTEXT("SelectNamedRerouteUsages", "Select Named Reroute Usages"),
+		LOCTEXT("SelectNamedRerouteUsagesTooltip", "Find all usage nodes of this named reroute and show them in Find Results"),
+		FSlateIcon(),
+		FUIAction(FExecuteAction::CreateLambda([WeakEditor, WeakDeclaration]()
+		{
+			if (TSharedPtr<FMaterialEditor> Editor = WeakEditor.Pin())
+			{
+				if (UMaterialExpressionNamedRerouteDeclaration* Declaration = WeakDeclaration.Get())
+				{
+					Editor->SelectNamedRerouteUsages(Declaration);
+				}
+			}
+		}))
+	);
+
+	return MenuBuilder.MakeWidget();
+}
+
+#undef LOCTEXT_NAMESPACE
diff --git a/Engine/Source/Editor/MaterialEditor/Private/SMaterialNamedReroutesPanel.h b/Engine/Source/Editor/MaterialEditor/Private/SMaterialNamedReroutesPanel.h
new file mode 100644
--- /dev/null
+++ b/Engine/Source/Editor/MaterialEditor/Private/SMaterialNamedReroutesPanel.h
@@ -0,0 +1,53 @@
+// Copyright Epic Games, Inc. All Rights Reserved.
+
+#pragma once
+
+#include "CoreMinimal.h"
+#include "Widgets/DeclarativeSyntaxSupport.h"
+#include "SGraphPalette.h"
+
+class FMaterialEditor;
+
+/** Widget for displaying a single named reroute item with a color swatch */
+class SMaterialNamedRerouteItem : public SGraphPaletteItem
+{
+public:
+	SLATE_BEGIN_ARGS(SMaterialNamedRerouteItem) {};
+	SLATE_END_ARGS()
+
+	void Construct(const FArguments& InArgs, FCreateWidgetForActionData* const InCreateData);
+
+private:
+	virtual FText GetItemTooltip() const override;
+};
+
+//////////////////////////////////////////////////////////////////////////
+
+/** Panel listing all Named Reroute Declaration nodes in the current material */
+class SMaterialNamedReroutesPanel : public SGraphPalette
+{
+public:
+	SLATE_BEGIN_ARGS(SMaterialNamedReroutesPanel) {};
+	SLATE_END_ARGS()
+
+	void Construct(const FArguments& InArgs, TWeakPtr<FMaterialEditor> InMaterialEditorPtr);
+
+	/** Force a full refresh of the list (e.g. when declarations are added/removed or properties change) */
+	void RequestRefresh();
+
+protected:
+	// SGraphPalette Interface
+	virtual TSharedRef<SWidget> OnCreateWidgetForAction(FCreateWidgetForActionData* const InCreateData) override;
+	virtual void CollectAllActions(FGraphActionListBuilderBase& OutAllActions) override;
+	// End of SGraphPalette Interface
+
+	/** Called when an action is selected in the list - navigates to the declaration node */
+	void OnActionDoubleClicked(const TArray<TSharedPtr<FEdGraphSchemaAction>>& InActions);
+
+private:
+	/** Creates the right-click context menu for panel items */
+	TSharedPtr<SWidget> OnContextMenuOpening();
+
+	/** Pointer back to the material editor that owns us */
+	TWeakPtr<FMaterialEditor> MaterialEditorPtr;
+};
diff --git a/Engine/Source/Editor/UnrealEd/Classes/MaterialGraph/MaterialGraphSchema.h b/Engine/Source/Editor/UnrealEd/Classes/MaterialGraph/MaterialGraphSchema.h
--- a/Engine/Source/Editor/UnrealEd/Classes/MaterialGraph/MaterialGraphSchema.h
+++ b/Engine/Source/Editor/UnrealEd/Classes/MaterialGraph/MaterialGraphSchema.h
@@ -221,6 +221,9 @@ class UMaterialGraphSchema : public UEdGraphSchema
 	UNREALED_API static const FLinearColor InactivePinColor;
 	UNREALED_API static const FLinearColor AlphaPinColor;
 
+	/** Gathers Named Reroute Declaration actions from the given graph into the action list */
+	UNREALED_API static void GetNamedRerouteActionsForGraph(FGraphActionListBuilderBase& OutAllActions, const UEdGraph* Graph);
+
 	/**
 	 * Connect a pin to one of the Material Function's outputs
 	 *
diff --git a/Engine/Source/Editor/UnrealEd/Private/MaterialGraphSchema.cpp b/Engine/Source/Editor/UnrealEd/Private/MaterialGraphSchema.cpp
--- a/Engine/Source/Editor/UnrealEd/Private/MaterialGraphSchema.cpp
+++ b/Engine/Source/Editor/UnrealEd/Private/MaterialGraphSchema.cpp
@@ -1108,7 +1108,7 @@ void UMaterialGraphSchema::GetCommentAction(FGraphActionMenuBuilder& ActionMenuB
 	}
 }
 
-void UMaterialGraphSchema::GetNamedRerouteActions(FGraphActionMenuBuilder& ActionMenuBuilder, const UEdGraph* CurrentGraph) const
+void UMaterialGraphSchema::GetNamedRerouteActionsForGraph(FGraphActionListBuilderBase& OutAllActions, const UEdGraph* CurrentGraph)
 {
 	if (CurrentGraph)
 	{
@@ -1123,13 +1123,18 @@ void UMaterialGraphSchema::GetNamedRerouteActions(FGraphActionMenuBuilder& Actio
 					const FText Tooltip = FText::Format(LOCTEXT("NamedRerouteTooltip", "Add a usage of {0} here"), Name);
 					TSharedPtr<FMaterialGraphSchemaAction_NewNamedRerouteUsage> NewAction(new FMaterialGraphSchemaAction_NewNamedRerouteUsage(Category, Name, Tooltip, 1 /* We want named reroutes to be on top */));
 					NewAction->Declaration = Declaration;
-					ActionMenuBuilder.AddAction(NewAction);
+					OutAllActions.AddAction(NewAction);
 				}
 			}
 		}
 	}
 }
 
+void UMaterialGraphSchema::GetNamedRerouteActions(FGraphActionMenuBuilder& ActionMenuBuilder, const UEdGraph* CurrentGraph) const
+{
+	GetNamedRerouteActionsForGraph(ActionMenuBuilder, CurrentGraph);
+}
+
 bool UMaterialGraphSchema::HasCompatibleConnection(const FAssetData& FunctionAssetData, EMaterialValueType TestType, EEdGraphPinDirection TestDirection) const
 {
 	if (TestType != 0)
