Add Material Instance graph diffing when MIs share the same parent and one of them has a static parameter set.

- This allows you to diff the graph between two material instances due to static switch changes.
- Refactoring to unify diffing of UMaterials and MaterialInstances to just use the MaterialDiffPanel and SMaterialDiff classes
- Remove the instance only versions of these files.
- This works by looking at the parent of each MI, making a UMaterialGraph, and setting the static switch values.  This graph is what we end up diffing as-if it was a UMaterial we are diffing.
- Everything else about the diff tool behaves the same.
- Some of the code is copy/paste from the deleted Instance versions of the files (CreateDiffWindow functions, create detail customization, and some of the lambdas for making editor only versions of the MaterialInstances)

todo: handle non-static switch parameters (quality, RT, PT, etc.)

#Synced-CL 51833860

diff --git a/Engine/Plugins/Editor/EngineAssetDefinitions/Source/Private/Material/AssetDefinition_MaterialInstanceConstant.cpp b/Engine/Plugins/Editor/EngineAssetDefinitions/Source/Private/Material/AssetDefinition_MaterialInstanceConstant.cpp
--- a/Engine/Plugins/Editor/EngineAssetDefinitions/Source/Private/Material/AssetDefinition_MaterialInstanceConstant.cpp
+++ b/Engine/Plugins/Editor/EngineAssetDefinitions/Source/Private/Material/AssetDefinition_MaterialInstanceConstant.cpp
@@ -8,7 +8,7 @@
 #include "MaterialEditorModule.h"
 #include "ToolMenu.h"
 #include "ToolMenuSection.h"
-#include "DiffTool/Widgets/SMaterialInstanceDiff.h"
+#include "DiffTool/Widgets/SMaterialDiff.h"
 
 #include UE_INLINE_GENERATED_CPP_BY_NAME(AssetDefinition_MaterialInstanceConstant)
 
@@ -32,7 +32,7 @@ EAssetCommandResult UAssetDefinition_MaterialInstanceConstant::PerformAssetDiff(
 
 	if (OldMaterialInstance && NewMaterialInstance)
 	{
-		SMaterialInstanceDiff::CreateDiffWindow(OldMaterialInstance, NewMaterialInstance, DiffArgs.OldRevision, DiffArgs.NewRevision, GetAssetClass().Get());
+		SMaterialDiff::CreateDiffWindow(OldMaterialInstance, NewMaterialInstance, DiffArgs.OldRevision, DiffArgs.NewRevision, GetAssetClass().Get());
 
 		return EAssetCommandResult::Handled;
 	}
diff --git a/Engine/Source/Editor/MaterialEditor/Internal/DiffTool/MaterialDiffPanel.h b/Engine/Source/Editor/MaterialEditor/Internal/DiffTool/MaterialDiffPanel.h
--- a/Engine/Source/Editor/MaterialEditor/Internal/DiffTool/MaterialDiffPanel.h
+++ b/Engine/Source/Editor/MaterialEditor/Internal/DiffTool/MaterialDiffPanel.h
@@ -16,6 +16,9 @@ class UMaterial;
 class UMaterialExpression;
 class UMaterialGraph;
 
+class UMaterialInstance;
+class UMaterialEditorInstanceConstant;
+
 struct FRevisionInfo;
 
 /** Panel used to display the MaterialGraph */
@@ -90,6 +93,12 @@ struct FMaterialDiffPanel : public FGCObject
 
 	/** The property that we will highlight in MaterialNodeDetailsView */
 	FPropertyPath PropertyToFocus;
+
+	/** The material instance being diffed (nullptr when diffing base materials) */
+	TObjectPtr<UMaterialInstance> MaterialInstance = nullptr;
+
+	/** Editor proxy used to display all parameter types including static switches (nullptr when diffing base materials) */
+	TObjectPtr<UMaterialEditorInstanceConstant> EditorProxy = nullptr;
 	
 	// Begin FGCObject interface
 	virtual void AddReferencedObjects(FReferenceCollector& Collector) override;
diff --git a/Engine/Source/Editor/MaterialEditor/Internal/DiffTool/MaterialInstanceDiffPanel.h b/Engine/Source/Editor/MaterialEditor/Internal/DiffTool/MaterialInstanceDiffPanel.h
deleted file mode 100644
diff --git a/Engine/Source/Editor/MaterialEditor/Internal/DiffTool/Widgets/SMaterialDiff.h b/Engine/Source/Editor/MaterialEditor/Internal/DiffTool/Widgets/SMaterialDiff.h
--- a/Engine/Source/Editor/MaterialEditor/Internal/DiffTool/Widgets/SMaterialDiff.h
+++ b/Engine/Source/Editor/MaterialEditor/Internal/DiffTool/Widgets/SMaterialDiff.h
@@ -8,6 +8,7 @@
 #include "DiffTool/MaterialDiffPanel.h"
 #include "SBlueprintDiff.h"
 #include "Widgets/SCompoundWidget.h"
+#include "Materials/MaterialInstance.h"
 
 #define UE_API MATERIALEDITOR_API
 
@@ -15,6 +16,9 @@ class FBlueprintDifferenceTreeEntry;
 class UEdGraph;
 class UMaterialGraph;
 
+class IDetailsView;
+class UMaterialEditorInstanceConstant;
+
 struct FMaterialToDiff;
 
 enum class EAssetEditorCloseReason : uint8;
@@ -35,6 +39,8 @@ public:
 		SLATE_ARGUMENT(FRevisionInfo, OldRevision)
 		SLATE_ARGUMENT(FRevisionInfo, NewRevision)
 		SLATE_ARGUMENT(bool, ShowAssetNames)
+		SLATE_ARGUMENT(TObjectPtr<UMaterialInstance>, OldMaterialInstance)
+		SLATE_ARGUMENT(TObjectPtr<UMaterialInstance>, NewMaterialInstance)
 		SLATE_ARGUMENT(TSharedPtr<SWindow>, ParentWindow)
 	SLATE_END_ARGS()
 
@@ -59,6 +65,12 @@ public:
 	/** Helper function to create a window that holds a diff widget (default window title) */
 	static UE_API TSharedPtr<SWindow> CreateDiffWindow(TObjectPtr<UMaterialGraph> OldMaterialGraph, TObjectPtr<UMaterialGraph> NewMaterialGraph, const FRevisionInfo& OldRevision, const FRevisionInfo& NewRevision, const UClass* ObjectClass);
 
+	/** Create a diff window for two material instances */
+	static UE_API TSharedPtr<SWindow> CreateDiffWindow(FText WindowTitle, TObjectPtr<UMaterialInstance> OldMaterialInstance, TObjectPtr<UMaterialInstance> NewMaterialInstance, const FRevisionInfo& OldRevision, const FRevisionInfo& NewRevision);
+
+	/** Create a diff window for two material instances (default window title) */
+	static UE_API TSharedPtr<SWindow> CreateDiffWindow(TObjectPtr<UMaterialInstance> OldMaterialInstance, TObjectPtr<UMaterialInstance> NewMaterialInstance, const FRevisionInfo& OldRevision, const FRevisionInfo& NewRevision, const UClass* ObjectClass);
+
 	DECLARE_MULTICAST_DELEGATE_OneParam(FOnWindowClosedEvent, TSharedRef<SMaterialDiff>)
 	FOnWindowClosedEvent OnWindowClosedEvent;
 
@@ -155,6 +167,9 @@ protected:
 	/** If the preview Viewports should be shown */
 	bool bShowViewport = true;
 
+	/** Should we show the graph diff? */
+	bool bShowGraphDiff = true;
+
 	/** Contents widget that we swap when mode changes (defaults, components, etc) */
 	TSharedPtr<SBox> ModeContents;
 
@@ -184,6 +199,10 @@ protected:
 	TWeakPtr<SWindow> WeakParentWindow;
 
 	FDelegateHandle AssetEditorCloseDelegate;
+
+private:
+	/** Registers FMaterialInstanceParameterDetails customization on instance diff details views */
+	void AddMaterialInstanceDetailCustomization(const TSharedRef<IDetailsView>& DetailsView);
 };
 
 #undef UE_API
diff --git a/Engine/Source/Editor/MaterialEditor/Internal/DiffTool/Widgets/SMaterialInstanceDiff.h b/Engine/Source/Editor/MaterialEditor/Internal/DiffTool/Widgets/SMaterialInstanceDiff.h
deleted file mode 100644
diff --git a/Engine/Source/Editor/MaterialEditor/Private/DiffTool/MaterialDiffPanel.cpp b/Engine/Source/Editor/MaterialEditor/Private/DiffTool/MaterialDiffPanel.cpp
--- a/Engine/Source/Editor/MaterialEditor/Private/DiffTool/MaterialDiffPanel.cpp
+++ b/Engine/Source/Editor/MaterialEditor/Private/DiffTool/MaterialDiffPanel.cpp
@@ -15,6 +15,7 @@
 #include "MaterialGraph/MaterialGraphNode.h"
 #include "MaterialGraph/MaterialGraphNode_Comment.h"
 #include "Materials/MaterialExpression.h"
+#include "Materials/MaterialInstance.h"
 #include "Materials/MaterialExpressionBreakMaterialAttributes.h"
 #include "Materials/MaterialExpressionFunctionOutput.h"
 #include "Materials/MaterialExpressionStaticBool.h"
@@ -368,18 +369,30 @@ TSharedRef<SWidget> FMaterialDiffPanel::GetMaterialNodeDetailsViewWidget() const
 
 void FMaterialDiffPanel::SetViewportToDisplay()
 {
-	if (!MaterialGraph)
+	// Determine the preview target: MaterialInstance takes priority over MaterialGraph->Material
+	UMaterialInterface* PreviewTarget = nullptr;
+
+	if (MaterialInstance)
+	{
+		PreviewTarget = MaterialInstance;
+	}
+	else if (MaterialGraph)
+	{
+		PreviewTarget = MaterialGraph->Material;
+	}
+
+	if (!PreviewTarget)
 	{
 		return;
 	}
 
-	if (MaterialGraph->Material && MaterialGraph->Material->IsUIMaterial())
+	if (PreviewTarget->IsUIMaterial())
 	{
-		SAssignNew(Preview2DViewport, SMaterialEditorUIPreviewViewport, MaterialGraph->Material);
+		SAssignNew(Preview2DViewport, SMaterialEditorUIPreviewViewport, PreviewTarget);
 	}
 
 	SAssignNew(Preview3DViewport, SMaterialEditor3DPreviewViewport)
-		.PreviewMaterial(MaterialGraph->Material);
+		.PreviewMaterial(PreviewTarget);
 
 	if (Preview3DViewport->PreviewMeshComponent != nullptr)
 	{
@@ -432,11 +445,6 @@ TSharedRef<SWidget> FMaterialDiffPanel::GetViewportToDisplay() const
 
 void FMaterialDiffPanel::SetViewportVisibility(bool bShowViewport)
 {
-	if (!MaterialGraph)
-	{
-		return;
-	}
-
 	if (Preview2DViewport)
 	{
 		Preview2DViewport->SetVisibility(bShowViewport ? EVisibility::Visible : EVisibility::Collapsed);
@@ -582,6 +590,8 @@ void FMaterialDiffPanel::AddReferencedObjects(FReferenceCollector& Collector)
 	Collector.AddReferencedObject(ExpressionPreviewMaterial);
 	Collector.AddReferencedObject(PreviewExpression);
 	Collector.AddReferencedObject(MaterialGraph);
+	Collector.AddReferencedObject(MaterialInstance);
+	Collector.AddReferencedObject(EditorProxy);
 }
 
 FString FMaterialDiffPanel::GetReferencerName() const
diff --git a/Engine/Source/Editor/MaterialEditor/Private/DiffTool/MaterialInstanceDiffPanel.cpp b/Engine/Source/Editor/MaterialEditor/Private/DiffTool/MaterialInstanceDiffPanel.cpp
deleted file mode 100644
diff --git a/Engine/Source/Editor/MaterialEditor/Private/DiffTool/Widgets/SMaterialDiff.cpp b/Engine/Source/Editor/MaterialEditor/Private/DiffTool/Widgets/SMaterialDiff.cpp
--- a/Engine/Source/Editor/MaterialEditor/Private/DiffTool/Widgets/SMaterialDiff.cpp
+++ b/Engine/Source/Editor/MaterialEditor/Private/DiffTool/Widgets/SMaterialDiff.cpp
@@ -9,6 +9,10 @@
 #include "GraphDiffControl.h"
 #include "IDetailsView.h"
 #include "MaterialGraph/MaterialGraph.h"
+#include "MaterialGraph/MaterialGraphNode.h"
+#include "MaterialGraph/MaterialGraphSchema.h"
+#include "Kismet2/BlueprintEditorUtils.h"
+#include "Materials/MaterialExpressionStaticSwitchParameter.h"
 #include "Materials/MaterialFunction.h"
 #include "Modules/ModuleManager.h"
 #include "PropertyEditorModule.h"
@@ -27,6 +31,9 @@
 #include "Widgets/Layout/SSplitter.h"
 #include "Widgets/Layout/SWidgetSwitcher.h"
 #include "Widgets/SOverlay.h"
+#include "MaterialEditor/MaterialEditorInstanceConstant.h"
+#include "MaterialEditorInstanceDetailCustomization.h"
+#include "Materials/MaterialInstanceConstant.h"
 
 #define LOCTEXT_NAMESPACE "SMaterialGraphDiff"
 
@@ -35,28 +42,131 @@ static const FName GraphMode = FName(TEXT("GraphMode"));
 BEGIN_SLATE_FUNCTION_BUILD_OPTIMIZATION
 void SMaterialDiff::Construct(const FArguments& InArgs)
 {
-	check(InArgs._OldMaterialGraph || InArgs._NewMaterialGraph);
+	check(InArgs._OldMaterialGraph || InArgs._NewMaterialGraph || InArgs._OldMaterialInstance || InArgs._NewMaterialInstance);
 
-	PanelOld.MaterialGraph = InArgs._OldMaterialGraph;
-	PanelNew.MaterialGraph = InArgs._NewMaterialGraph;
-	PanelOld.RevisionInfo = InArgs._OldRevision;
-	PanelNew.RevisionInfo = InArgs._NewRevision;
+	if (InArgs._OldMaterialInstance || InArgs._NewMaterialInstance)
+	{
+		PanelOld.MaterialInstance = InArgs._OldMaterialInstance;
+		PanelNew.MaterialInstance = InArgs._NewMaterialInstance;
 
-	PanelOld.bShowAssetName = InArgs._ShowAssetNames;
-	PanelNew.bShowAssetName = InArgs._ShowAssetNames;
+		// Create UMaterialEditorInstanceConstant proxies for each side so that ALL parameter types
+		// (including static switches and component masks) are visible in the details view.
+		// This is the same proxy object the Material Instance Editor uses to display parameters.
+		auto CreateEditorProxy = [](UMaterialInstance* MI) -> UMaterialEditorInstanceConstant*
+		{
+			UMaterialInstanceConstant* MIC = Cast<UMaterialInstanceConstant>(MI);
+			if (!MIC)
+			{
+				return nullptr;
+			}
+
+			UMaterialEditorInstanceConstant* Proxy = NewObject<UMaterialEditorInstanceConstant>(GetTransientPackage(), NAME_None, RF_Transactional);
+			Proxy->SetSourceInstance(MIC);
+			return Proxy;
+		};
+
+		PanelOld.EditorProxy = CreateEditorProxy(PanelOld.MaterialInstance);
+		PanelNew.EditorProxy = CreateEditorProxy(PanelNew.MaterialInstance);
+
+		PanelOld.RevisionInfo = InArgs._OldRevision;
+		PanelNew.RevisionInfo = InArgs._NewRevision;
+
+		// Build resolved material graphs from each instance's parent material.
+		// Each side gets its own duplicated parent so expression objects are independent.
+		auto BuildResolvedGraph = [](UMaterial* ParentCopy, UMaterialInstance* Instance) -> UMaterialGraph*
+		{
+			UMaterialGraph* Graph = CastChecked<UMaterialGraph>(
+				FBlueprintEditorUtils::CreateNewGraph(
+					GetTransientPackage(),
+					MakeUniqueObjectName(GetTransientPackage(), UMaterialGraph::StaticClass()),
+					UMaterialGraph::StaticClass(),
+					UMaterialGraphSchema::StaticClass()));
+			Graph->Material = ParentCopy;
+			Graph->MaterialFunction = nullptr;
+			Graph->RebuildGraph();
+
+			// Apply static switch overrides from the instance
+			if (Instance)
+			{
+				FStaticParameterSet StaticParams = Instance->GetStaticParameters();
+
+				TMap<FName, bool> SwitchOverrides;
+				for (const FStaticSwitchParameter& Param : StaticParams.StaticSwitchParameters)
+				{
+					if (Param.bOverride)
+					{
+						SwitchOverrides.Add(Param.ParameterInfo.Name, Param.Value);
+					}
+				}
+
+				for (UEdGraphNode* Node : Graph->Nodes)
+				{
+					if (UMaterialGraphNode* MatNode = Cast<UMaterialGraphNode>(Node))
+					{
+						if (UMaterialExpressionStaticSwitchParameter* SwitchExpr = Cast<UMaterialExpressionStaticSwitchParameter>(MatNode->MaterialExpression))
+						{
+							if (const bool* Override = SwitchOverrides.Find(SwitchExpr->ParameterName))
+							{
+								SwitchExpr->DefaultValue = *Override;
+							}
+						}
+					}
+				}
+			}
+
+			return Graph;
+		};
+
+		if (UMaterial* OldParent = PanelOld.MaterialInstance ? PanelOld.MaterialInstance->GetMaterial() : nullptr)
+		{
+			PanelOld.MaterialGraph = BuildResolvedGraph(DuplicateObject<UMaterial>(OldParent, GetTransientPackage()), PanelOld.MaterialInstance);
+		}
+		if (UMaterial* NewParent = PanelNew.MaterialInstance ? PanelNew.MaterialInstance->GetMaterial() : nullptr)
+		{
+			PanelNew.MaterialGraph = BuildResolvedGraph(DuplicateObject<UMaterial>(NewParent, GetTransientPackage()), PanelNew.MaterialInstance);
+		}
+
+		// Only display the graph diff if we have two material instances
+		if (PanelOld.MaterialInstance && PanelNew.MaterialInstance)
+		{
+			const bool bSameParent = PanelOld.MaterialInstance->GetMaterial() == PanelNew.MaterialInstance->GetMaterial();
+			const bool bHasStaticPermutation = PanelOld.MaterialInstance->bHasStaticPermutationResource || PanelNew.MaterialInstance->bHasStaticPermutationResource;
+			bShowGraphDiff = bSameParent && bHasStaticPermutation;
+		}
+	}
+	else
+	{
+		PanelOld.MaterialGraph = InArgs._OldMaterialGraph;
+		PanelNew.MaterialGraph = InArgs._NewMaterialGraph;
+		PanelOld.RevisionInfo = InArgs._OldRevision;
+		PanelNew.RevisionInfo = InArgs._NewRevision;
+
+		PanelOld.bShowAssetName = InArgs._ShowAssetNames;
+		PanelNew.bShowAssetName = InArgs._ShowAssetNames;
+
+		bShowGraphDiff = true;
+	}
 
 	UObject* OldAsset = nullptr;
-	if (PanelOld.MaterialGraph && PanelOld.MaterialGraph->MaterialFunction)
+	if (PanelOld.EditorProxy)
+	{
+		OldAsset = PanelOld.EditorProxy;
+	}
+	else if (PanelOld.MaterialGraph && PanelOld.MaterialGraph->MaterialFunction)
 	{
 		OldAsset = PanelOld.MaterialGraph->MaterialFunction;
 	}
-	else if(PanelOld.MaterialGraph)
+	else if (PanelOld.MaterialGraph)
 	{
 		OldAsset = PanelOld.MaterialGraph->Material;
 	}
 
 	UObject* NewAsset = nullptr;
-	if (PanelNew.MaterialGraph && PanelNew.MaterialGraph->MaterialFunction)
+	if (PanelNew.EditorProxy)
+	{
+		NewAsset = PanelNew.EditorProxy;
+	}
+	else if (PanelNew.MaterialGraph && PanelNew.MaterialGraph->MaterialFunction)
 	{
 		NewAsset = PanelNew.MaterialGraph->MaterialFunction;
 	}
@@ -203,6 +313,7 @@ void SMaterialDiff::Construct(const FArguments& InArgs)
 					[
 						SNew(SButton)
 						.HAlign(HAlign_Center)
+						.IsEnabled(bShowGraphDiff)
 						.OnClicked_Lambda([this]()->FReply { SetCurrentWidgetIndex(0); return FReply::Handled(); })
 						[
 							SNew(STextBlock)
@@ -286,6 +397,7 @@ void SMaterialDiff::Construct(const FArguments& InArgs)
 							.NewRevision(PanelNew.RevisionInfo)
 							.ShowAssetNames(true)
 							.ParentWindow(WeakParentWindow.Pin())
+							.OnCustomizeDetailsWidget_Raw(this, &SMaterialDiff::AddMaterialInstanceDetailCustomization)
 						]
 					]
 				]
@@ -295,9 +407,9 @@ void SMaterialDiff::Construct(const FArguments& InArgs)
 	];
 	// clang-format on
 
-	CurrentWidgetIndex = 0;
+	// Default to Properties tab for instances, Graph tab for materials
+	CurrentWidgetIndex = (PanelOld.MaterialInstance || PanelNew.MaterialInstance) ? 1 : 0;
 
-	// Display the Material Graphs by default
 	HandleGraphChanged(FString{});
 }
 END_SLATE_FUNCTION_BUILD_OPTIMIZATION
@@ -812,4 +924,92 @@ void SMaterialDiff::UpdateTopSectionVisibility(const FName& InNewViewMode) const
 	}
 }
 
+void SMaterialDiff::AddMaterialInstanceDetailCustomization(const TSharedRef<IDetailsView>& DetailsView)
+{
+	if (!PanelOld.EditorProxy && !PanelNew.EditorProxy)
+	{
+		return;
+	}
+
+	// Determine which panel this details view belongs to by checking which EditorProxy it is displaying.
+	// SDetailsDiff calls this callback twice: once for the Old side and once for the New side.
+	UMaterialEditorInstanceConstant* ProxyForThisView = nullptr;
+	const TArray<TWeakObjectPtr<UObject>>& SelectedObjects = DetailsView->GetSelectedObjects();
+	if (SelectedObjects.Num() > 0 && SelectedObjects[0].IsValid())
+	{
+		UObject* DisplayedObject = SelectedObjects[0].Get();
+		if (DisplayedObject == PanelOld.EditorProxy)
+		{
+			ProxyForThisView = PanelOld.EditorProxy;
+		}
+		else if (DisplayedObject == PanelNew.EditorProxy)
+		{
+			ProxyForThisView = PanelNew.EditorProxy;
+		}
+	}
+
+	if (!ProxyForThisView)
+	{
+		return;
+	}
+
+	FOnGetDetailCustomizationInstance LayoutMICDetails = FOnGetDetailCustomizationInstance::CreateStatic(
+		&FMaterialInstanceParameterDetails::MakeInstance,
+		ProxyForThisView,
+		static_cast<SMaterialLayersFunctionsInstanceWrapper*>(nullptr),
+		FGetShowHiddenParameters());
+
+	DetailsView->RegisterInstancedCustomPropertyLayout(UMaterialEditorInstanceConstant::StaticClass(), LayoutMICDetails);
+	DetailsView->ForceRefresh();
+}
+
+TSharedPtr<SWindow> SMaterialDiff::CreateDiffWindow(FText WindowTitle, TObjectPtr<UMaterialInstance> OldMaterialInstance, TObjectPtr<UMaterialInstance> NewMaterialInstance, const FRevisionInfo& OldRevision, const FRevisionInfo& NewRevision)
+{
+	TSharedPtr<SWindow> Window = SNew(SWindow)
+									 .Title(WindowTitle)
+									 .ClientSize(FVector2D(1000, 800));
+
+	Window->SetContent(SNew(SMaterialDiff)
+						   .OldMaterialInstance(OldMaterialInstance)
+						   .NewMaterialInstance(NewMaterialInstance)
+						   .OldRevision(OldRevision)
+						   .NewRevision(NewRevision)
+						   .ParentWindow(Window));
+
+	// Make this window a child of the modal window if we've been spawned while one is active.
+	const TSharedPtr<SWindow> ActiveModal = FSlateApplication::Get().GetActiveModalWindow();
+	if (ActiveModal.IsValid())
+	{
+		FSlateApplication::Get().AddWindowAsNativeChild(Window.ToSharedRef(), ActiveModal.ToSharedRef());
+	}
+	else
+	{
+		FSlateApplication::Get().AddWindow(Window.ToSharedRef());
+	}
+
+	return Window;
+}
+
+TSharedPtr<SWindow> SMaterialDiff::CreateDiffWindow(TObjectPtr<UMaterialInstance> OldMaterialInstance, TObjectPtr<UMaterialInstance> NewMaterialInstance, const FRevisionInfo& OldRevision, const FRevisionInfo& NewRevision, const UClass* ObjectClass)
+{
+	check(OldMaterialInstance || NewMaterialInstance);
+
+	// sometimes we're comparing different revisions of one single asset (other
+	// times we're comparing two completely separate assets altogether)
+	const bool bIsSingleAsset = !OldMaterialInstance ||
+		!NewMaterialInstance ||
+		(NewMaterialInstance->GetName() == OldMaterialInstance->GetName());
+
+	FText WindowTitle = FText::Format(LOCTEXT("NamelessMaterialInstanceDiff", "{0} Diff"), ObjectClass->GetDisplayNameText());
+	// if we're diffing one asset against itself
+	if (bIsSingleAsset)
+	{
+		// identify the assumed single asset in the window's title
+		FString MaterialInstanceName = OldMaterialInstance ? OldMaterialInstance->GetName() : NewMaterialInstance->GetName();
+		WindowTitle = FText::Format(LOCTEXT("NamedMaterialInstanceDiff", "{0} - {1} Diff"), FText::FromString(MaterialInstanceName), ObjectClass->GetDisplayNameText());
+	}
+
+	return CreateDiffWindow(WindowTitle, OldMaterialInstance, NewMaterialInstance, OldRevision, NewRevision);
+}
+
 #undef LOCTEXT_NAMESPACE
diff --git a/Engine/Source/Editor/MaterialEditor/Private/DiffTool/Widgets/SMaterialInstanceDiff.cpp b/Engine/Source/Editor/MaterialEditor/Private/DiffTool/Widgets/SMaterialInstanceDiff.cpp
deleted file mode 100644
