Add GetMaterialUsedTextures which accepts a UMaterialInterface so it cleanly works with UMaterials and UMaterialInstances alike.

- Deprecate the existing GetUsedTextures which only accepts a UMaterial and forward its calls to the new BP function.

#Synced-CL 51833860

diff --git a/Engine/Source/Editor/MaterialEditor/Private/MaterialEditingLibrary.cpp b/Engine/Source/Editor/MaterialEditor/Private/MaterialEditingLibrary.cpp
--- a/Engine/Source/Editor/MaterialEditor/Private/MaterialEditingLibrary.cpp
+++ b/Engine/Source/Editor/MaterialEditor/Private/MaterialEditingLibrary.cpp
@@ -1093,9 +1093,14 @@ void UMaterialEditingLibrary::GetMaterialExpressionNodePosition(UMaterialExpress
 }
 
 TArray<UTexture*> UMaterialEditingLibrary::GetUsedTextures(UMaterial* Material)
+{
+	return UMaterialEditingLibrary::GetMaterialUsedTextures(Material);
+}
+
+TArray<UTexture*> UMaterialEditingLibrary::GetMaterialUsedTextures(UMaterialInterface* MaterialInterface)
 {
 	TArray<UTexture*> OutTextures;
-	Material->GetUsedTextures(OutTextures, GetCurrentMaterialQualityLevelChecked());
+	MaterialInterface->GetUsedTextures(OutTextures, GetCurrentMaterialQualityLevelChecked());
 	return OutTextures;
 }
 
diff --git a/Engine/Source/Editor/MaterialEditor/Public/MaterialEditingLibrary.h b/Engine/Source/Editor/MaterialEditor/Public/MaterialEditingLibrary.h
--- a/Engine/Source/Editor/MaterialEditor/Public/MaterialEditingLibrary.h
+++ b/Engine/Source/Editor/MaterialEditor/Public/MaterialEditingLibrary.h
@@ -283,8 +283,12 @@ public:
 	static UE_API void GetMaterialExpressionNodePosition(UMaterialExpression* MaterialExpression, int32& NodePosX, int32& NodePosY);
 
 	/** Get the list of textures used by a material */
-	UFUNCTION(BlueprintPure, Category = "MaterialEditing")
+	UFUNCTION(BlueprintPure, Category = "MaterialEditing", meta = (DeprecatedFunction, DeprecationMessage = "Use GetMaterialUsedTextures instead since it works with Material Instances."))
 	static UE_API TArray<UTexture*> GetUsedTextures(UMaterial* Material);
+
+	/** Get the list of textures used by a material */
+	UFUNCTION(BlueprintPure, Category = "MaterialEditing")
+	static UE_API TArray<UTexture*> GetMaterialUsedTextures(UMaterialInterface* MaterialInterface);
 	
 	//////// MATERIAL FUNCTION EDITING
 
