UX adjustments to Material Editor - Shader Code window
Added permutation name to shader name in the material editor's show shader code tool (simplified name in the list + full permutation name in the tooltip) to avoid seeing duplicate names in the list

#Synced-CL 52169855

diff --git a/Engine/Source/Editor/MaterialEditor/Private/MaterialStats.cpp b/Engine/Source/Editor/MaterialEditor/Private/MaterialStats.cpp
--- a/Engine/Source/Editor/MaterialEditor/Private/MaterialStats.cpp
+++ b/Engine/Source/Editor/MaterialEditor/Private/MaterialStats.cpp
@@ -93,7 +93,7 @@ FText FShaderPlatformSettings::GetSelectedShaderViewComboText(EMaterialQualityLe
 		return FText::FromString(TEXT("-Compiling-Shaders-"));
 	}
 
-	return FText::FromString(Instance.ComboBoxSelectedEntry.Text);
+	return FText::FromString(Instance.ComboBoxSelectedEntry.ShaderName);
 }
 
 void FShaderPlatformSettings::OnShaderViewComboSelectionChanged(TSharedPtr<FMaterialShaderEntry> Item, EMaterialQualityLevel::Type QualityType, const int32 InstanceIndex)
@@ -406,15 +406,29 @@ bool FShaderPlatformSettings::Update()
 				{
 					FShaderType* EntryShader = Entry.Value.GetType();
 					FVertexFactoryType* VertexFactory = Entry.Value.GetVertexFactoryType();
+					
 					FString ShaderName = FString::Printf(TEXT("%s/%s"), VertexFactory ? VertexFactory->GetName() : TEXT("NullVF"), EntryShader->GetName());
-					Instance.ArrShaderEntries.Add(MakeShareable(new FMaterialShaderEntry{ Entry.Key, ShaderName }));
+					FString ShaderFullName = ShaderName;
+
+					if (FString PermutationName = EntryShader->GetPermutationIdString(Entry.Key.PermutationId, /*bFullNames = */false); !PermutationName.IsEmpty())
+					{
+						// Add a (shortened) permutation string to the shader name if applicable, so that we don't see multiple identical shader names when their permutation differs
+						ShaderName += FString::Printf(TEXT("[%s]"), *PermutationName);
+					}
+					if (FString PermutationFullName = EntryShader->GetPermutationIdString(Entry.Key.PermutationId, /*bFullNames = */true); !PermutationFullName.IsEmpty())
+					{
+						// In the tooltip, we have room to display the full permutation name, which is already line-formatted 
+						ShaderFullName += FString::Printf(TEXT("\n%s"), *PermutationFullName);
+					}
+
+					Instance.ArrShaderEntries.Add(MakeShareable(new FMaterialShaderEntry{ .ShaderId = Entry.Key, .ShaderName = ShaderName, .ShaderFullName = ShaderFullName }));
 				}
 
 				if (Instance.ArrShaderEntries.Num() > 0)
 				{
 					Instance.ArrShaderEntries.Sort([](const TSharedPtr<FMaterialShaderEntry>& First, const TSharedPtr<FMaterialShaderEntry>& Second)
 						{
-							return First->Text < Second->Text;
+							return First->ShaderName < Second->ShaderName;
 						});
 
 					Instance.ComboBoxSelectedEntry = *Instance.ArrShaderEntries[0];
@@ -1180,14 +1194,21 @@ TSharedRef<SDockTab> FMaterialStats::SpawnTab_ShaderCode(const FSpawnTabArgs& Ar
 		ShaderEntriesPtr = &EmptyShaderEntriesArray;
 	}
 
+	// Use an external scrollbar for the SListView, so that the vertical scrollbar remains visible all the time (and the SListView's horizontal one can still be used to peek at the full shader names)
+	TSharedRef<SScrollBar> ExternalScrollbar = SNew(SScrollBar).Orientation(EOrientation::Orient_Vertical);
+
 	ensure(ShaderEntriesPtr != nullptr);
 	TSharedRef<SListView<TSharedPtr<FMaterialShaderEntry>>> ShaderListView = SNew(SListView<TSharedPtr<FMaterialShaderEntry>>)
 		.ListItemsSource(ShaderEntriesPtr)
 		.SelectionMode(ESelectionMode::Single)
+		.ExternalScrollbar(ExternalScrollbar)
 		.OnGenerateRow_Lambda([](TSharedPtr<FMaterialShaderEntry> Value, const TSharedRef<STableViewBase>& OwnerTable) {
 			return SNew(STableRow<TSharedPtr<FMaterialShaderEntry>>, OwnerTable)
 			[
-				SNew(STextBlock).Text(FText::FromString(Value->Text))
+				SNew(STextBlock)
+					.AutoWrapText(false)
+					.Text(FText::FromString(Value->ShaderName))
+					.ToolTipText(FText::FromString(Value->ShaderFullName))
 			];
 		})
 		.OnSelectionChanged_Lambda([PlatformPtr, QualityLevel, InstanceIndex](TSharedPtr<FMaterialShaderEntry> Item, ESelectInfo::Type SelectInfo) { PlatformPtr->OnShaderViewComboSelectionChanged(Item, QualityLevel, InstanceIndex); });
@@ -1237,7 +1258,22 @@ TSharedRef<SDockTab> FMaterialStats::SpawnTab_ShaderCode(const FSpawnTabArgs& Ar
 				+ SSplitter::Slot()
 				.Value(0.3f)
 				[
-					ShaderListView
+					SNew(SHorizontalBox)
+					+ SHorizontalBox::Slot()
+					.FillWidth(1.0f)
+					[
+						SNew(SScrollBox)
+						.Orientation(EOrientation::Orient_Horizontal)
+						+ SScrollBox::Slot()
+						[
+							ShaderListView
+						]
+					]
+					+ SHorizontalBox::Slot()
+					.AutoWidth()
+					[
+						ExternalScrollbar
+					]
 				]
 				+ SSplitter::Slot()
 				.Value(0.7f)
diff --git a/Engine/Source/Editor/MaterialEditor/Private/MaterialStats.h b/Engine/Source/Editor/MaterialEditor/Private/MaterialStats.h
--- a/Engine/Source/Editor/MaterialEditor/Private/MaterialStats.h
+++ b/Engine/Source/Editor/MaterialEditor/Private/MaterialStats.h
@@ -102,7 +102,8 @@ struct FShaderStatsInfo
 struct FMaterialShaderEntry
 {
 	FShaderId ShaderId;
-	FString Text;
+	FString ShaderName;
+	FString ShaderFullName;
 };
 
 /** structure used to manage shader compilation and source code extraction for a specified shader platform
