Add getter and setter for material instance parameter override state.

This adds two new Blueprint-exposed functions to UMaterialEditingLibrary:

IsMaterialInstanceParameterOverridden ??? checks if a parameter is overridden on this instance
SetMaterialInstanceParameterOverride ??? enables or disables the override for a parameter
Both work across all parameter types (scalar, vector, texture, static switch, component mask, etc.). Enabling an override populates it with the inherited value from the parent chain. Disabling removes the parameter from the instance's arrays (or clears bOverride for static parameters). The open MI editor's proxies are rebuilt after changes so the UI stays in sync.

Supporting changes:

Add SetDoubleVectorParameterValueEditorOnly and SetParameterCollectionParameterValueEditorOnly to UMaterialInstanceConstant (missing EditorOnly wrappers, matching the existing pattern for other parameter types)
Add SetStaticComponentMaskParameterValueEditorOnly to UMaterialInstance.

[RN] minor | Rendering Architecture

#Synced-CL 52210234

#Synced-CL 52210234

diff --git a/Engine/Source/Editor/MaterialEditor/Private/MaterialEditingLibrary.cpp b/Engine/Source/Editor/MaterialEditor/Private/MaterialEditingLibrary.cpp
--- a/Engine/Source/Editor/MaterialEditor/Private/MaterialEditingLibrary.cpp
+++ b/Engine/Source/Editor/MaterialEditor/Private/MaterialEditingLibrary.cpp
@@ -1394,6 +1394,149 @@ bool UMaterialEditingLibrary::SetMaterialInstanceStaticSwitchParameterValue(UMat
 	return bResult;
 }
 
+bool UMaterialEditingLibrary::IsMaterialInstanceParameterOverridden(UMaterialInstanceConstant* Instance, FName ParameterName, EMaterialParameterAssociation Association)
+{
+	if (!Instance)
+	{
+		return false;
+	}
+
+	const FMaterialParameterInfo ParamInfo(ParameterName, Association, Association == EMaterialParameterAssociation::LayerParameter ? 0 : INDEX_NONE);
+	FMaterialParameterMetadata OutResult;
+
+	for (int32 TypeIndex = 0; TypeIndex < NumMaterialParameterTypes; ++TypeIndex)
+	{
+		const EMaterialParameterType Type = static_cast<EMaterialParameterType>(TypeIndex);
+		if (Instance->GetParameterOverrideValue(Type, ParamInfo, OutResult))
+		{
+			return true;
+		}
+	}
+
+	return false;
+}
+
+bool UMaterialEditingLibrary::SetMaterialInstanceParameterOverride(UMaterialInstanceConstant* Instance, FName ParameterName, bool bOverride, EMaterialParameterAssociation Association)
+{
+	if (!Instance)
+	{
+		return false;
+	}
+
+	const FMaterialParameterInfo ParamInfo(ParameterName, Association, Association == EMaterialParameterAssociation::LayerParameter ? 0 : INDEX_NONE);
+	FMaterialParameterMetadata Meta;
+	EMaterialParameterType FoundType = EMaterialParameterType::None;
+
+	for (int32 TypeIndex = 0; TypeIndex < NumMaterialParameterTypes; ++TypeIndex)
+	{
+		const EMaterialParameterType Type = static_cast<EMaterialParameterType>(TypeIndex);
+		if (Instance->GetParameterValue(Type, ParamInfo, Meta))
+		{
+			FoundType = Type;
+			break;
+		}
+	}
+
+	if (FoundType == EMaterialParameterType::None)
+	{
+		UE_LOG(LogMaterialEditingLibrary, Warning, TEXT("SetMaterialInstanceParameterOverride: Parameter '%s' not found in material hierarchy."), *ParameterName.ToString());
+		return false;
+	}
+
+	// Already in the desired state — nothing to do.
+	if (bOverride == Meta.bOverride)
+	{
+		return true;
+	}
+
+
+	if (bOverride)
+	{
+		// Enabling override: set the parameter on this instance using the inherited value from parent.
+		switch (FoundType)
+		{
+		case EMaterialParameterType::Scalar:				Instance->SetScalarParameterValueEditorOnly(ParamInfo, Meta.Value.AsScalar()); break;
+		case EMaterialParameterType::Vector:				Instance->SetVectorParameterValueEditorOnly(ParamInfo, Meta.Value.AsLinearColor()); break;
+		case EMaterialParameterType::DoubleVector:			Instance->SetDoubleVectorParameterValueEditorOnly(ParamInfo, Meta.Value.AsVector4d()); break;
+		case EMaterialParameterType::Texture:				Instance->SetTextureParameterValueEditorOnly(ParamInfo, Meta.Value.Texture); break;
+		case EMaterialParameterType::TextureCollection:		Instance->SetTextureCollectionParameterValueEditorOnly(ParamInfo, Meta.Value.TextureCollection); break;
+		case EMaterialParameterType::Font:					Instance->SetFontParameterValueEditorOnly(ParamInfo, Meta.Value.Font.Value, Meta.Value.Font.Page); break;
+		case EMaterialParameterType::RuntimeVirtualTexture: Instance->SetRuntimeVirtualTextureParameterValueEditorOnly(ParamInfo, Meta.Value.RuntimeVirtualTexture); break;
+		case EMaterialParameterType::SparseVolumeTexture:	Instance->SetSparseVolumeTextureParameterValueEditorOnly(ParamInfo, Meta.Value.SparseVolumeTexture); break;
+		case EMaterialParameterType::StaticSwitch:			Instance->SetStaticSwitchParameterValueEditorOnly(ParamInfo, Meta.Value.AsStaticSwitch()); break;
+		case EMaterialParameterType::ParameterCollection:	Instance->SetParameterCollectionParameterValueEditorOnly(ParamInfo, Meta.Value.ParameterCollection); break;
+		case EMaterialParameterType::StaticComponentMask:	Instance->SetStaticComponentMaskParameterValueEditorOnly(ParamInfo, Meta.Value.AsStaticComponentMask()); break;
+		default:
+			UE_LOG(LogMaterialEditingLibrary, Warning, TEXT("SetMaterialInstanceParameterOverride: Unsupported parameter type for '%s'."), *ParameterName.ToString());
+			return false;
+		}
+	}
+	else
+	{
+		bool bRemoved = false;
+
+		// Disabling override: remove the parameter from the appropriate array.
+		auto RemoveParameterByInfo = [&ParamInfo](auto& ParameterArray) -> bool
+		{
+			const int32 Index = ParameterArray.IndexOfByPredicate([&ParamInfo](const auto& Param)
+			{
+				return Param.ParameterInfo == ParamInfo;
+			});
+
+			if (Index != INDEX_NONE)
+			{
+				ParameterArray.RemoveAt(Index);
+				return true;
+			}
+			return false;
+		};
+
+		auto RemoveStaticParameter = [&ParamInfo](auto& StaticParameterArray) -> bool
+		{
+			for (auto& Param : StaticParameterArray)
+			{
+				if (Param.ParameterInfo == ParamInfo)
+				{
+					Param.bOverride = false;
+					return true;
+				}
+			}
+			return false;
+		};
+
+		switch (FoundType)
+		{
+		case EMaterialParameterType::Scalar:				bRemoved = RemoveParameterByInfo(Instance->ScalarParameterValues); break;
+		case EMaterialParameterType::Vector:				bRemoved = RemoveParameterByInfo(Instance->VectorParameterValues); break;
+		case EMaterialParameterType::DoubleVector:			bRemoved = RemoveParameterByInfo(Instance->DoubleVectorParameterValues); break;
+		case EMaterialParameterType::Texture:				bRemoved = RemoveParameterByInfo(Instance->TextureParameterValues); break;
+		case EMaterialParameterType::TextureCollection:		bRemoved = RemoveParameterByInfo(Instance->TextureCollectionParameterValues); break;
+		case EMaterialParameterType::ParameterCollection:	bRemoved = RemoveParameterByInfo(Instance->ParameterCollectionParameterValues); break;
+		case EMaterialParameterType::Font:					bRemoved = RemoveParameterByInfo(Instance->FontParameterValues); break;
+		case EMaterialParameterType::RuntimeVirtualTexture: bRemoved = RemoveParameterByInfo(Instance->RuntimeVirtualTextureParameterValues); break;
+		case EMaterialParameterType::SparseVolumeTexture:	bRemoved = RemoveParameterByInfo(Instance->SparseVolumeTextureParameterValues); break;
+		case EMaterialParameterType::StaticSwitch:			bRemoved = RemoveStaticParameter(Instance->StaticParametersRuntime.StaticSwitchParameters); break;
+		case EMaterialParameterType::StaticComponentMask:	bRemoved = RemoveStaticParameter(Instance->GetEditorOnlyData()->StaticParameters.StaticComponentMaskParameters); break;
+		default:
+			UE_LOG(LogMaterialEditingLibrary, Warning, TEXT("SetMaterialInstanceParameterOverride: Unsupported parameter type for '%s'."), *ParameterName.ToString());
+		}
+
+		if (!bRemoved)
+		{
+			return false;
+		}
+	}
+
+	UMaterialEditingLibrary::UpdateMaterialInstance(Instance);
+	if (FMaterialInstanceEditor* MaterialInstanceEditor = MaterialEditingLibraryImpl::FindMaterialInstanceEditorForAsset(Instance))
+	{
+		// Rebuild the editor's parameter proxies so the UI reflects the override change.
+		MaterialInstanceEditor->RebuildMaterialInstanceEditor();
+	}
+
+	return true;
+}
+
 void UMaterialEditingLibrary::UpdateMaterialInstance(UMaterialInstanceConstant* Instance)
 {
 	if (Instance)
diff --git a/Engine/Source/Editor/MaterialEditor/Public/MaterialEditingLibrary.h b/Engine/Source/Editor/MaterialEditor/Public/MaterialEditingLibrary.h
--- a/Engine/Source/Editor/MaterialEditor/Public/MaterialEditingLibrary.h
+++ b/Engine/Source/Editor/MaterialEditor/Public/MaterialEditingLibrary.h
@@ -407,6 +407,27 @@ public:
 	UFUNCTION(BlueprintCallable, Category = "MaterialEditing")
 	static UE_API bool SetMaterialInstanceStaticSwitchParameterValue(UMaterialInstanceConstant* Instance, FName ParameterName, bool Value, EMaterialParameterAssociation Association = EMaterialParameterAssociation::GlobalParameter, bool bUpdateMaterialInstance = true);
 
+	/**
+	 *	Check whether a parameter is overridden in this Material Instance.
+	 *	@param	Instance		Material Instance to query
+	 *	@param	ParameterName	Name of the parameter
+	 *	@param	Association		Parameter association type
+	 *	@return Whether the parameter is overridden in this instance
+	 */
+	UFUNCTION(BlueprintPure, Category = "MaterialEditing")
+	static UE_API bool IsMaterialInstanceParameterOverridden(UMaterialInstanceConstant* Instance, FName ParameterName, EMaterialParameterAssociation Association = EMaterialParameterAssociation::GlobalParameter);
+
+	/**
+	 *	Set whether a parameter override is enabled in this Material Instance.
+	 *	@param	Instance		Material Instance to modify
+	 *	@param	ParameterName	Name of the parameter
+	 *	@param	bOverride		Whether the parameter should be overridden
+	 *	@param	Association		Parameter association type
+	 *	@return Whether the operation succeeded (false if Instance is null or the parameter was not found)
+	 */
+	UFUNCTION(BlueprintCallable, Category = "MaterialEditing")
+	static UE_API bool SetMaterialInstanceParameterOverride(UMaterialInstanceConstant* Instance, FName ParameterName, bool bOverride, EMaterialParameterAssociation Association = EMaterialParameterAssociation::GlobalParameter);
+
 	/** Called after making modifications to a Material Instance to recompile shaders etc. */
 	UFUNCTION(BlueprintCallable, Category = "MaterialEditing")
 	static UE_API void UpdateMaterialInstance(UMaterialInstanceConstant* Instance);
diff --git a/Engine/Source/Runtime/Engine/Private/Materials/MaterialInstance.cpp b/Engine/Source/Runtime/Engine/Private/Materials/MaterialInstance.cpp
--- a/Engine/Source/Runtime/Engine/Private/Materials/MaterialInstance.cpp
+++ b/Engine/Source/Runtime/Engine/Private/Materials/MaterialInstance.cpp
@@ -2216,6 +2216,29 @@ void UMaterialInstance::SetStaticSwitchParameterValueEditorOnly(const FMaterialP
 	new(StaticParametersRuntime.StaticSwitchParameters) FStaticSwitchParameter(ParameterInfo, Value, true, FGuid());
 }
 
+void UMaterialInstance::SetStaticComponentMaskParameterValueEditorOnly(const FMaterialParameterInfo& ParameterInfo, FStaticComponentMaskValue Value)
+{
+	check(GIsEditor || IsRunningCommandlet());
+
+	UMaterialInstanceEditorOnlyData* EditorOnly = GetEditorOnlyData();
+	check(EditorOnly);
+
+	for (FStaticComponentMaskParameter& Parameter : EditorOnly->StaticParameters.StaticComponentMaskParameters)
+	{
+		if (Parameter.ParameterInfo == ParameterInfo)
+		{
+			Parameter.bOverride = true;
+			Parameter.R = Value.R;
+			Parameter.G = Value.G;
+			Parameter.B = Value.B;
+			Parameter.A = Value.A;
+			return;
+		}
+	}
+
+	new(EditorOnly->StaticParameters.StaticComponentMaskParameters) FStaticComponentMaskParameter(ParameterInfo, Value.R, Value.G, Value.B, Value.A, true, FGuid());
+}
+
 void UMaterialInterface::GetStaticParameterValues(FStaticParameterSet& OutStaticParameters)
 {
 	check(IsInGameThread());
diff --git a/Engine/Source/Runtime/Engine/Private/Materials/MaterialInstanceConstant.cpp b/Engine/Source/Runtime/Engine/Private/Materials/MaterialInstanceConstant.cpp
--- a/Engine/Source/Runtime/Engine/Private/Materials/MaterialInstanceConstant.cpp
+++ b/Engine/Source/Runtime/Engine/Private/Materials/MaterialInstanceConstant.cpp
@@ -116,6 +116,12 @@ void UMaterialInstanceConstant::SetVectorParameterValueEditorOnly(const FMateria
 	SetVectorParameterValueInternal(ParameterInfo,Value);
 }
 
+void UMaterialInstanceConstant::SetDoubleVectorParameterValueEditorOnly(const FMaterialParameterInfo& ParameterInfo, FVector4d Value)
+{
+	check(GIsEditor || IsRunningCommandlet());
+	SetDoubleVectorParameterValueInternal(ParameterInfo, Value);
+}
+
 void UMaterialInstanceConstant::SetScalarParameterValueEditorOnly(const FMaterialParameterInfo& ParameterInfo, float Value)
 {
 	check(GIsEditor || IsRunningCommandlet());
@@ -140,6 +146,12 @@ void UMaterialInstanceConstant::SetTextureCollectionParameterValueEditorOnly(con
 	SetTextureCollectionParameterValueInternal(ParameterInfo,Value);
 }
 
+void UMaterialInstanceConstant::SetParameterCollectionParameterValueEditorOnly(const FMaterialParameterInfo& ParameterInfo, UMaterialParameterCollection* Value)
+{
+	check(GIsEditor || IsRunningCommandlet());
+	SetParameterCollectionParameterValueInternal(ParameterInfo, Value);
+}
+
 void UMaterialInstanceConstant::SetRuntimeVirtualTextureParameterValueEditorOnly(const FMaterialParameterInfo& ParameterInfo, URuntimeVirtualTexture* Value)
 {
 	check(GIsEditor || IsRunningCommandlet());
diff --git a/Engine/Source/Runtime/Engine/Public/Materials/MaterialInstance.h b/Engine/Source/Runtime/Engine/Public/Materials/MaterialInstance.h
--- a/Engine/Source/Runtime/Engine/Public/Materials/MaterialInstance.h
+++ b/Engine/Source/Runtime/Engine/Public/Materials/MaterialInstance.h
@@ -1118,6 +1118,7 @@ public:
 #if WITH_EDITORONLY_DATA
 
 	ENGINE_API void SetStaticSwitchParameterValueEditorOnly(const FMaterialParameterInfo& ParameterInfo, bool Value);
+	ENGINE_API void SetStaticComponentMaskParameterValueEditorOnly(const FMaterialParameterInfo& ParameterInfo, FStaticComponentMaskValue Value);
 #endif // WITH_EDITORONLY_DATA
 
 	ENGINE_API virtual void GetAllParametersOfType(EMaterialParameterType Type, TMap<FMaterialParameterInfo, FMaterialParameterMetadata>& OutParameters) const override;
diff --git a/Engine/Source/Runtime/Engine/Public/Materials/MaterialInstanceConstant.h b/Engine/Source/Runtime/Engine/Public/Materials/MaterialInstanceConstant.h
--- a/Engine/Source/Runtime/Engine/Public/Materials/MaterialInstanceConstant.h
+++ b/Engine/Source/Runtime/Engine/Public/Materials/MaterialInstanceConstant.h
@@ -94,10 +94,12 @@ class UMaterialInstanceConstant : public UMaterialInstance
 	 * @param Value - The value to set.
 	 */
 	ENGINE_API void SetVectorParameterValueEditorOnly(const FMaterialParameterInfo& ParameterInfo, FLinearColor Value);
+	ENGINE_API void SetDoubleVectorParameterValueEditorOnly(const FMaterialParameterInfo& ParameterInfo, FVector4d Value);
 	ENGINE_API void SetScalarParameterValueEditorOnly(const FMaterialParameterInfo& ParameterInfo, float Value);
 	ENGINE_API void SetScalarParameterAtlasEditorOnly(const FMaterialParameterInfo& ParameterInfo, FScalarParameterAtlasInstanceData AtlasData);
 	ENGINE_API void SetTextureParameterValueEditorOnly(const FMaterialParameterInfo& ParameterInfo, class UTexture* Value);
 	ENGINE_API void SetTextureCollectionParameterValueEditorOnly(const FMaterialParameterInfo& ParameterInfo, class UTextureCollection* Value);
+	ENGINE_API void SetParameterCollectionParameterValueEditorOnly(const FMaterialParameterInfo& ParameterInfo, class UMaterialParameterCollection* Value);
 	ENGINE_API void SetRuntimeVirtualTextureParameterValueEditorOnly(const FMaterialParameterInfo& ParameterInfo, class URuntimeVirtualTexture* Value);
 	ENGINE_API void SetSparseVolumeTextureParameterValueEditorOnly(const FMaterialParameterInfo& ParameterInfo, class USparseVolumeTexture* Value);
 	ENGINE_API void SetFontParameterValueEditorOnly(const FMaterialParameterInfo& ParameterInfo, class UFont* FontValue, int32 FontPage);
