Fix shader count widget in material instance editor.
Before the fix it was showing the shader count for the parent material instead of the instance.
#rnx
#Synced-CL 52290580

diff --git a/Engine/Source/Editor/MaterialEditor/Private/MaterialStats.cpp b/Engine/Source/Editor/MaterialEditor/Private/MaterialStats.cpp
--- a/Engine/Source/Editor/MaterialEditor/Private/MaterialStats.cpp
+++ b/Engine/Source/Editor/MaterialEditor/Private/MaterialStats.cpp
@@ -155,7 +155,7 @@ FText FShaderPlatformSettings::GetShaderCodeText(const EMaterialQualityLevel::Ty
 FText FShaderPlatformSettings::GetShaderCountText(const EMaterialQualityLevel::Type QualityType, const int32 InstanceIndex)
 {
 	// Preview material always gets translated, so use its parameters
-	const UMaterialInterface* CurrentMaterial = Material ? static_cast<UMaterialInterface*>(Material) : static_cast<UMaterialInterface*>(MaterialInstance);
+	const UMaterialInterface* CurrentMaterial = MaterialInstance ? static_cast<UMaterialInterface*>(MaterialInstance) : static_cast<UMaterialInterface*>(Material);
 
 	if (!CurrentMaterial)
 	{
