Fixed horizontal scrollbar behavior in material stats shader list view

#Synced-CL 52169855

diff --git a/Engine/Source/Editor/MaterialEditor/Private/MaterialStats.cpp b/Engine/Source/Editor/MaterialEditor/Private/MaterialStats.cpp
--- a/Engine/Source/Editor/MaterialEditor/Private/MaterialStats.cpp
+++ b/Engine/Source/Editor/MaterialEditor/Private/MaterialStats.cpp
@@ -1264,6 +1264,7 @@ TSharedRef<SDockTab> FMaterialStats::SpawnTab_ShaderCode(const FSpawnTabArgs& Ar
 					[
 						SNew(SScrollBox)
 						.Orientation(EOrientation::Orient_Horizontal)
+						.ConsumeMouseWheel(EConsumeMouseWheel::Never)
 						+ SScrollBox::Slot()
 						[
 							ShaderListView
