Detect clip()/discard in custom HLSL expressions for bWritesEveryPixel

Adds FMaterialCompilationOutput::bUsesCustomHLSLPixelDiscard, set during
material translation when a compiled UMaterialExpressionCustom contains
clip() or discard in its Code. Combined with bUsesOpacityMask to correctly
gate the bWritesEveryPixel optimisation for masked materials, replacing the
temporary workaround that hard-coded bUsesOpacityMask=true.

Also adds ContainsClipInstruction (Search/Yes/No) on UMaterialExpressionCustom
so material authors can override the automatic string search. This is useful when a
false positive is suppressing the optimisation, or when clip is somehow present in
included code outside the Code string. Now displays a message on the node widget to
indicate that a clip was found.

#Synced-CL 52353515

diff --git a/Engine/Source/Editor/MaterialEditor/Private/MaterialNodes/SGraphNodeMaterialCustom.cpp b/Engine/Source/Editor/MaterialEditor/Private/MaterialNodes/SGraphNodeMaterialCustom.cpp
--- a/Engine/Source/Editor/MaterialEditor/Private/MaterialNodes/SGraphNodeMaterialCustom.cpp
+++ b/Engine/Source/Editor/MaterialEditor/Private/MaterialNodes/SGraphNodeMaterialCustom.cpp
@@ -2,11 +2,15 @@
 
 #include "MaterialEditor/MaterialNodes/SGraphNodeMaterialCustom.h"
 
+#include "GraphEditorSettings.h"
+#include "MaterialGraph/MaterialGraphNode_Custom.h"
+#include "Materials/MaterialExpressionCustom.h"
+#include "Styling/AppStyle.h"
 #include "Templates/Casts.h"
-#include "Widgets/Input/SCheckBox.h"
 #include "Widgets/Images/SImage.h"
-#include "MaterialGraph/MaterialGraphNode_Custom.h"
-#include "GraphEditorSettings.h"
+#include "Widgets/Input/SCheckBox.h"
+#include "Widgets/SBoxPanel.h"
+#include "Widgets/Text/STextBlock.h"
 
 void SGraphNodeMaterialCustom::Construct(const FArguments& InArgs, class UMaterialGraphNode* InNode)
 {
@@ -34,6 +38,7 @@ void SGraphNodeMaterialCustom::Construct(const FArguments& InArgs, class UMateri
 void SGraphNodeMaterialCustom::CreateBelowPinControls(TSharedPtr<SVerticalBox> MainBox)
 {
 	UMaterialGraphNode_Custom* CustomNode = CastChecked<UMaterialGraphNode_Custom>(GraphNode);
+	WeakCustomExpression = Cast<UMaterialExpressionCustom>(MaterialNode->MaterialExpression);
 
 	TAttribute<FText> GetText;
 	GetText.BindUObject(CustomNode, &UMaterialGraphNode_Custom::GetHlslText);
@@ -71,6 +76,17 @@ void SGraphNodeMaterialCustom::CreateBelowPinControls(TSharedPtr<SVerticalBox> M
 			PreviewBox.ToSharedRef()
 		]
 	];
+
+	MainBox->AddSlot()
+	.AutoHeight()
+	.Padding(FMargin(10.f, 0.f, 10.f, 5.f))
+	[
+		SNew(STextBlock)
+		.Text(NSLOCTEXT("MaterialEditor", "CustomNodePixelDiscard", "The code contains clip() or discard instructions."))
+		.ToolTipText(NSLOCTEXT("MaterialEditor", "CustomNodePixelDiscardWarning_Tooltip", "Use Blend Mode Masked to create supporting shaders.\nUse ContainsClipInstruction=No to override if this is a false positive."))
+		.Visibility(this, &SGraphNodeMaterialCustom::PixelDiscardVisibility)
+		.AutoWrapText(true)
+	];
 }
 
 EVisibility SGraphNodeMaterialCustom::CodeVisibility() const
@@ -86,6 +102,16 @@ EVisibility SGraphNodeMaterialCustom::CodeVisibility() const
 	return (State == ECheckBoxState::Checked) ? EVisibility::Visible : EVisibility::Collapsed;
 }
 
+EVisibility SGraphNodeMaterialCustom::PixelDiscardVisibility() const
+{
+	UMaterialExpressionCustom const* CustomExpression = WeakCustomExpression.Get();
+	if (CustomExpression == nullptr)
+	{
+		return EVisibility::Collapsed;
+	}
+	return CustomExpression->HasPixelDiscard() ? EVisibility::Visible : EVisibility::Collapsed;
+}
+
 void SGraphNodeMaterialCustom::CreateAdvancedViewArrow(TSharedPtr<SVerticalBox> MainBox)
 {
 	UMaterialGraphNode_Custom* CustomNode = CastChecked<UMaterialGraphNode_Custom>(GraphNode);
diff --git a/Engine/Source/Editor/MaterialEditor/Public/MaterialEditor/MaterialNodes/SGraphNodeMaterialCustom.h b/Engine/Source/Editor/MaterialEditor/Public/MaterialEditor/MaterialNodes/SGraphNodeMaterialCustom.h
--- a/Engine/Source/Editor/MaterialEditor/Public/MaterialEditor/MaterialNodes/SGraphNodeMaterialCustom.h
+++ b/Engine/Source/Editor/MaterialEditor/Public/MaterialEditor/MaterialNodes/SGraphNodeMaterialCustom.h
@@ -4,8 +4,11 @@
 
 #include "MaterialEditor/MaterialNodes/SGraphNodeMaterialBase.h"
 #include "Text/HLSLSyntaxHighlighterMarshaller.h"
+#include "UObject/WeakObjectPtr.h"
 #include "Widgets/Input/SMultiLineEditableTextBox.h"
 
+class UMaterialExpressionCustom;
+
 #define UE_API MATERIALEDITOR_API
 
 template <>
@@ -39,11 +42,13 @@ protected:
 
 private:
 	UE_API EVisibility CodeVisibility() const;
+	UE_API EVisibility PixelDiscardVisibility() const;
 
 protected:
 	/** Shader text box */
 	TSharedPtr<FHLSLSyntaxHighlighterMarshaller> SyntaxHighlighter;
 	TSharedPtr<SMultiLineEditableTextBox> ShaderTextBox;
+	TWeakObjectPtr<UMaterialExpressionCustom> WeakCustomExpression;
 };
 
 #undef UE_API
diff --git a/Engine/Source/Runtime/Engine/Private/Materials/HLSLMaterialTranslator.cpp b/Engine/Source/Runtime/Engine/Private/Materials/HLSLMaterialTranslator.cpp
--- a/Engine/Source/Runtime/Engine/Private/Materials/HLSLMaterialTranslator.cpp
+++ b/Engine/Source/Runtime/Engine/Private/Materials/HLSLMaterialTranslator.cpp
@@ -14278,6 +14278,13 @@ int32 FHLSLMaterialTranslator::CustomExpression( class UMaterialExpressionCustom
 
 		int32 CustomExpressionIndex = CustomExpressions.Num();
 		FString Code = Custom->Code;
+
+		// Detect clip() or discard. These conditionally kill pixels and may require additional shader permutations. 
+		if (Custom->HasPixelDiscard())
+		{
+			MaterialCompilationOutput.bUsesCustomHLSLPixelDiscard = true;
+		}
+
 		if (!Code.Contains(TEXT("return")))
 		{
 			Code = FString("return ") + Code + TEXT(";");
diff --git a/Engine/Source/Runtime/Engine/Private/Materials/Material.cpp b/Engine/Source/Runtime/Engine/Private/Materials/Material.cpp
--- a/Engine/Source/Runtime/Engine/Private/Materials/Material.cpp
+++ b/Engine/Source/Runtime/Engine/Private/Materials/Material.cpp
@@ -5374,9 +5374,6 @@ void UMaterial::PostEditChangePropertyInternal(FPropertyChangedEvent& PropertyCh
 	// Check for distortion in material 
 	bUsesDistortion = RefractionMethod != RM_None;
 
-	//If we can be sure this material would be the same opaque as it is masked then allow it to be assumed opaque.
-	bCanMaskedBeAssumedOpaque = !EditorOnly->OpacityMask.Expression && !(EditorOnly->OpacityMask.UseConstant && EditorOnly->OpacityMask.Constant < 0.999f) && !bUseMaterialAttributes;
-
 	// If BLEND_TranslucentColoredTransmittance is selected while Substrate is not enabled, force BLEND_Translucent blend mode
 	if (!Substrate::IsSubstrateEnabled() && BlendMode == BLEND_TranslucentColoredTransmittance)
 	{
diff --git a/Engine/Source/Runtime/Engine/Private/Materials/MaterialExpressions.cpp b/Engine/Source/Runtime/Engine/Private/Materials/MaterialExpressions.cpp
--- a/Engine/Source/Runtime/Engine/Private/Materials/MaterialExpressions.cpp
+++ b/Engine/Source/Runtime/Engine/Private/Materials/MaterialExpressions.cpp
@@ -12932,6 +12932,12 @@ void UMaterialExpressionCustom::GetIncludeFilePaths(TSet<FString>& OutIncludeFil
 {
 	OutIncludeFilePaths.Append(IncludeFilePaths);
 }
+
+bool UMaterialExpressionCustom::HasPixelDiscard() const
+{
+	return ContainsClipInstruction == CMCI_Yes ||
+		(ContainsClipInstruction == CMCI_Search && (Code.Contains(TEXT("clip("), ESearchCase::CaseSensitive) || Code.Contains(TEXT("discard"), ESearchCase::CaseSensitive)));
+}
 #endif // WITH_EDITOR
 
 void UMaterialExpressionCustom::Serialize(FStructuredArchive::FRecord Record)
diff --git a/Engine/Source/Runtime/Engine/Private/Materials/MaterialShared.cpp b/Engine/Source/Runtime/Engine/Private/Materials/MaterialShared.cpp
--- a/Engine/Source/Runtime/Engine/Private/Materials/MaterialShared.cpp
+++ b/Engine/Source/Runtime/Engine/Private/Materials/MaterialShared.cpp
@@ -6086,15 +6086,12 @@ void FMaterialShaderParameters::InitializeShaderLayoutParamsFromCompOutput(const
 	bUsesTemporalResponsiveness = InCompilationOutput.bUsesTemporalResponsiveness;
 	bUsesMotionVectorWorldOffset = InCompilationOutput.bUsesMotionVectorWorldOffset;
 	
-	// Need to consider whether opacity mask is used when calculating WritesEveryPixel().
+	// Consider whether opacity mask is actually used when calculating WritesEveryPixel().
+	// Treat the case where a custom HLSL expression is discarding pixels as equivalent to the use of opacity mask here.
 	const bool bVFTypeSupportsNullPixelShader = !InMaterial->IsUsedWithInstancedStaticMeshes();
-	// Temporarily assume bUsesOpacityMask is true while we fix issues that come from not detecting the use of clip() in the shader.
-	const bool bUsesOpacityMask = true; // InCompilationOutput.bUsesOpacityMask
-	bWritesEveryPixel = InMaterial->WritesEveryPixel(/*bShadowPass*/false, bVFTypeSupportsNullPixelShader, !bUsesOpacityMask);
-	bWritesEveryPixelShadowPass = InMaterial->WritesEveryPixel(/*bShadowPass*/true, bVFTypeSupportsNullPixelShader, !bUsesOpacityMask);
-	
-	// Also override blend mode if opacity mask is not used.
-	BlendMode = BlendMode == BLEND_Masked && !bUsesOpacityMask ? BLEND_Opaque : BlendMode;
+	const bool bUsesOpacityMaskOrDiscard = InCompilationOutput.bUsesOpacityMask || InCompilationOutput.bUsesCustomHLSLPixelDiscard;
+	bWritesEveryPixel = InMaterial->WritesEveryPixel(/*bShadowPass*/false, bVFTypeSupportsNullPixelShader, !bUsesOpacityMaskOrDiscard);
+	bWritesEveryPixelShadowPass = InMaterial->WritesEveryPixel(/*bShadowPass*/true, bVFTypeSupportsNullPixelShader, !bUsesOpacityMaskOrDiscard);
 }
 
 FCompilationShaderParameters::FCompilationShaderParameters(const FMaterial* InMaterial, const FMaterialCompilationOutput& InCompilationOutput)
diff --git a/Engine/Source/Runtime/Engine/Public/MaterialShared.h b/Engine/Source/Runtime/Engine/Public/MaterialShared.h
--- a/Engine/Source/Runtime/Engine/Public/MaterialShared.h
+++ b/Engine/Source/Runtime/Engine/Public/MaterialShared.h
@@ -1084,6 +1084,10 @@ public:
 
 	/** True if the opacity mask isn't constant 1.0 (after constant folding). */
 	LAYOUT_BITFIELD(uint8, bUsesOpacityMask, 1);
+
+	/** True if any UMaterialExpressionCustom compiled into this material uses clip() or discard
+	    to conditionally discard pixels, making it behave like an opacity mask. */
+	LAYOUT_BITFIELD(uint8, bUsesCustomHLSLPixelDiscard, 1);
 };
 
 struct FDebugShaderPipelineInfo
diff --git a/Engine/Source/Runtime/Engine/Public/Materials/MaterialExpressionCustom.h b/Engine/Source/Runtime/Engine/Public/Materials/MaterialExpressionCustom.h
--- a/Engine/Source/Runtime/Engine/Public/Materials/MaterialExpressionCustom.h
+++ b/Engine/Source/Runtime/Engine/Public/Materials/MaterialExpressionCustom.h
@@ -59,6 +59,17 @@ struct FCustomDefine
 	FString DefineValue;
 };
 
+UENUM()
+enum ECustomMaterialClipInstruction : int
+{
+	/** Search the Code string for clip() or discard at compile time to determine whether pixel discarding is possible. */
+	CMCI_Search	UMETA(DisplayName="Search"),
+	/** Assume this expression may discard pixels, regardless of what the Code contains. */
+	CMCI_Yes	UMETA(DisplayName="Yes"),
+	/** Assume this expression never discards pixels, regardless of what the Code contains. */
+	CMCI_No		UMETA(DisplayName="No"),
+};
+
 UCLASS(collapsecategories, hidecategories=Object, MinimalAPI)
 class UMaterialExpressionCustom : public UMaterialExpression
 {
@@ -85,6 +96,10 @@ class UMaterialExpressionCustom : public UMaterialExpression
 	UPROPERTY(EditAnywhere, Category=MaterialExpressionCustom)
 	TArray<FString> IncludeFilePaths;
 
+	/** Controls whether this expression is considered to contain a clip() or discard instruction that may conditionally discard pixels. */
+	UPROPERTY(EditAnywhere, Category=MaterialExpressionCustom)
+	TEnumAsByte<enum ECustomMaterialClipInstruction> ContainsClipInstruction = ECustomMaterialClipInstruction::CMCI_Search;
+
 	UPROPERTY(VisibleAnywhere, Category=MaterialExpressionCustom)
 	bool ShowCode;
 
@@ -96,8 +111,8 @@ class UMaterialExpressionCustom : public UMaterialExpression
 	virtual void Serialize(FStructuredArchive::FRecord Record) override;
 	//~ End UObject Interface.
 	
-	//~ Begin UMaterialExpression Interface
 #if WITH_EDITOR
+	//~ Begin UMaterialExpression Interface
 	virtual void Build(MIR::FEmitter& Emitter) override;
 	virtual int32 Compile(class FMaterialCompiler* Compiler, int32 OutputIndex) override;
 	virtual void GetCaption(TArray<FString>& OutCaptions) const override;
@@ -109,8 +124,14 @@ class UMaterialExpressionCustom : public UMaterialExpression
 	virtual EMaterialValueType GetOutputValueType(int32 OutputIndex) override;
 	virtual bool IsResultMaterialAttributes(int32 OutputIndex) override;
 	virtual void GetIncludeFilePaths(TSet<FString>& OutIncludeFilePaths) const override;
-#endif // WITH_EDITOR
 	//~ End UMaterialExpression Interface
+
+	/** 
+	 * Returns true if this expression may conditionally discard pixels using clip() or discard.
+	 * The string search is conservative so ContainsClipInstruction lets the material author override if necessary.
+	 */
+	ENGINE_API bool HasPixelDiscard() const;
+#endif // WITH_EDITOR
 };
 
 
