Ensure duplicated Material Functions have bExposeToLibrary as false by default to avoid exposing them automatically to the Material Editor Palette.

[RN] minor | Rendering Architecture
#Synced-CL 52381300

diff --git a/Engine/Source/Editor/MaterialEditor/Private/MaterialEditor.cpp b/Engine/Source/Editor/MaterialEditor/Private/MaterialEditor.cpp
--- a/Engine/Source/Editor/MaterialEditor/Private/MaterialEditor.cpp
+++ b/Engine/Source/Editor/MaterialEditor/Private/MaterialEditor.cpp
@@ -2876,6 +2876,10 @@ bool FMaterialEditor::UpdateOriginalMaterial()
 		bMaterialDirty = false;
 		bStatsFromPreviewMaterial = false;
 
+		// Since MaterialFunction->ParentFunction gets initialized as a duplicate here and bExposeToLibrary 
+		// is protected against duplicate for non-transient objects, we need to manually copy it
+		MaterialFunction->ParentFunction->bExposeToLibrary = MaterialFunction->bExposeToLibrary;
+
 		UMaterialEditingLibrary::UpdateMaterialFunction(MaterialFunction->ParentFunction, Material);
 	}
 	// Handle propagation of the material being edited
diff --git a/Engine/Source/Runtime/Engine/Private/Materials/MaterialExpressions.cpp b/Engine/Source/Runtime/Engine/Private/Materials/MaterialExpressions.cpp
--- a/Engine/Source/Runtime/Engine/Private/Materials/MaterialExpressions.cpp
+++ b/Engine/Source/Runtime/Engine/Private/Materials/MaterialExpressions.cpp
@@ -13895,6 +13895,18 @@ void UMaterialFunction::GetAllCustomOutputExpressions(TArray<class UMaterialExpr
 
 #endif //WITH_EDITORONLY_DATA
 
+void UMaterialFunction::PostDuplicate(bool bDuplicateForPIE)
+{
+	Super::PostDuplicate(bDuplicateForPIE);
+
+	// Only clear when duplicating to a real package,
+	// not when the material editor creates its transient editing copy
+	if (!GetOutermost()->HasAnyFlags(RF_Transient))
+	{
+		bExposeToLibrary = false;
+	}
+}
+
 void UMaterialFunction::PostLoad()
 {
 	LLM_SCOPE(ELLMTag::Materials);
diff --git a/Engine/Source/Runtime/Engine/Public/Materials/MaterialFunction.h b/Engine/Source/Runtime/Engine/Public/Materials/MaterialFunction.h
--- a/Engine/Source/Runtime/Engine/Public/Materials/MaterialFunction.h
+++ b/Engine/Source/Runtime/Engine/Public/Materials/MaterialFunction.h
@@ -130,6 +130,7 @@ public:
 #endif // WITH_EDITOR
 	virtual void PreSave(FObjectPreSaveContext ObjectSaveContext) override;
 	virtual void Serialize(FArchive& Ar) override;
+	virtual void PostDuplicate(bool bDuplicateForPIE) override;
 	virtual void PostLoad() override;
 #if WITH_EDITORONLY_DATA
 	static void DeclareConstructClasses(TArray<FTopLevelAssetPath>& OutConstructClasses, const UClass* SpecificSubclass);
