- Add a tooltip to the material editor node preview that displays the sampled value under the cursor

#Synced-CL 52392748

diff --git a/Engine/Source/Editor/MaterialEditor/Private/MaterialEditor.cpp b/Engine/Source/Editor/MaterialEditor/Private/MaterialEditor.cpp
--- a/Engine/Source/Editor/MaterialEditor/Private/MaterialEditor.cpp
+++ b/Engine/Source/Editor/MaterialEditor/Private/MaterialEditor.cpp
@@ -33,7 +33,6 @@
 #include "MaterialGraph/MaterialGraphSchema.h"
 #include "MaterialEditor/PreviewMaterial.h"
 #include "ThumbnailRendering/SceneThumbnailInfoWithPrimitive.h"
-#include "Particles/ParticleSystemComponent.h"
 #include "Materials/MaterialAttributeDefinitionMap.h"
 #include "Materials/MaterialInstance.h"
 #include "Materials/MaterialInstanceDynamic.h"
@@ -356,6 +355,16 @@ int32 FMatExpressionPreview::CompilePropertyAndSetMaterialProperty(EMaterialProp
 		int32 PreviewCodeChunk = INDEX_NONE;
 		PreviewCodeChunk = Expression->CompilePreview(Compiler, OutputIndex);
 
+		if (PreviewCodeChunk != INDEX_NONE)
+		{
+			const EMaterialValueType CompiledType = Compiler->GetType(PreviewCodeChunk);
+			CompiledPreviewNumChannels = static_cast<int32>(GetNumComponents(MakeNonLWCType(CompiledType)));
+		}
+		else
+		{
+			CompiledPreviewNumChannels = 0;
+		}
+
 		// Get back into gamma corrected space, as DrawTile does not do this adjustment.
 		Ret = Compiler->Power(Compiler->Max(PreviewCodeChunk, Compiler->Constant(0)), Compiler->Constant(1.f / 2.2f));
 	}
diff --git a/Engine/Source/Editor/MaterialEditor/Private/MaterialEditor.h b/Engine/Source/Editor/MaterialEditor/Private/MaterialEditor.h
--- a/Engine/Source/Editor/MaterialEditor/Private/MaterialEditor.h
+++ b/Engine/Source/Editor/MaterialEditor/Private/MaterialEditor.h
@@ -176,9 +176,12 @@ public:
 
 	float UnrelatedNodesOpacity;
 
+	int32 GetCompiledPreviewNumChannels() const { return CompiledPreviewNumChannels; }
+
 private:
 	TWeakObjectPtr<UMaterialExpression> Expression;
 	FGuid Id;
+	mutable int32 CompiledPreviewNumChannels = 0;
 };
 
 /** Wrapper for each material expression, including a trimmed name */
diff --git a/Engine/Source/Editor/MaterialEditor/Private/MaterialNodes/SGraphNodeMaterialBase.cpp b/Engine/Source/Editor/MaterialEditor/Private/MaterialNodes/SGraphNodeMaterialBase.cpp
--- a/Engine/Source/Editor/MaterialEditor/Private/MaterialNodes/SGraphNodeMaterialBase.cpp
+++ b/Engine/Source/Editor/MaterialEditor/Private/MaterialNodes/SGraphNodeMaterialBase.cpp
@@ -28,6 +28,7 @@
 #include "Materials/MaterialExpressionMaterialFunctionCall.h"
 #include "Materials/MaterialExpressionSubstrate.h"
 #include "Materials/MaterialFunction.h"
+#include "MaterialEditor.h"
 #include "Math/Color.h"
 #include "Math/IntPoint.h"
 #include "Math/IntRect.h"
@@ -38,6 +39,7 @@
 #include "Misc/Optional.h"
 #include "RHI.h"
 #include "RHICommandList.h"
+#include "RenderGraphUtils.h"
 #include "Rendering/DrawElements.h"
 #include "Rendering/SlateRenderer.h"
 #include "RenderingThread.h"
@@ -61,12 +63,14 @@
 #include "Widgets/SBoxPanel.h"
 #include "Widgets/SNullWidget.h"
 #include "Widgets/SOverlay.h"
+#include "Widgets/SToolTip.h"
 #include "Widgets/SViewport.h"
 #include "Widgets/Text/STextBlock.h"
 #include "Rendering/SubstrateMaterialShared.h"
 #include "MaterialEditor/SGraphSubstrateMaterial.h"
 #include "MaterialShared.h"
 #include "RenderGraphBuilder.h"
+#include "RHIGPUReadback.h"
 
 class FWidgetStyle;
 class SWidget;
@@ -183,7 +187,7 @@ void FPreviewViewport::OnDrawViewport( const FGeometry& AllottedGeometry, const
 
 	bool bIsRealtime = MaterialNode->RealtimeDelegate.IsBound() ? MaterialNode->RealtimeDelegate.Execute() : false;
 
-	if (PreviewElement->BeginRenderingCanvas(CanvasRect, ClippingRect, MaterialNode, bIsRealtime))
+	if (PreviewElement->BeginRenderingCanvas(CanvasRect, ClippingRect, MaterialNode, bIsRealtime, HoveredUV.IsSet()))
 	{
 		// Draw above everything else
 		uint32 PreviewLayer = LayerId+1;
@@ -196,6 +200,30 @@ FIntPoint FPreviewViewport::GetSize() const
 	return FIntPoint(96,96);
 }
 
+FReply FPreviewViewport::OnMouseMove(const FGeometry& MyGeometry, const FPointerEvent& MouseEvent)
+{
+	const FVector2D LocalPos = MyGeometry.AbsoluteToLocal(MouseEvent.GetScreenSpacePosition());
+	const FVector2D Size = MyGeometry.GetLocalSize();
+	HoveredUV = FVector2D(LocalPos.X / Size.X, LocalPos.Y / Size.Y);
+	return FReply::Unhandled();
+}
+
+void FPreviewViewport::OnMouseEnter(const FGeometry& MyGeometry, const FPointerEvent& MouseEvent)
+{
+	const FVector2D LocalPos = MyGeometry.AbsoluteToLocal(MouseEvent.GetScreenSpacePosition());
+	const FVector2D Size = MyGeometry.GetLocalSize();
+	HoveredUV = FVector2D(LocalPos.X / Size.X, LocalPos.Y / Size.Y);
+}
+
+void FPreviewViewport::OnMouseLeave(const FPointerEvent& MouseEvent)
+{
+	HoveredUV.Reset();
+	if (PreviewElement)
+	{
+		PreviewElement->ClearReadback();
+	}
+}
+
 void FPreviewViewport::UpdatePreviewNodeRenderProxy()
 {
 	if (PreviewElement.IsValid())
@@ -209,17 +237,11 @@ void FPreviewViewport::UpdatePreviewNodeRenderProxy()
 
 FPreviewElement::FPreviewElement()
 	: RenderTarget(new FSlateMaterialPreviewRenderTarget)
-	, ExpressionPreview(nullptr)
-	, bIsRealtime(false)
+	, ReadbackRenderTarget(new FSlateMaterialPreviewRenderTarget)
 {
 }
 
-FPreviewElement::~FPreviewElement()
-{
-	delete RenderTarget;
-}
-
-bool FPreviewElement::BeginRenderingCanvas( const FIntRect& InCanvasRect, const FIntRect& InClippingRect, UMaterialGraphNode* InGraphNode, bool bInIsRealtime )
+bool FPreviewElement::BeginRenderingCanvas(const FIntRect& InCanvasRect, const FIntRect& InClippingRect, UMaterialGraphNode* InGraphNode, bool bInIsRealtime, bool bInShouldReadback)
 {
 	if(InCanvasRect.Size().X > 0 && InCanvasRect.Size().Y > 0 && InClippingRect.Size().X > 0 && InClippingRect.Size().Y > 0 && InGraphNode != NULL)
 	{
@@ -236,6 +258,7 @@ bool FPreviewElement::BeginRenderingCanvas( const FIntRect& InCanvasRect, const
 			FMaterialRenderProxy* RenderProxy;
 			/** Whether preview is using realtime values */
 			bool bIsRealtime;
+			bool bShouldReadback;
 		};
 
 		FPreviewRenderInfo RenderInfo;
@@ -243,6 +266,7 @@ bool FPreviewElement::BeginRenderingCanvas( const FIntRect& InCanvasRect, const
 		RenderInfo.ClippingRect = InClippingRect;
 		RenderInfo.RenderProxy = InGraphNode->GetExpressionPreview();
 		RenderInfo.bIsRealtime = bInIsRealtime;
+		RenderInfo.bShouldReadback = bInShouldReadback;
 
 		FPreviewElement* PreviewElement = this;
 		ENQUEUE_RENDER_COMMAND(BeginRenderingPreviewCanvas)(
@@ -252,6 +276,7 @@ bool FPreviewElement::BeginRenderingCanvas( const FIntRect& InCanvasRect, const
 				PreviewElement->RenderTarget->SetClippingRect(RenderInfo.ClippingRect);
 				PreviewElement->ExpressionPreview = RenderInfo.RenderProxy;
 				PreviewElement->bIsRealtime = RenderInfo.bIsRealtime;
+				PreviewElement->bShouldReadback = RenderInfo.bShouldReadback;
 			}
 			);
 		return true;
@@ -272,29 +297,127 @@ void FPreviewElement::UpdateExpressionPreview(UMaterialGraphNode* MaterialNode)
 	);
 }
 
+TOptional<FLinearColor> FPreviewElement::GetReadbackValue(const FVector2D& UV) const
+{
+	UE::TScopeLock ReadLock(ReadbackDataGuard);
+	if (ReadbackData.IsEmpty())
+	{
+		return {};
+	}
+
+	const int32 X = FMath::Clamp(FMath::CeilToInt(UV.X * ReadbackDimension), 0, ReadbackDimension - 1);
+	const int32 Y = FMath::Clamp(FMath::CeilToInt(UV.Y * ReadbackDimension), 0, ReadbackDimension - 1);
+
+	FLinearColor Value = ReadbackData[(Y * ReadbackDimension) + X].GetFloats();
+
+	// Note: Gamma is baked into the preview shader (see FMatExpressionPreview::CompilePropertyAndSetMaterialProperty) so we need to undo it here
+	const float Gamma = 2.2f;
+	Value.R = FMath::Pow(Value.R, Gamma);
+	Value.G = FMath::Pow(Value.G, Gamma);
+	Value.B = FMath::Pow(Value.B, Gamma);
+	return Value;
+}
+
+void FPreviewElement::ClearReadback()
+{
+	ENQUEUE_RENDER_COMMAND(ClearReadbackPreviewElement)(
+		[this](FRHICommandListImmediate& RHICmdList)
+		{
+			GPUTextureReadback.Reset();
+			bReadbackPending = false;
+
+			UE::TScopeLock ReadLock(ReadbackDataGuard);
+			ReadbackData.Empty();
+		}
+	);
+}
+
 void FPreviewElement::Draw_RenderThread(FRDGBuilder& GraphBuilder, const FDrawPassInputs& Inputs)
 {
-	if(ExpressionPreview)
+	if(!ExpressionPreview)
 	{
-		RenderTarget->SetRenderTargetTexture(Inputs.OutputTexture);
+		return;
+	}
+
+	RenderTarget->SetRenderTargetTexture(Inputs.OutputTexture);
+	{
+		// Check realtime mode for whether to pass current time to canvas
+		double CurrentTime = bIsRealtime ? (FApp::GetCurrentTime() - GStartTime) : 0.0;
+		float DeltaTime = bIsRealtime ? static_cast<float>(FApp::GetDeltaTime()) : 0.0f;
+
+		FCanvas* Canvas = GraphBuilder.AllocObject<FCanvas>(RenderTarget.Get(), nullptr, FGameTime::CreateUndilated(CurrentTime, DeltaTime), GMaxRHIFeatureLevel);
 		{
-			// Check realtime mode for whether to pass current time to canvas
-			double CurrentTime = bIsRealtime ? (FApp::GetCurrentTime() - GStartTime) : 0.0;
-			float DeltaTime = bIsRealtime ? static_cast<float>(FApp::GetDeltaTime()) : 0.0f;
+			Canvas->SetAllowedModes(0);
+			Canvas->SetRenderTargetRect(RenderTarget->GetViewRect());
+			Canvas->SetRenderTargetScissorRect(RenderTarget->GetClippingRect());
 
-			FCanvas* Canvas = GraphBuilder.AllocObject<FCanvas>(RenderTarget, nullptr, FGameTime::CreateUndilated(CurrentTime, DeltaTime), GMaxRHIFeatureLevel);
-			{
-				Canvas->SetAllowedModes(0);
-				Canvas->SetRenderTargetRect(RenderTarget->GetViewRect());
-				Canvas->SetRenderTargetScissorRect(RenderTarget->GetClippingRect());
+			FCanvasTileItem TileItem(FVector2D::ZeroVector, ExpressionPreview, RenderTarget->GetSizeXY());
+			Canvas->DrawItem(TileItem);
+			Canvas->Flush_RenderThread(GraphBuilder, true);
+		}
+	}
+	RenderTarget->ClearRenderTargetTexture();
 
-				FCanvasTileItem TileItem(FVector2D::ZeroVector, ExpressionPreview, RenderTarget->GetSizeXY());
-				Canvas->DrawItem(TileItem);
-				Canvas->Flush_RenderThread(GraphBuilder, true);
+	if (bShouldReadback)
+	{
+		Readback_RenderThread(GraphBuilder);
+		return;
+	}
+}
+
+void FPreviewElement::Readback_RenderThread(FRDGBuilder& GraphBuilder)
+{
+	// Do we have a readback pending?  If so don't issue another one until we process the pending one
+	if (bReadbackPending)
+	{
+		if (!GPUTextureReadback->IsReady())
+		{
+			return;
+		}
+
+		int32 RowPitchInPixels = 0;
+		const FFloat16Color* SrcData = static_cast<const FFloat16Color*>(GPUTextureReadback->Lock(RowPitchInPixels));
+		{
+			UE::TScopeLock WriteLock(ReadbackDataGuard);
+			ReadbackData.SetNumUninitialized(ReadbackDimension * ReadbackDimension);
+			for (int32 Y=0; Y < ReadbackDimension; ++Y)
+			{
+				FMemory::Memcpy(ReadbackData.GetData() + (Y * ReadbackDimension), SrcData + (Y * RowPitchInPixels), sizeof(FFloat16Color) * ReadbackDimension);
 			}
 		}
-		RenderTarget->ClearRenderTargetTexture();
+		GPUTextureReadback->Unlock();
+	}
+
+	// Issue a new render & readback request
+	const FIntPoint ReadbackSize(ReadbackDimension, ReadbackDimension);
+
+	const FRDGTextureDesc ReadbackTextureDesc = FRDGTextureDesc::Create2D(ReadbackSize, PF_FloatRGBA, FClearValueBinding::Black, TexCreate_RenderTargetable | TexCreate_ShaderResource);
+	FRDGTextureRef ReadbackTexture = GraphBuilder.CreateTexture(ReadbackTextureDesc, TEXT("MaterialNodeReadbackRT"));
+
+	ReadbackRenderTarget->SetViewRect(FIntRect(FIntPoint::ZeroValue, ReadbackSize));
+	ReadbackRenderTarget->SetClippingRect(FIntRect(FIntPoint::ZeroValue, ReadbackSize));
+	ReadbackRenderTarget->SetRenderTargetTexture(ReadbackTexture);
+
+	{
+		const double CurrentTime = bIsRealtime ? (FApp::GetCurrentTime() - GStartTime) : 0.0;
+		const float  DeltaTime = bIsRealtime ? static_cast<float>(FApp::GetDeltaTime()) : 0.0f;
+
+		FCanvas* ReadbackCanvas = GraphBuilder.AllocObject<FCanvas>(ReadbackRenderTarget.Get(), nullptr, FGameTime::CreateUndilated(CurrentTime, DeltaTime), GMaxRHIFeatureLevel);
+		ReadbackCanvas->SetAllowedModes(0);
+		ReadbackCanvas->SetRenderTargetRect(ReadbackRenderTarget->GetViewRect());
+		ReadbackCanvas->SetRenderTargetScissorRect(ReadbackRenderTarget->GetClippingRect());
+		FCanvasTileItem TileItem(FVector2D::ZeroVector, ExpressionPreview, FVector2D(ReadbackSize));
+		ReadbackCanvas->DrawItem(TileItem);
+		ReadbackCanvas->Flush_RenderThread(GraphBuilder, true);
+	}
+	ReadbackRenderTarget->ClearRenderTargetTexture();
+
+	if (!GPUTextureReadback)
+	{
+		GPUTextureReadback = MakeUnique<FRHIGPUTextureReadback>(TEXT("MaterialNodeReadback"));
 	}
+	AddEnqueueCopyPass(GraphBuilder, GPUTextureReadback.Get(), ReadbackTexture);
+	bReadbackPending = true;
 }
 
 /////////////////////////////////////////////////////
@@ -655,6 +778,14 @@ TSharedRef<SWidget> SGraphNodeMaterialBase::CreatePreviewWidget()
 			.MaxAspectRatio(1.0f)
 			.MaxDesiredHeight(ExpressionPreviewSize)
 			.Visibility(ExpressionPreviewVisibility())
+			.ToolTip(
+				SNew(SToolTip)
+				[
+					SNew(STextBlock)
+					.Text(this, &SGraphNodeMaterialBase::ExpressionPreviewValueText)
+				]
+			)
+			.Cursor(EMouseCursor::Crosshairs)
 			[
 				SNew(SBorder)
 				.Padding(CentralPadding)
@@ -670,10 +801,15 @@ TSharedRef<SWidget> SGraphNodeMaterialBase::CreatePreviewWidget()
 					.VAlign(VAlign_Bottom)
 					.Padding(0.0f,0.0f,8.0f,0.0f)
 					[
-						SNew(STextBlock)
-						.ShadowColorAndOpacity(FLinearColor::Black)
-						.ShadowOffset(FVector2D(1.0f, 1.0f))
-						.Text(this, &SGraphNodeMaterialBase::ExpressionPreviewOverlayText)
+						SNew(SBox)
+						.Cursor(EMouseCursor::Crosshairs)
+						.Visibility(EVisibility::HitTestInvisible)
+						[
+							SNew(STextBlock)
+							.ShadowColorAndOpacity(FLinearColor::Black)
+							.ShadowOffset(FVector2D(1.0f, 1.0f))
+							.Text(this, &SGraphNodeMaterialBase::ExpressionPreviewOverlayText)
+						]
 					]
 				]
 			];
@@ -685,7 +821,46 @@ TSharedRef<SWidget> SGraphNodeMaterialBase::CreatePreviewWidget()
 FText SGraphNodeMaterialBase::ExpressionPreviewOverlayText() const
 {
 	UMaterialExpression* MaterialExpression = MaterialNode->MaterialExpression;
-	return MaterialNode->MaterialExpression->GetPreviewOverlayText();
+	return MaterialExpression->GetPreviewOverlayText();
+}
+
+FText SGraphNodeMaterialBase::ExpressionPreviewValueText() const
+{
+	FText ValueText;
+	if (PreviewViewport->HoveredUV.IsSet())
+	{
+		TOptional<FLinearColor> Value = PreviewViewport->PreviewElement->GetReadbackValue(PreviewViewport->HoveredUV.GetValue());
+		if (Value.IsSet())
+		{
+			int32 NumChannels = 0;
+			if (FMaterialRenderProxy* Proxy = MaterialNode->GetExpressionPreview())
+			{
+				// FMatExpressionPreview is the only concrete type returned here
+				if (const FMatExpressionPreview* ExprPreview = static_cast<const FMatExpressionPreview*>(Proxy))
+				{
+					NumChannels = ExprPreview->GetCompiledPreviewNumChannels();
+				}
+			}
+
+			switch (NumChannels)
+			{
+			case 1:
+				ValueText = FText::FromString(FString::Printf(TEXT("%.2g"), Value->R));
+				break;
+			case 2:
+				ValueText = FText::FromString(FString::Printf(TEXT("%.2g, %.2g"), Value->R, Value->G));
+				break;
+			case 3:
+				ValueText = FText::FromString(FString::Printf(TEXT("%.2g, %.2g, %.2g"), Value->R, Value->G, Value->B));
+				break;
+			default:
+				ValueText = FText::FromString(FString::Printf(TEXT("%.2g, %.2g, %.2g, %.2g"), Value->R, Value->G, Value->B, Value->A));
+				break;
+			}
+		}
+	}
+
+	return ValueText;
 }
 
 EVisibility SGraphNodeMaterialBase::ExpressionPreviewVisibility() const
diff --git a/Engine/Source/Editor/MaterialEditor/Public/MaterialEditor/MaterialNodes/SGraphNodeMaterialBase.h b/Engine/Source/Editor/MaterialEditor/Public/MaterialEditor/MaterialNodes/SGraphNodeMaterialBase.h
--- a/Engine/Source/Editor/MaterialEditor/Public/MaterialEditor/MaterialNodes/SGraphNodeMaterialBase.h
+++ b/Engine/Source/Editor/MaterialEditor/Public/MaterialEditor/MaterialNodes/SGraphNodeMaterialBase.h
@@ -41,11 +41,16 @@ public:
 	virtual FIntPoint GetSize() const override;
 	virtual class FSlateShaderResource* GetViewportRenderTargetTexture() const override {return NULL;}
 	virtual bool RequiresVsync() const override {return false;}
+	virtual FReply OnMouseMove(const FGeometry& MyGeometry, const FPointerEvent& MouseEvent) override;
+	virtual void OnMouseEnter(const FGeometry& MyGeometry, const FPointerEvent& MouseEvent) override;
+	virtual void OnMouseLeave(const FPointerEvent& MouseEvent) override;
 
 	/** Material node to get expression preview from */
 	UMaterialGraphNode* MaterialNode;
 	/** Custom Slate Element to display the preview */
 	FThreadSafePreviewPtr PreviewElement;
+	/** The current UV we are hovered over, if any. */
+	TOptional<FVector2D> HoveredUV;
 private:
 	/**
 	 * Updates the expression preview render proxy from a graph node
@@ -57,7 +62,7 @@ class FPreviewElement : public ICustomSlateElement
 {
 public:
 	FPreviewElement();
-	~FPreviewElement();
+	virtual ~FPreviewElement() = default;
 
 	/**
 	 * Sets up the canvas for rendering
@@ -66,28 +71,56 @@ public:
 	 * @param	InClippingRect			How to clip the canvas tile
 	 * @param	InGraphNode				The graph node for the material preview
 	 * @param	bInIsRealtime			Whether preview is using realtime values
+	 * @param	bInShouldReadback		Whether we should readback the render target data or not
 	 *
 	 * @return	Whether there is anything to render
 	 */
-	bool BeginRenderingCanvas( const FIntRect& InCanvasRect, const FIntRect& InClippingRect, UMaterialGraphNode* InGraphNode, bool bInIsRealtime );
+	bool BeginRenderingCanvas(const FIntRect& InCanvasRect, const FIntRect& InClippingRect, UMaterialGraphNode* InGraphNode, bool bInIsRealtime, bool bInShouldReadback);
 
 	/**
 	 * Updates the expression preview render proxy from a graph node on the render thread
 	 */
 	void UpdateExpressionPreview(UMaterialGraphNode* PreviewNode);
+
+	/**
+	 * Returns the readback value at the provided coordinated, unset if no readback data existss.
+	 */
+	 TOptional<FLinearColor> GetReadbackValue(const FVector2D& UV) const;
+
+	 /**
+	  * Clear any pending readbacks that may be in flight to ensure we don't read stale data later.
+	  */
+	 void ClearReadback();
+
 private:
 	/**
 	 * ICustomSlateElement interface 
 	 */
 	virtual void Draw_RenderThread(FRDGBuilder& GraphBuilder, const FDrawPassInputs& Inputs) override;
 
+	void Readback_RenderThread(FRDGBuilder& GraphBuilder);
+
 private:
 	/** Render target that the canvas renders to */
-	class FSlateMaterialPreviewRenderTarget* RenderTarget;
+	TUniquePtr<class FSlateMaterialPreviewRenderTarget> RenderTarget;
 	/** Render proxy for the expression preview */
-	FMaterialRenderProxy* ExpressionPreview;
+	FMaterialRenderProxy* ExpressionPreview = nullptr;
 	/** Whether preview is using realtime values */
-	bool bIsRealtime;
+	bool bIsRealtime = false;
+
+	/** Used on the render thread to determine if we should readback the data or not. */
+	bool bShouldReadback = false;
+	/** Used on the render thread to determine if we have a readback in flight or not. */
+	bool bReadbackPending = false;
+	/** Readback render target, different from the regular render target as we need float range */
+	TUniquePtr<class FSlateMaterialPreviewRenderTarget> ReadbackRenderTarget;
+	/** Readback primitive */
+	TUniquePtr<class FRHIGPUTextureReadback> GPUTextureReadback;
+
+	static constexpr int32 ReadbackDimension = 128;
+
+	mutable UE::FMutex ReadbackDataGuard;
+	TArray<FFloat16Color> ReadbackData;
 };
 
 template<>
@@ -136,6 +169,9 @@ protected:
 	/** Returns text to over lay over the expression preview viewport */
 	UE_API FText ExpressionPreviewOverlayText() const;
 
+	/** Returns text to over lay as the tooltip value text */
+	UE_API FText ExpressionPreviewValueText() const;
+
 	/** Called when we click to preview node button */
 	UE_API FReply OnTogglePreviewClicked();
 
