- Fix material function call pins losing their PinCategory and PinName data after paste
- Reconstruct optional/required pin category and type-suffix pin names from the function input metadata post-paste

#Synced-CL 52445503

diff --git a/Engine/Source/Editor/MaterialEditor/Private/MaterialEditor.cpp b/Engine/Source/Editor/MaterialEditor/Private/MaterialEditor.cpp
--- a/Engine/Source/Editor/MaterialEditor/Private/MaterialEditor.cpp
+++ b/Engine/Source/Editor/MaterialEditor/Private/MaterialEditor.cpp
@@ -6464,6 +6464,30 @@ void FMaterialEditor::PostPasteMaterialExpression(UMaterialExpression* NewExpres
 		const bool bRecreateAndLinkNode = false;
 		FunctionCall->UpdateFromFunctionResource(bRecreateAndLinkNode);
 
+		// Some data is lost due to being transient on paste, we need to reconstruct that data now
+		// PinCategory - Ensure we set PC_Optional as needed, otherwise they default to PC_Required
+		// PinName - Add back in the type information (V4) / (S) / etc
+		if (UMaterialGraphNode* MaterialGraphNode = Cast<UMaterialGraphNode>(FunctionCall->GraphNode))
+		{
+			for (UEdGraphPin* Pin : MaterialGraphNode->Pins)
+			{
+				if (Pin->Direction == EGPD_Input)
+				{
+					const int32 InputIndex = Pin->SourceIndex;
+					if (InputIndex >= 0 && InputIndex < FunctionCall->FunctionInputs.Num())
+					{
+						Pin->PinType.PinCategory = FunctionCall->IsInputConnectionRequired(InputIndex) ? UMaterialGraphSchema::PC_Required : UMaterialGraphSchema::PC_Optional;
+
+						const FName NewPinName = UMaterialGraphNode::GetShortenPinName(FunctionCall->GetInputName(InputIndex));
+						if (!NewPinName.IsNone())
+						{
+							Pin->PinName = NewPinName;
+						}
+					}
+				}
+			}
+		}
+
 		// If an unknown material function has been pasted, remove the graph node pins (as the expression will also have had its inputs/outputs removed).
 		// This will be displayed as an orphaned "Unspecified Function" node.
 		if (FunctionCall->MaterialFunction == nullptr &&
