- Add custom primitive data preview support to the Material Editor
- Introduce editable preview values for custom primitive data slots in the material editor, allowing artists to set per-slot float values that are applied to the preview mesh component at edit time
- Add a "Use Preview Values" toggle, right-click context menu with clear and copy-defaults actions, and wire up the preview mesh component to reflect the current values via a new SetCustomPrimitiveData interface on IMaterialEditor

#Synced-CL 52631319

diff --git a/Engine/Source/Editor/MaterialEditor/Private/MaterialEditor.cpp b/Engine/Source/Editor/MaterialEditor/Private/MaterialEditor.cpp
--- a/Engine/Source/Editor/MaterialEditor/Private/MaterialEditor.cpp
+++ b/Engine/Source/Editor/MaterialEditor/Private/MaterialEditor.cpp
@@ -4,6 +4,7 @@
 #include "EngineGlobals.h"
 #include "Engine/SkeletalMesh.h"
 #include "AI/NavigationSystemBase.h"
+#include "Components/MeshComponent.h"
 #include "Engine/Engine.h"
 #include "EngineModule.h"
 #include "Misc/MessageDialog.h"
@@ -1730,7 +1731,7 @@ void FMaterialEditor::CreateInternalWidgets()
 
 	Generator->OnFinishedChangingProperties().AddSP(this, &FMaterialEditor::OnFinishedChangingParametersFromOverview);
 	Generator->OnRowsRefreshed().AddSP(this, &FMaterialEditor::GeneratorRowsRefreshed);
-	MaterialCustomPrimitiveDataWidget = SNew(SMaterialCustomPrimitiveDataPanel, MaterialEditorInstance);
+	MaterialCustomPrimitiveDataWidget = SNew(SMaterialCustomPrimitiveDataPanel, this, MaterialEditorInstance);
 
 	MaterialExpressionsOverviewWidget = SNew(SMaterialExpressionsOverviewPanel)
 		.InMaterialEditorInstance(MaterialEditorInstance)
@@ -2562,10 +2563,25 @@ void FMaterialEditor::RefreshPreviewViewport()
 {
 	if (PreviewViewport.IsValid())
 	{
+		UpdatePreviewCustomPrimitiveData();
 		PreviewViewport->RefreshViewport();
 	}
 }
 
+void FMaterialEditor::UpdatePreviewCustomPrimitiveData()
+{
+	if (!PreviewViewport.IsValid() || !PreviewViewport->PreviewMeshComponent || !Material)
+	{
+		return;
+	}
+
+	UPrimitiveComponent* PreviewComponent = PreviewViewport->PreviewMeshComponent;
+
+	PreviewComponent->ResetCustomPrimitiveData();
+	PreviewComponent->SetCustomPrimitiveDataFloatArray(0, CustomPrimitiveDataFloats);
+	PreviewComponent->MarkRenderStateDirty();
+}
+
 void FMaterialEditor::LoadEditorSettings()
 {
 	EditorOptions = UMaterialEditorOptions::Get();
@@ -4706,6 +4722,12 @@ void FMaterialEditor::PreviewNode(UMaterialGraphNode* InGraphNode)
 	SetPreviewExpression(InGraphNode ? InGraphNode->MaterialExpression : nullptr);
 }
 
+void FMaterialEditor::SetCustomPrimitiveData(TConstArrayView<float> CustomFloats)
+{
+	CustomPrimitiveDataFloats = CustomFloats;
+	UpdatePreviewCustomPrimitiveData();
+}
+
 void FMaterialEditor::OnToggleRealtimePreview()
 {
 	const FGraphPanelSelectionSet SelectedNodes = GetSelectedNodes();
diff --git a/Engine/Source/Editor/MaterialEditor/Private/MaterialEditor.h b/Engine/Source/Editor/MaterialEditor/Private/MaterialEditor.h
--- a/Engine/Source/Editor/MaterialEditor/Private/MaterialEditor.h
+++ b/Engine/Source/Editor/MaterialEditor/Private/MaterialEditor.h
@@ -295,7 +295,10 @@ public:
 	 * Refreshes the viewport containing the preview mesh.
 	 */
 	void RefreshPreviewViewport();
-	
+
+	/** Pushes the currently set custom primitive data to the preview mesh component. */
+	void UpdatePreviewCustomPrimitiveData();
+
 	/** Regenerates the code view widget with new text */
 	void RegenerateCodeView(bool bForce=false);
 	
@@ -841,6 +844,9 @@ private:
 	 */
 	virtual void PreviewNode(UMaterialGraphNode* InGraphNode) override;
 
+	/** Set the custom primitive data to use in the preview display area */
+	virtual void SetCustomPrimitiveData(TConstArrayView<float> CustomFloats) override;
+
 	/** Command for toggling real time preview of selected node */
 	void OnToggleRealtimePreview();
 	/** Command to select nodes downstream of selected node */
@@ -1169,4 +1175,7 @@ private:
 
 	/** The text for the number of shaders in the material. */
 	TSharedPtr<STextBlock> NumShadersText;
+
+	/** Current custom primitive data set */
+	TArray<float> CustomPrimitiveDataFloats;
 };
diff --git a/Engine/Source/Editor/MaterialEditor/Private/SMaterialEditorCustomPrimitiveDataWidget.cpp b/Engine/Source/Editor/MaterialEditor/Private/SMaterialEditorCustomPrimitiveDataWidget.cpp
--- a/Engine/Source/Editor/MaterialEditor/Private/SMaterialEditorCustomPrimitiveDataWidget.cpp
+++ b/Engine/Source/Editor/MaterialEditor/Private/SMaterialEditorCustomPrimitiveDataWidget.cpp
@@ -7,18 +7,27 @@
 #include "Materials/MaterialExpressionVectorParameter.h"
 #include "EdGraph/EdGraphNode.h"
 #include "EdGraph/EdGraph.h"
+#include "Framework/Application/SlateApplication.h"
+#include "Framework/MultiBox/MultiBoxBuilder.h"
+#include "Widgets/Input/SSpinBox.h"
+#include "Widgets/Input/SCheckBox.h"
 #include "Widgets/SToolTip.h"
+#include "IMaterialEditor.h"
+
+#define LOCTEXT_NAMESPACE "SMaterialEditorCustomPrimitiveDataWidget"
 
 class SCustomPrimitiveDataRow : public SMultiColumnTableRow<TSharedPtr<FCustomPrimitiveDataRowData>>
 {
 public:
 	SLATE_BEGIN_ARGS(SCustomPrimitiveDataRow) {}
 		SLATE_ARGUMENT(TSharedPtr<FCustomPrimitiveDataRowData>, Entry)
+		SLATE_EVENT(FSimpleDelegate, OnValueChanged)
 	SLATE_END_ARGS()
 
 	void Construct(const FArguments& Args, const TSharedRef<STableViewBase>& OwnerTableView)
 	{
 		RowData = Args._Entry;
+		OnValueChangedDelegate = Args._OnValueChanged;
 
 		FSuperRowType::Construct(
 			FSuperRowType::FArguments()
@@ -31,7 +40,6 @@ public:
 	virtual TSharedRef<SWidget> GenerateWidgetForColumn(const FName& ColumnName) override
 	{
 		TSharedRef<STextBlock> TextBlock = SNew(STextBlock);
-
 		if (RowData->bIsDuplicate)
 		{
 			TextBlock->SetToolTip(SNew(SToolTip).Text(FText::FromString("This slot is potentially incorrectly overlapping")).BorderImage(FCoreStyle::Get().GetBrush("ToolTip.BrightBackground")));
@@ -48,20 +56,38 @@ public:
 			TextBlock->SetText(FText::FromString(RowData->Name));
 			return TextBlock;
 		}
+		else if (ColumnName == "PreviewValue")
+		{
+			TSharedRef<SSpinBox<float>> SpinBox = SNew(SSpinBox<float>)
+				.Font(FAppStyle::GetFontStyle(TEXT("MenuItem.Font")))
+				.IsEnabled(RowData->bIsDuplicate == false)
+				.Value_Lambda([this]() { return RowData->PreviewValue; })
+				.OnValueChanged_Lambda([this](float InValue)
+				{
+					RowData->PreviewValue = InValue;
+					OnValueChangedDelegate.ExecuteIfBound();
+				}
+			);
+
+			return SpinBox;
+		}
 
 		return SNullWidget::NullWidget;
 	}
 
 private:
 	TSharedPtr<FCustomPrimitiveDataRowData> RowData;
+	FSimpleDelegate OnValueChangedDelegate;
 };
 
 void SMaterialCustomPrimitiveDataPanel::Refresh()
 {
-	TMap<int32, TSet<FString>> IndexBuckets;
-
 	if (MaterialEditorInstance && MaterialEditorInstance->PreviewMaterial)
 	{
+		const TArray<float> CustomFloats = MakeCustomPrimitiveFloatData();
+
+		Items.Empty();
+
 		TArray<UMaterialExpressionScalarParameter*> Scalars;
 		MaterialEditorInstance->PreviewMaterial->GetAllExpressionsInMaterialAndFunctionsOfType<UMaterialExpressionScalarParameter>(Scalars);
 
@@ -69,8 +95,9 @@ void SMaterialCustomPrimitiveDataPanel::Refresh()
 		{
 			if (Expr->bUseCustomPrimitiveData)
 			{
-				TSet<FString>& Params = IndexBuckets.FindOrAdd((int32)Expr->PrimitiveDataIndex);
-				Params.Add(Expr->GetParameterName().ToString());
+				const int32 SlotIndex = (int32)Expr->PrimitiveDataIndex;
+				const float PreviewValue = CustomFloats.IsValidIndex(SlotIndex) ? CustomFloats[SlotIndex] : 0.0f;
+				Items.Emplace(MakeShared<FCustomPrimitiveDataRowData>(SlotIndex, Expr->GetParameterName().ToString(), PreviewValue));
 			}
 		}
 
@@ -84,23 +111,13 @@ void SMaterialCustomPrimitiveDataPanel::Refresh()
 				TArray<FString> Elements({".r", ".g" , ".b",  ".a" });
 				for (int32 i = 0; i < 4; i++)
 				{
-					TSet<FString>& Params = IndexBuckets.FindOrAdd((int32)Expr->PrimitiveDataIndex + i);
-					Params.Add(Expr->GetParameterName().ToString() + Elements[i]);
+					const int32 SlotIndex = (int32)Expr->PrimitiveDataIndex + i;
+					const float PreviewValue = CustomFloats.IsValidIndex(SlotIndex) ? CustomFloats[SlotIndex] : 0.0f;
+					Items.Emplace(MakeShared<FCustomPrimitiveDataRowData>(SlotIndex, Expr->GetParameterName().ToString() + Elements[i], PreviewValue));
 				}
 			}
 		}
 
-		Items.Empty();
-
-		// Go through and mark slots as duplicate while copying over to the shared ptr array
-		for (auto& Params : IndexBuckets)
-		{
-			for (FString& ParamName : Params.Value)
-			{
-				Items.Add(MakeShareable(new FCustomPrimitiveDataRowData(Params.Key, ParamName, Params.Value.Num() > 1)));
-			}
-		}
-
 		// Sort the items
 		Items.StableSort([](const TSharedPtr<FCustomPrimitiveDataRowData>& A, const TSharedPtr<FCustomPrimitiveDataRowData>& B) {
 			int32 SlotA = A->Slot;
@@ -111,12 +128,56 @@ void SMaterialCustomPrimitiveDataPanel::Refresh()
 			}
 			return SlotA < SlotB;
 			});
+
+		// Mark items that are in duplicate slots
+		for ( int32 i=0; i < Items.Num() - 1; ++i )
+		{
+			if (Items[i]->Slot == Items[i + 1]->Slot)
+			{
+				Items[i]->bIsDuplicate = true;
+				Items[i + 1]->bIsDuplicate = true;
+			}
+		}
+
+		RefreshCustomPrimitiveData();
 	}
 
 	ListViewWidget->RequestListRefresh();
 }
 
-void SMaterialCustomPrimitiveDataPanel::Construct(const FArguments& InArgs, UMaterialEditorPreviewParameters* InMaterialEditorInstance)
+void SMaterialCustomPrimitiveDataPanel::RefreshCustomPrimitiveData()
+{
+	if (MaterialEditor)
+	{
+		TArray<float> CustomFloats;
+		if (bUsePreviewValues)
+		{
+			CustomFloats = MakeCustomPrimitiveFloatData();
+		}
+
+		MaterialEditor->SetCustomPrimitiveData(CustomFloats);
+	}
+}
+
+TArray<float> SMaterialCustomPrimitiveDataPanel::MakeCustomPrimitiveFloatData() const
+{
+	TArray<float> CustomFloats;
+
+	int32 SlotMax = 0;
+	for (const TSharedPtr<FCustomPrimitiveDataRowData>& Item : Items)
+	{
+		SlotMax = FMath::Max(SlotMax, Item->Slot);
+	}
+
+	CustomFloats.AddZeroed(SlotMax + 1);
+	for (const TSharedPtr<FCustomPrimitiveDataRowData>& Item : Items)
+	{
+		CustomFloats[Item->Slot] = Item->PreviewValue;
+	}
+	return CustomFloats;
+}
+
+void SMaterialCustomPrimitiveDataPanel::Construct(const FArguments& InArgs, IMaterialEditor* InMaterialEditor, UMaterialEditorPreviewParameters* InMaterialEditorInstance)
 {	
 	this->ChildSlot
 		[
@@ -143,12 +204,31 @@ void SMaterialCustomPrimitiveDataPanel::Construct(const FArguments& InArgs, UMat
 							.Font(FAppStyle::GetFontStyle("DetailsView.CategoryFontStyle"))
 							.ShadowOffset(FVector2D(1.0f, 1.0f))
 						]
+						+ SHorizontalBox::Slot()
+						.HAlign(HAlign_Right)
+						.VAlign(VAlign_Center)
+						.Padding(FMargin(8.0f, 4.0f, 3.0f, 4.0f))
+						[
+							SNew(SCheckBox)
+							.IsChecked_Lambda([this]() { return bUsePreviewValues ? ECheckBoxState::Checked : ECheckBoxState::Unchecked; })
+							.OnCheckStateChanged_Lambda([this](ECheckBoxState NewState)
+							{
+								bUsePreviewValues = (NewState == ECheckBoxState::Checked);
+								RefreshCustomPrimitiveData();
+							})
+							[
+								SNew(STextBlock)
+								.Text(FText::FromString("Use Preview Values"))
+								.Font(FAppStyle::GetFontStyle("DetailsView.CategoryFontStyle"))
+							]
+						]
 					]
 					+ SVerticalBox::Slot()
 					.Padding(FMargin(3.0f, 2.0f, 3.0f, 3.0f))
 					[
 						SNew(SBorder)
 						.BorderImage(FAppStyle::GetBrush("DetailsView.CategoryTop"))
+						.OnMouseButtonDown(this, &SMaterialCustomPrimitiveDataPanel::OnMouseButtonDown)
 						[
 							SNew(SScrollBox)
 							+ SScrollBox::Slot()
@@ -159,9 +239,10 @@ void SMaterialCustomPrimitiveDataPanel::Construct(const FArguments& InArgs, UMat
 								.SelectionMode(ESelectionMode::None)
 								.HeaderRow
 								(
-									SNew(SHeaderRow)
+									SAssignNew(HeaderRow, SHeaderRow)
 									+ SHeaderRow::Column("Slot").DefaultLabel(FText::FromString("Slot")).ManualWidth(48.0f)
 									+ SHeaderRow::Column("Name").DefaultLabel(FText::FromString("Name"))
+									+ SHeaderRow::Column("PreviewValue").DefaultLabel(FText::FromString("Preview Value"))
 								)
 							]
 						]
@@ -170,11 +251,107 @@ void SMaterialCustomPrimitiveDataPanel::Construct(const FArguments& InArgs, UMat
 			]
 		];
 
+	MaterialEditor = InMaterialEditor;
 	MaterialEditorInstance = InMaterialEditorInstance;
 	Refresh();
 }
 
 TSharedRef<ITableRow> SMaterialCustomPrimitiveDataPanel::OnGenerateRowForList(TSharedPtr<FCustomPrimitiveDataRowData> Item, const TSharedRef<STableViewBase>& OwnerTable)
 {
-	return SNew(SCustomPrimitiveDataRow, OwnerTable).Entry(Item);
+	return SNew(SCustomPrimitiveDataRow, OwnerTable)
+		.Entry(Item)
+		.OnValueChanged(FSimpleDelegate::CreateSP(this, &SMaterialCustomPrimitiveDataPanel::RefreshCustomPrimitiveData));
+}
+
+FReply SMaterialCustomPrimitiveDataPanel::OnMouseButtonDown(const FGeometry& MyGeometry, const FPointerEvent& MouseEvent)
+{
+	if (MouseEvent.GetEffectingButton() == EKeys::RightMouseButton)
+	{
+		FMenuBuilder MenuBuilder(true, nullptr);
+		MenuBuilder.BeginSection(NAME_None, LOCTEXT("PreviewValues", "Preview Values"));
+		{
+			MenuBuilder.AddMenuEntry(
+				LOCTEXT("PreviewValues_Clear", "Clear Values"),
+				LOCTEXT("PreviewValues_ClearTooltip", "Clears all the preview values back to 0.0"),
+				FSlateIcon(),
+				FUIAction(FExecuteAction::CreateSP(this, &SMaterialCustomPrimitiveDataPanel::ClearPreviewValues))
+			);
+
+			MenuBuilder.AddMenuEntry(
+				LOCTEXT("PreviewValues_CopyDefaults", "Copy Defaults to Values"),
+				LOCTEXT("PreviewValues_CopyDefaultsTooltip", "Copies the default values from the material nodes to the preview values"),
+				FSlateIcon(),
+				FUIAction(FExecuteAction::CreateSP(this, &SMaterialCustomPrimitiveDataPanel::CopyDefaultToPreviewValues))
+			);
+		}
+		MenuBuilder.EndSection();
+
+		FSlateApplication::Get().PushMenu(
+			AsShared(),
+			*MouseEvent.GetEventPath(),
+			MenuBuilder.MakeWidget(),
+			MouseEvent.GetScreenSpacePosition(),
+			FPopupTransitionEffect(FPopupTransitionEffect::ContextMenu)
+		);
+		return FReply::Handled();
+	}
+	return FReply::Unhandled();
 }
+
+void SMaterialCustomPrimitiveDataPanel::ClearPreviewValues()
+{
+	for (const TSharedPtr<FCustomPrimitiveDataRowData>& Item : Items)
+	{
+		Item->PreviewValue = 0.0f;
+	}
+
+	RefreshCustomPrimitiveData();
+}
+
+void SMaterialCustomPrimitiveDataPanel::CopyDefaultToPreviewValues()
+{
+	TMap<int32, float> DefaultValues;
+
+	if (MaterialEditorInstance && MaterialEditorInstance->PreviewMaterial)
+	{
+		// Scalars
+		{
+			TArray<UMaterialExpressionScalarParameter*> Scalars;
+			MaterialEditorInstance->PreviewMaterial->GetAllExpressionsInMaterialAndFunctionsOfType<UMaterialExpressionScalarParameter>(Scalars);
+
+			for (const UMaterialExpressionScalarParameter* Expr : Scalars)
+			{
+				if (Expr->bUseCustomPrimitiveData)
+				{
+					DefaultValues.FindOrAdd((int32)Expr->PrimitiveDataIndex, Expr->DefaultValue);
+				}
+			}
+		}
+
+		// Vectors
+		{
+			TArray<UMaterialExpressionVectorParameter*> Vectors;
+			MaterialEditorInstance->PreviewMaterial->GetAllExpressionsInMaterialAndFunctionsOfType<UMaterialExpressionVectorParameter>(Vectors);
+
+			for (const UMaterialExpressionVectorParameter* Expr : Vectors)
+			{
+				if (Expr->bUseCustomPrimitiveData)
+				{
+					for (int32 i = 0; i < 4; i++)
+					{
+						DefaultValues.FindOrAdd((int32)Expr->PrimitiveDataIndex + i, Expr->DefaultValue.Component(i));
+					}
+				}
+			}
+		}
+	}
+
+	for (const TSharedPtr<FCustomPrimitiveDataRowData>& Item : Items)
+	{
+		Item->PreviewValue = DefaultValues.FindRef(Item->Slot);
+	}
+
+	RefreshCustomPrimitiveData();
+}
+
+#undef LOCTEXT_NAMESPACE
diff --git a/Engine/Source/Editor/MaterialEditor/Private/SMaterialEditorCustomPrimitiveDataWidget.h b/Engine/Source/Editor/MaterialEditor/Private/SMaterialEditorCustomPrimitiveDataWidget.h
--- a/Engine/Source/Editor/MaterialEditor/Private/SMaterialEditorCustomPrimitiveDataWidget.h
+++ b/Engine/Source/Editor/MaterialEditor/Private/SMaterialEditorCustomPrimitiveDataWidget.h
@@ -14,15 +14,17 @@
 #include "IDetailPropertyRow.h"
 #include "MaterialPropertyHelpers.h"
 
+class IMaterialEditor;
 class UMaterialEditorPreviewParameters;
 
 /** Data for a row in the list of custom primitive data entries*/
 struct FCustomPrimitiveDataRowData
 {
-	FCustomPrimitiveDataRowData(int32 InSlot, FString InName, bool bInIsDuplicate = false) : Slot(InSlot), Name(InName), bIsDuplicate(bInIsDuplicate) {}
+	FCustomPrimitiveDataRowData(int32 InSlot, FString InName, float InPreviewValue, bool bInIsDuplicate = false) : Slot(InSlot), Name(InName), PreviewValue(InPreviewValue), bIsDuplicate(bInIsDuplicate) {}
 
 	int32 Slot = 0;
 	FString Name = "";
+	float PreviewValue = 0.0f;
 	bool bIsDuplicate = false;
 };
 
@@ -32,7 +34,7 @@ public:
 	SLATE_BEGIN_ARGS(SMaterialCustomPrimitiveDataPanel) {}
 	SLATE_END_ARGS()
 
-	void Construct(const FArguments& InArgs, UMaterialEditorPreviewParameters* InMaterialEditorInstance);
+	void Construct(const FArguments& InArgs, IMaterialEditor* MaterialEditor, UMaterialEditorPreviewParameters* InMaterialEditorInstance);
 
 	void SetEditorInstance(UMaterialEditorPreviewParameters* InMaterialEditorInstance) { MaterialEditorInstance = InMaterialEditorInstance; Refresh(); }
 
@@ -40,15 +42,36 @@ private:
 
 	void Refresh();
 
+	void RefreshCustomPrimitiveData();
+
+	TArray<float> MakeCustomPrimitiveFloatData() const;
+
 	/** Adds a new textbox with the string to the list */
 	TSharedRef<ITableRow> OnGenerateRowForList(TSharedPtr<FCustomPrimitiveDataRowData> Item, const TSharedRef<STableViewBase>& OwnerTable);
 
+	/** Called when a mouse button is down in the table area */
+	FReply OnMouseButtonDown(const FGeometry& MyGeometry, const FPointerEvent& MouseEvent);
+
+	/** Clears all the preview values */
+	void ClearPreviewValues();
+
+	/** Copies all the default node values to the preview values. */
+	void CopyDefaultToPreviewValues();
+
 	/** The list of strings */
 	TArray<TSharedPtr<FCustomPrimitiveDataRowData>> Items;
 
 	/** The actual UI list */
 	TSharedPtr<SListView<TSharedPtr<FCustomPrimitiveDataRowData>>> ListViewWidget;
 
+	TSharedPtr<SHeaderRow> HeaderRow;
+
+	/** When enabled the preview values will be set as custom primitive data. */
+	bool bUsePreviewValues = true;
+
+	/** Pointer to the material editor interface */
+	IMaterialEditor* MaterialEditor = nullptr;
+
 	/** The set of material parameters this is associated with */
 	UMaterialEditorPreviewParameters* MaterialEditorInstance;
 };
\ No newline at end of file
diff --git a/Engine/Source/Editor/MaterialEditor/Public/IMaterialEditor.h b/Engine/Source/Editor/MaterialEditor/Public/IMaterialEditor.h
--- a/Engine/Source/Editor/MaterialEditor/Public/IMaterialEditor.h
+++ b/Engine/Source/Editor/MaterialEditor/Public/IMaterialEditor.h
@@ -172,6 +172,10 @@ public:
 
 	/** Start/Stop Previewing the specified material graph node */
 	virtual void PreviewNode(UMaterialGraphNode* InGraphNode) {};
+
+	/** Set the custom primitive data to use in the preview display area */
+	virtual void SetCustomPrimitiveData(TConstArrayView<float> CustomFloats) {}
+
 private:
 	FMaterialEditorClosedEvent MaterialEditorClosedEvent;
 	FRegisterTabSpawnersEvent RegisterTabSpawnersEvent;
