Add material experssion/property disconnect support to MaterialEditingLibrary.

#Synced-CL 52671970

diff --git a/Engine/Source/Editor/MaterialEditor/Private/MaterialEditingLibrary.cpp b/Engine/Source/Editor/MaterialEditor/Private/MaterialEditingLibrary.cpp
--- a/Engine/Source/Editor/MaterialEditor/Private/MaterialEditingLibrary.cpp
+++ b/Engine/Source/Editor/MaterialEditor/Private/MaterialEditingLibrary.cpp
@@ -798,6 +798,38 @@ bool UMaterialEditingLibrary::ConnectMaterialExpressions(UMaterialExpression* Fr
 	return bResult;
 }
 
+bool UMaterialEditingLibrary::DisconnectMaterialExpressions(UMaterialExpression* ToExpression, FString ToInputName)
+{
+	if (!ToExpression)
+	{
+		return false;
+	}
+	FExpressionInput* Input = GetExpressionInputByName(ToExpression, *ToInputName);
+	if (!Input || !Input->IsConnected())
+	{
+		return false;
+	}
+	Input->Expression = nullptr;
+	Input->OutputIndex = 0;
+	return true;
+}
+
+bool UMaterialEditingLibrary::DisconnectMaterialProperty(UMaterial* Material, EMaterialProperty Property)
+{
+	if (!Material)
+	{
+		return false;
+	}
+	FExpressionInput* Input = Material->GetExpressionInputForProperty(Property);
+	if (!Input || !Input->IsConnected())
+	{
+		return false;
+	}
+	Input->Expression = nullptr;
+	Input->OutputIndex = 0;
+	return true;
+}
+
 
 void UMaterialEditingLibrary::RecompileMaterialInternal(const TArray<UMaterial*>& Materials, const FOnItemComplete* OnItemComplete)
 {
diff --git a/Engine/Source/Editor/MaterialEditor/Public/MaterialEditingLibrary.h b/Engine/Source/Editor/MaterialEditor/Public/MaterialEditingLibrary.h
--- a/Engine/Source/Editor/MaterialEditor/Public/MaterialEditingLibrary.h
+++ b/Engine/Source/Editor/MaterialEditor/Public/MaterialEditingLibrary.h
@@ -211,7 +211,25 @@ public:
 	UFUNCTION(BlueprintCallable, Category = "MaterialEditing")
 	static UE_API bool ConnectMaterialExpressions(UMaterialExpression* FromExpression, FString FromOutputName, UMaterialExpression* ToExpression, FString ToInputName);
 
-	/** 
+	/**
+	 *	Disconnect the input pin of a material expression, removing whatever is currently connected to it.
+	 *	@param	ToExpression	Expression whose input pin should be disconnected
+	 *	@param	ToInputName		Name of the input pin to disconnect.
+	 *	@return	true if a connection was broken, false if the pin was already disconnected or not found
+	 */
+	UFUNCTION(BlueprintCallable, Category = "MaterialEditing")
+	static UE_API bool DisconnectMaterialExpressions(UMaterialExpression* ToExpression, FString ToInputName);
+
+	/**
+	 *	Disconnect the expression node currently connected to a material property input (e.g. base color, roughness).
+	 *	@param	Material	Material asset to modify
+	 *	@param	Property	Property input on material to disconnect
+	 *	@return	true if a connection was broken, false if the input was already disconnected
+	 */
+	UFUNCTION(BlueprintCallable, Category = "MaterialEditing")
+	static UE_API bool DisconnectMaterialProperty(UMaterial* Material, EMaterialProperty Property);
+
+	/**
 	 *	Trigger a recompile of a material. Must be performed after making changes to the graph to have changes reflected.
 	 */
 	UFUNCTION(BlueprintCallable, Category = "MaterialEditing")
