[NewMaterialTranslator] - Fix bug with my previous checkin that is causing connection links to show up inactive in the old transalator even when they were "active".

- I originally thought my change only effected the new translator, but there was a new code path I added that effected the old translator.
- Input pins marked as optional were falling through a new code path I added which was causing these connections to get the "inactive" color for the links.
- The fix here is to revert back to the old behavior and the new code is just for the NMT.

#rnx
#Synced-CL 52699221

diff --git a/Engine/Source/Editor/MaterialEditor/Private/MaterialGraphConnectionDrawingPolicy.cpp b/Engine/Source/Editor/MaterialEditor/Private/MaterialGraphConnectionDrawingPolicy.cpp
--- a/Engine/Source/Editor/MaterialEditor/Private/MaterialGraphConnectionDrawingPolicy.cpp
+++ b/Engine/Source/Editor/MaterialEditor/Private/MaterialGraphConnectionDrawingPolicy.cpp
@@ -202,10 +202,6 @@ void FMaterialGraphConnectionDrawingPolicy::DetermineWiringStyle(UEdGraphPin* Ou
 		{
 			Params.WireColor = UMaterialGraphSchema::GetColorForConnectionType(ValueTypeObj->ValueType);
 		}
-		else
-		{
-			Params.WireColor = MaterialGraphSchema->GetPinTypeColor(InputPin->PinType);
-		}
 	}
 
 	const bool bDeemphasizeUnhoveredPins = HoveredPins.Num() > 0;
