Return any compile erorrs from MaterialLibrary ReompileMaterial.

#Synced-CL 52724283

diff --git a/Engine/Source/Editor/MaterialEditor/Private/MaterialEditingLibrary.cpp b/Engine/Source/Editor/MaterialEditor/Private/MaterialEditingLibrary.cpp
--- a/Engine/Source/Editor/MaterialEditor/Private/MaterialEditingLibrary.cpp
+++ b/Engine/Source/Editor/MaterialEditor/Private/MaterialEditingLibrary.cpp
@@ -882,11 +882,11 @@ void UMaterialEditingLibrary::RecompileMaterialInternal(const TArray<UMaterial*>
 	}
 }
 
-void UMaterialEditingLibrary::RecompileMaterial(UMaterial* Material)
+TArray<FString> UMaterialEditingLibrary::RecompileMaterial(UMaterial* Material)
 {
 	if (!ensureAsRuntimeWarning(Material != nullptr))
 	{
-		return;
+		return { TEXT("Cannot compile a null material.") };
 	}
 
 	RecompileMaterialInternal( { Material }, nullptr);
@@ -897,6 +897,13 @@ void UMaterialEditingLibrary::RecompileMaterial(UMaterial* Material)
 
 	// Update streaming data
 	FMaterialEditorUtilities::BuildTextureStreamingData(Material);
+
+	if (const FMaterialResource* Resource = Material->GetMaterialResource(GMaxRHIShaderPlatform))
+	{
+		return Resource->GetCompileErrors();
+	}
+
+	return {};
 }
 
 void UMaterialEditingLibrary::RecompileMaterials(TArray<UMaterial*>& Materials, const FOnItemComplete& OnItemComplete)
diff --git a/Engine/Source/Editor/MaterialEditor/Public/MaterialEditingLibrary.h b/Engine/Source/Editor/MaterialEditor/Public/MaterialEditingLibrary.h
--- a/Engine/Source/Editor/MaterialEditor/Public/MaterialEditingLibrary.h
+++ b/Engine/Source/Editor/MaterialEditor/Public/MaterialEditingLibrary.h
@@ -231,9 +231,10 @@ public:
 
 	/**
 	 *	Trigger a recompile of a material. Must be performed after making changes to the graph to have changes reflected.
+	 *  @return Any compiler error messages produced during compilation. Empty if compilation succeeded.
 	 */
 	UFUNCTION(BlueprintCallable, Category = "MaterialEditing")
-	static UE_API void RecompileMaterial(UMaterial* Material);
+	static UE_API TArray<FString> RecompileMaterial(UMaterial* Material);
 
 	/** Item complete callback for use when processing arrays of items. Allows higher level to tick UI etc. */
 	DECLARE_DELEGATE(FOnItemComplete)
